StateTree: Rename some delegate variables and functions to integrate more with the engine architecture.
Change the dispatcher to be unique to prevent sending a delegate on the wrong frame.
Support bind and unbind delegate from Blueprint task.

diff --git a/Engine/Plugins/Runtime/StateTree/Source/StateTreeEditorModule/Private/StateTreeCompiler.cpp b/Engine/Plugins/Runtime/StateTree/Source/StateTreeEditorModule/Private/StateTreeCompiler.cpp
--- a/Engine/Plugins/Runtime/StateTree/Source/StateTreeEditorModule/Private/StateTreeCompiler.cpp
+++ b/Engine/Plugins/Runtime/StateTree/Source/StateTreeEditorModule/Private/StateTreeCompiler.cpp
@@ -20,6 +20,12 @@
 
 namespace UE::StateTree::Compiler
 {
+	FAutoConsoleVariable CVarLogEnableParameterDelegateDispatcherBinding(
+		TEXT("StateTree.Compiler.EnableParameterDelegateDispatcherBinding"),
+		false,
+		TEXT("Enable binding from delegate dispatchers that are in the state tree parameters.")
+	);
+
 	FAutoConsoleVariable CVarLogCompiledStateTree(
 		TEXT("StateTree.Compiler.LogResultOnCompilationCompleted"),
 		false,
@@ -221,21 +227,7 @@ bool FStateTreeCompiler::Compile(UStateTree& InStateTree)
 	// Copy schema the EditorData
 	StateTree->Schema = DuplicateObject(EditorData->Schema, StateTree);
 
-	// Copy parameters from EditorData	
-	StateTree->Parameters = EditorData->RootParameters.Parameters;
-	
-	// Mark parameters as binding source
-	const FStateTreeBindableStructDesc ParametersDesc = {
-			UE::StateTree::Editor::GlobalStateName,
-			TEXT("Parameters"),
-			StateTree->Parameters.GetPropertyBagStruct(),
-			FStateTreeDataHandle(EStateTreeDataSourceType::GlobalParameterData),
-			EStateTreeBindableStructSource::Parameter,
-			EditorData->RootParameters.ID
-		};
-	BindingsCompiler.AddSourceStruct(ParametersDesc);
-
-	if (!UE::StateTree::Compiler::ValidateNoLevelActorReferences(Log, ParametersDesc, FStateTreeDataView(), FStateTreeDataView(EditorData->RootParameters.Parameters.GetMutableValue())))
+	if (!CreateParameters())
 	{
 		StateTree->ResetCompiled();
 		return false;
@@ -260,7 +252,7 @@ bool FStateTreeCompiler::Compile(UStateTree& InStateTree)
 			BindingsCompiler.AddSourceStruct(ExtDataDesc);
 			if (const auto Validation = UE::StateTree::Compiler::IsValidIndex16(ContextDataIndex); Validation.DidFail())
 			{
-				Validation.Log(Log, TEXT("ExternalStructIndex"), ParametersDesc);
+				Validation.Log(Log, TEXT("ExternalStructIndex"), ExtDataDesc);
 				return false;
 			}
 			Desc.Handle.DataHandle = ExtDataDesc.DataHandle;
@@ -269,7 +261,7 @@ bool FStateTreeCompiler::Compile(UStateTree& InStateTree)
 
 	if (const UE::StateTree::Compiler::FValidationResult Validation = UE::StateTree::Compiler::IsValidIndex16(ContextDataIndex); Validation.DidFail())
 	{
-		Validation.Log(Log, TEXT("NumContextData"), ParametersDesc);
+		Validation.Log(Log, TEXT("NumContextData"));
 		return false;
 	}
 	StateTree->NumContextData = static_cast<uint16>(ContextDataIndex);
@@ -298,7 +290,7 @@ bool FStateTreeCompiler::Compile(UStateTree& InStateTree)
 	const int32 NumGlobalInstanceData = InstanceStructs.Num();
 	if (const auto Validation = UE::StateTree::Compiler::IsValidIndex16(NumGlobalInstanceData); Validation.DidFail())
 	{
-		Validation.Log(Log, TEXT("NumGlobalInstanceData"), ParametersDesc);
+		Validation.Log(Log, TEXT("NumGlobalInstanceData"));
 		return false;
 	}
 	StateTree->NumGlobalInstanceData = uint16(NumGlobalInstanceData);
@@ -324,7 +316,10 @@ bool FStateTreeCompiler::Compile(UStateTree& InStateTree)
 	StateTree->Nodes = Nodes;
 	StateTree->DefaultInstanceData.Init(*StateTree, InstanceStructs);
 	StateTree->SharedInstanceData.Init(*StateTree, SharedInstanceStructs);
-	
+
+	// Store the new compiled dispatchers.
+	EditorData->CompiledDispatchers = BindingsCompiler.GetCompiledDelegateDispatchers();
+
 	BindingsCompiler.Finalize();
 
 	if (!StateTree->Link())
@@ -385,6 +380,54 @@ UStateTreeState* FStateTreeCompiler::GetState(const FGuid& StateID) const
 	return SourceStates[*Idx];
 }
 
+bool FStateTreeCompiler::CreateParameters()
+{
+	// Copy parameters from EditorData	
+	StateTree->Parameters = EditorData->RootParameters.Parameters;
+
+	// Mark parameters as binding source
+	const FStateTreeBindableStructDesc ParametersDesc = {
+			UE::StateTree::Editor::GlobalStateName,
+			TEXT("Parameters"),
+			StateTree->Parameters.GetPropertyBagStruct(),
+			FStateTreeDataHandle(EStateTreeDataSourceType::GlobalParameterData),
+			EStateTreeBindableStructSource::Parameter,
+			EditorData->RootParameters.ID
+	};
+	BindingsCompiler.AddSourceStruct(ParametersDesc);
+
+	if (!UE::StateTree::Compiler::ValidateNoLevelActorReferences(Log, ParametersDesc, FStateTreeDataView(), FStateTreeDataView(EditorData->RootParameters.Parameters.GetMutableValue())))
+	{
+		return false;
+	}
+
+	// Compile the delegate dispatcher.
+	if (UE::StateTree::Compiler::CVarLogEnableParameterDelegateDispatcherBinding->GetBool())
+	{
+		FValidatedPathBindings Bindings;
+		FStateTreeDataView SourceValue(StateTree->Parameters.GetMutableValue());
+		if (!GetAndValidateBindings(ParametersDesc, SourceValue, Bindings))
+		{
+			Log.Reportf(EMessageSeverity::Error, TEXT("Failed to create bindings for global parameters."));
+			return false;
+		}
+
+		if (Bindings.CopyBindings.Num() != 0 || Bindings.DelegateListeners.Num() != 0 || Bindings.ReferenceBindings.Num() != 0)
+		{
+			Log.Reportf(EMessageSeverity::Warning, TEXT("The global parameters should not target have binding."));
+			return false;
+		}
+
+		if (!BindingsCompiler.CompileDelegateDispatchers(ParametersDesc, EditorData->CompiledDispatchers, Bindings.DelegateDispatchers, SourceValue))
+		{
+			Log.Reportf(EMessageSeverity::Error, TEXT("Failed to create delegate dispatcher bindings."));
+			return false;
+		}
+	}
+
+	return true;
+}
+
 bool FStateTreeCompiler::CreateStates()
 {
 	check(EditorData);
@@ -1006,9 +1049,8 @@ bool FStateTreeCompiler::CreateStateTransitions()
 					continue;
 				}
 
-				CompactTransition.DelegateDispatcherID = BindingsCompiler.GetDispatcherIDFromPath(Binding->GetSourcePath());
-
-				if (!CompactTransition.DelegateDispatcherID.IsValid())
+				CompactTransition.RequiredDelegateDispatcher = BindingsCompiler.GetDispatcherFromPath(Binding->GetSourcePath());
+				if (!CompactTransition.RequiredDelegateDispatcher.IsValid())
 				{
 					bTransitionSucceeded = false;
 					Log.Reportf(EMessageSeverity::Error,
@@ -1316,8 +1358,14 @@ bool FStateTreeCompiler::CreateBindingsForStruct(const FStateTreeBindableStructD
 		OutBatchIndex = FStateTreeIndex16(BatchIndex);
 	}
 
-	// Delegate Bindings
-	if (!BindingsCompiler.CompileDelegates(TargetStruct, Bindings.DelegateSourceBindings, Bindings.DelegateTargetBindings, TargetValue))
+	// Delegate Dispatcher
+	if (!BindingsCompiler.CompileDelegateDispatchers(TargetStruct, EditorData->CompiledDispatchers, Bindings.DelegateDispatchers, TargetValue))
+	{
+		return false;
+	}
+
+	// Delegate Listener
+	if (!BindingsCompiler.CompileDelegateListeners(TargetStruct, Bindings.DelegateListeners, TargetValue))
 	{
 		return false;
 	}
@@ -2005,26 +2053,31 @@ bool FStateTreeCompiler::GetAndValidateBindings(const FStateTreeBindableStructDe
 
 	bool bSucceeded = true;
 
+	// Handle sources. Need to handle them now while we have the instance.
 	for (FPropertyBindingBinding& Binding : EditorData->EditorBindings.GetMutableBindings())
 	{
-		if (Binding.GetTargetPath().GetStructID() != TargetStruct.ID)
+		if (Binding.GetSourcePath().GetStructID() == TargetStruct.ID)
 		{
-			if(Binding.GetSourcePath().GetStructID() == TargetStruct.ID)
+			if (IsPropertyOfTypeOrChild(*FStateTreeDelegateDispatcher::StaticStruct(), TargetStruct, Binding.GetSourcePath()))
 			{
-				if (IsPropertyOfTypeOrChild(*FStateTreeDelegateDispatcher::StaticStruct(), TargetStruct, Binding.GetSourcePath()))
-				{
-					FStateTreePropertyPathBinding& BindingCopy = OutValidatedBindings.DelegateSourceBindings.Add_GetRef(FStateTreePropertyPathBinding(Binding));
-					BindingCopy.SetSourceDataHandle(TargetStruct.DataHandle);
-				}
+				FStateTreePropertyPathBinding& BindingCopy = OutValidatedBindings.DelegateDispatchers.Add_GetRef(FStateTreePropertyPathBinding(Binding));
+				BindingCopy.SetSourceDataHandle(TargetStruct.DataHandle);
 			}
+		}
+	}
 
+	// Handle targets.
+	for (FPropertyBindingBinding& Binding : EditorData->EditorBindings.GetMutableBindings())
+	{
+		if (Binding.GetTargetPath().GetStructID() != TargetStruct.ID)
+		{
 			continue;
 		}
 
 		TArray<FStateTreePropertyPathBinding>* OutputBindings = nullptr;
 		if (IsPropertyOfTypeOrChild(*FStateTreeDelegateListener::StaticStruct(), TargetStruct, Binding.GetTargetPath()))
 		{
-			OutputBindings = &OutValidatedBindings.DelegateTargetBindings;
+			OutputBindings = &OutValidatedBindings.DelegateListeners;
 		}
 		else if (IsPropertyOfTypeOrChild(*FStateTreePropertyRef::StaticStruct(), TargetStruct, Binding.GetTargetPath()))
 		{
diff --git a/Engine/Plugins/Runtime/StateTree/Source/StateTreeEditorModule/Private/StateTreePropertyBindingCompiler.cpp b/Engine/Plugins/Runtime/StateTree/Source/StateTreeEditorModule/Private/StateTreePropertyBindingCompiler.cpp
--- a/Engine/Plugins/Runtime/StateTree/Source/StateTreeEditorModule/Private/StateTreePropertyBindingCompiler.cpp
+++ b/Engine/Plugins/Runtime/StateTree/Source/StateTreeEditorModule/Private/StateTreePropertyBindingCompiler.cpp
@@ -133,112 +133,157 @@ bool FStateTreePropertyBindingCompiler::CompileBatch(const FStateTreeBindableStr
 	return true;
 }
 
-bool FStateTreePropertyBindingCompiler::CompileDelegates(const FStateTreeBindableStructDesc& TargetStruct, TConstArrayView<FStateTreePropertyPathBinding> DelegateSourceBindings, TConstArrayView<FStateTreePropertyPathBinding> DelegateTargetBindings, FStateTreeDataView InstanceDataView)
+bool FStateTreePropertyBindingCompiler::CompileDelegateDispatchers(
+	const FStateTreeBindableStructDesc& SourceStruct,
+	TConstArrayView<FStateTreeEditorDelegateDispatcherCompiledBinding> PreviousCompiledDistpatchers,
+	TConstArrayView<FStateTreePropertyPathBinding> DelegateDispatcherBindings,
+	FStateTreeDataView InstanceDataView)
 {
 	check(Log);
 	check(PropertyBindings);
 
 	StoreSourceStructs();
 
-	// Compiling dispatchers
+	bool bSuccess = true;
+	for (const FStateTreePropertyPathBinding& Binding : DelegateDispatcherBindings)
 	{
-		TArray<FPropertyBindingPath, TInlineAllocator<4>> FoundDispatchers;
-		for (const FStateTreePropertyPathBinding& Binding : DelegateSourceBindings)
+		if (Binding.GetSourcePath().GetStructID() != SourceStruct.ID)
 		{
-			if (Binding.GetSourcePath().GetStructID() != TargetStruct.ID)
-			{
-				continue;
-			}
+			continue;
+		}
 
-			if (FoundDispatchers.Contains(Binding.GetSourcePath()))
-			{
-				continue;
-			}
+		// Source must be in the source array
+		const FStateTreeBindableStructDesc* DispatcherStruct = GetSourceStructDescByID(Binding.GetSourcePath().GetStructID());
+		if (!DispatcherStruct)
+		{
+			Log->Reportf(EMessageSeverity::Error, SourceStruct, TEXT("Could not find a binding source."));
+			bSuccess = false;
+			continue;
+		}
 
-			FoundDispatchers.Add(Binding.GetSourcePath());
+		auto FindBySourcePathPredicate = [&Binding](const FStateTreeEditorDelegateDispatcherCompiledBinding& Other)
+			{
+				return Other.DispatcherPath == Binding.GetSourcePath();
+			};
 
+		if (!CompiledDelegateDispatchers.ContainsByPredicate(FindBySourcePathPredicate))
+		{
 			FString Error;
 			TArray<FPropertyBindingPathIndirection> DispatcherIndirections;
 			if (!Binding.GetSourcePath().ResolveIndirectionsWithValue(InstanceDataView, DispatcherIndirections, &Error))
 			{
-				Log->Reportf(EMessageSeverity::Error, TargetStruct, TEXT("Resolving path in %s: %s"), *TargetStruct.ToString(), *Error);
-				return false;
+				Log->Reportf(EMessageSeverity::Error, SourceStruct, TEXT("Resolving path in %s: %s"), *DispatcherStruct->ToString(), *Error);
+				bSuccess = false;
+				continue;
 			}
 
-			FPropertyBindingPathIndirection& DispatcherLeafIndirection = DispatcherIndirections.Last();
-
+			const FPropertyBindingPathIndirection& DispatcherLeafIndirection = DispatcherIndirections.Last();
 			const FStructProperty* LeafAsStructProperty = CastField<FStructProperty>(DispatcherLeafIndirection.GetProperty());
-			check(LeafAsStructProperty && LeafAsStructProperty->Struct == FStateTreeDelegateDispatcher::StaticStruct())
-
-			FStateTreeDelegateDispatcher* Dispatcher = reinterpret_cast<FStateTreeDelegateDispatcher*>(DispatcherLeafIndirection.GetMutablePropertyAddress());
-			Dispatcher->ID = FStateTreeIndex16(CompiledDelegateDispatchers.IsEmpty() ? 0 : CompiledDelegateDispatchers.Num());
-			CompiledDelegateDispatchers.Add({Binding.GetSourcePath(), Dispatcher->ID});
-		}
-	}
-
-	// Compiling listeners
-	{
-		for (const FStateTreePropertyPathBinding& Binding : DelegateTargetBindings)
-		{
-			if (Binding.GetTargetPath().GetStructID() != TargetStruct.ID)
+			if (LeafAsStructProperty == nullptr || LeafAsStructProperty->Struct != FStateTreeDelegateDispatcher::StaticStruct())
 			{
+				Log->Reportf(EMessageSeverity::Error, SourceStruct, TEXT("The source is not a valid delegate dispatcher."));
+				bSuccess = false;
 				continue;
 			}
 
-			// Source must be in the source array
-			const FStateTreeBindableStructDesc* DispatcherStruct = GetSourceStructDescByID(Binding.GetSourcePath().GetStructID());
-			if (!DispatcherStruct)
+			if (DispatcherLeafIndirection.GetContainerAddress() == nullptr)
 			{
-				Log->Reportf(EMessageSeverity::Error, TargetStruct,
-					TEXT("Could not find a binding source."));
-				return false;
+				Log->Reportf(EMessageSeverity::Error, SourceStruct, TEXT("The dispatcher can't be initialized."));
+				bSuccess = false;
+				continue;
 			}
 
-			FString Error;
-			TArray<FPropertyBindingPathIndirection> DispatcherIndirections;
-			TArray<FPropertyBindingPathIndirection> ListenerIndirections;
-		
-			if (!Binding.GetSourcePath().ResolveIndirections(DispatcherStruct->Struct, DispatcherIndirections, &Error))
+			FStateTreeDelegateDispatcher* Dispatcher = reinterpret_cast<FStateTreeDelegateDispatcher*>(DispatcherLeafIndirection.GetMutablePropertyAddress());
+			if (Dispatcher == nullptr)
 			{
-				Log->Reportf(EMessageSeverity::Error, TargetStruct, TEXT("Resolving path in %s: %s"), *DispatcherStruct->ToString(), *Error);
-				return false;
+				Log->Reportf(EMessageSeverity::Error, SourceStruct, TEXT("The dispatcher can't be initialized."));
+				bSuccess = false;
+				continue;
 			}
 
-			if (!Binding.GetTargetPath().ResolveIndirectionsWithValue(InstanceDataView, ListenerIndirections, &Error))
+			if (const FStateTreeEditorDelegateDispatcherCompiledBinding* PreviousCompiled = PreviousCompiledDistpatchers.FindByPredicate(FindBySourcePathPredicate))
 			{
-				Log->Reportf(EMessageSeverity::Error, TargetStruct, TEXT("Resolving path in %s: %s"), *TargetStruct.ToString(), *Error);
-				return false;
+				// Reuse previous ID
+				*Dispatcher = PreviousCompiled->ID;
 			}
-
-			const FStateTreeIndex16 DispatcherID = GetDispatcherIDFromPath(Binding.GetSourcePath());
-
-			if (!DispatcherID.IsValid())
+			else
 			{
-				Log->Reportf(EMessageSeverity::Error, TargetStruct, TEXT("Delegate Listener %s is bound to unknown dispatcher %s"), *TargetStruct.ToString(), *DispatcherStruct->ToString());
-				return false;
+				Dispatcher->ID = FGuid::NewGuid();
 			}
 
-			FPropertyBindingPathIndirection& ListenerLeafIndirection = ListenerIndirections.Last();
-			const FStructProperty* LeafAsStructProperty = CastField<FStructProperty>(ListenerLeafIndirection.GetProperty());
-			check(LeafAsStructProperty && LeafAsStructProperty->Struct == FStateTreeDelegateListener::StaticStruct())
-			FStateTreeDelegateListener* Listener = reinterpret_cast<FStateTreeDelegateListener*>(ListenerLeafIndirection.GetMutablePropertyAddress());
-			Listener->DispatcherID = DispatcherID;
-			Listener->ListenerID = FStateTreeIndex16(ListenersNum++);
+			const bool bFoundID = CompiledDelegateDispatchers.ContainsByPredicate([ID = Dispatcher->ID](const FStateTreeEditorDelegateDispatcherCompiledBinding& Other)
+				{
+					return Other.ID.ID == ID;
+				});
+			ensureMsgf(!bFoundID, TEXT("The ID is already used by another delegate dispatcher."));
+			CompiledDelegateDispatchers.Add({ Binding.GetSourcePath(), *Dispatcher });
 		}
 	}
 
-	return true;
+	return bSuccess;
 }
-
-FStateTreeIndex16 FStateTreePropertyBindingCompiler::GetDispatcherIDFromPath(const FPropertyBindingPath& PathToDispatcher) const
+	
+bool FStateTreePropertyBindingCompiler::CompileDelegateListeners(
+	const FStateTreeBindableStructDesc& TargetStruct,
+	TConstArrayView<FStateTreePropertyPathBinding> DelegateSourceBindings,
+	FStateTreeDataView InstanceDataView)
 {
-	const FCompiledDelegateDispatcher* FoundDispatcher = Algo::FindBy(CompiledDelegateDispatchers, PathToDispatcher, [](const FCompiledDelegateDispatcher& Dispatcher) { return Dispatcher.Path; });
-	if (FoundDispatcher)
+	check(Log);
+	check(PropertyBindings);
+
+	StoreSourceStructs();
+
+	bool bSuccess = true;
+
+	for (const FStateTreePropertyPathBinding& Binding : DelegateSourceBindings)
 	{
-		return FoundDispatcher->ID;
+		if (Binding.GetTargetPath().GetStructID() != TargetStruct.ID)
+		{
+			continue;
+		}
+
+		// Source must be in the source array
+		const FStateTreeBindableStructDesc* DispatcherStruct = GetSourceStructDescByID(Binding.GetSourcePath().GetStructID());
+		if (!DispatcherStruct)
+		{
+			Log->Reportf(EMessageSeverity::Error, TargetStruct, TEXT("Could not find a binding source."));
+			bSuccess = false;
+			continue;
+		}
+
+		FString Error;
+		TArray<FPropertyBindingPathIndirection> ListenerIndirections;
+		if (!Binding.GetTargetPath().ResolveIndirectionsWithValue(InstanceDataView, ListenerIndirections, &Error))
+		{
+			Log->Reportf(EMessageSeverity::Error, TargetStruct, TEXT("Resolving path in %s: %s"), *TargetStruct.ToString(), *Error);
+			bSuccess = false;
+			continue;
+		}
+
+		const FStateTreeDelegateDispatcher Dispatcher = GetDispatcherFromPath(Binding.GetSourcePath());
+		if (!Dispatcher.IsValid())
+		{
+			Log->Reportf(EMessageSeverity::Error, TargetStruct, TEXT("Delegate Listener %s is bound to unknown dispatcher %s"), *TargetStruct.ToString(), *DispatcherStruct->ToString());
+			bSuccess = false;
+			continue;
+		}
+
+		FPropertyBindingPathIndirection& ListenerLeafIndirection = ListenerIndirections.Last();
+		const FStructProperty* LeafAsStructProperty = CastField<FStructProperty>(ListenerLeafIndirection.GetProperty());
+		check(LeafAsStructProperty && LeafAsStructProperty->Struct == FStateTreeDelegateListener::StaticStruct());
+		if (LeafAsStructProperty == nullptr || LeafAsStructProperty->Struct != FStateTreeDelegateListener::StaticStruct())
+		{
+			Log->Reportf(EMessageSeverity::Error, TargetStruct, TEXT("The target is not a valid delegate listener."));
+			bSuccess = false;
+			continue;
+		}
+
+		FStateTreeDelegateListener* Listener = reinterpret_cast<FStateTreeDelegateListener*>(ListenerLeafIndirection.GetMutablePropertyAddress());
+		Listener->Dispatcher = Dispatcher;
+		Listener->ID = ++ListenersNum;
 	}
 
-	return FStateTreeIndex16::Invalid;
+	return bSuccess;
 }
 
 bool FStateTreePropertyBindingCompiler::CompileReferences(const FStateTreeBindableStructDesc& TargetStruct, TConstArrayView<FStateTreePropertyPathBinding> PropertyReferenceBindings, FStateTreeDataView InstanceDataView, const TMap<FGuid, const FStateTreeDataView>& IDToStructValue)
@@ -389,6 +434,26 @@ int32 FStateTreePropertyBindingCompiler::AddSourceStruct(const FStateTreeBindabl
 	return SourceStructs.Num() - 1;
 }
 
+FStateTreeDelegateDispatcher FStateTreePropertyBindingCompiler::GetDispatcherFromPath(const FPropertyBindingPath& PathToDispatcher) const
+{
+	const FStateTreeEditorDelegateDispatcherCompiledBinding* FoundDispatcher = CompiledDelegateDispatchers.FindByPredicate(
+		[&PathToDispatcher](const FStateTreeEditorDelegateDispatcherCompiledBinding& Dispatcher)
+		{
+			return PathToDispatcher == Dispatcher.DispatcherPath;
+		});
+	if (FoundDispatcher)
+	{
+		return FoundDispatcher->ID;
+	}
+
+	return FStateTreeDelegateDispatcher();
+}
+
+TArray<FStateTreeEditorDelegateDispatcherCompiledBinding> FStateTreePropertyBindingCompiler::GetCompiledDelegateDispatchers() const
+{
+	return CompiledDelegateDispatchers;
+}
+
 void FStateTreePropertyBindingCompiler::StoreSourceStructs()
 {
 	// Check that existing structs are compatible
diff --git a/Engine/Plugins/Runtime/StateTree/Source/StateTreeEditorModule/Public/StateTreeCompiler.h b/Engine/Plugins/Runtime/StateTree/Source/StateTreeEditorModule/Public/StateTreeCompiler.h
--- a/Engine/Plugins/Runtime/StateTree/Source/StateTreeEditorModule/Public/StateTreeCompiler.h
+++ b/Engine/Plugins/Runtime/StateTree/Source/StateTreeEditorModule/Public/StateTreeCompiler.h
@@ -41,9 +41,9 @@ private:
 	FStateTreeStateHandle GetStateHandle(const FGuid& StateID) const;
 	UStateTreeState* GetState(const FGuid& StateID) const;
 
+	bool CreateParameters();
 	bool CreateStates();
 	bool CreateStateRecursive(UStateTreeState& State, const FStateTreeStateHandle Parent);
-	
 	bool CreateEvaluators();
 	bool CreateGlobalTasks();
 	bool CreateStateTasksAndParameters();
@@ -70,8 +70,8 @@ private:
 	struct FValidatedPathBindings
 	{
 		TArray<FStateTreePropertyPathBinding> CopyBindings;
-		TArray<FStateTreePropertyPathBinding> DelegateSourceBindings;
-		TArray<FStateTreePropertyPathBinding> DelegateTargetBindings;
+		TArray<FStateTreePropertyPathBinding> DelegateDispatchers;
+		TArray<FStateTreePropertyPathBinding> DelegateListeners;
 		TArray<FStateTreePropertyPathBinding> ReferenceBindings;
 	};
 
diff --git a/Engine/Plugins/Runtime/StateTree/Source/StateTreeEditorModule/Public/StateTreeEditorData.h b/Engine/Plugins/Runtime/StateTree/Source/StateTreeEditorModule/Public/StateTreeEditorData.h
--- a/Engine/Plugins/Runtime/StateTree/Source/StateTreeEditorModule/Public/StateTreeEditorData.h
+++ b/Engine/Plugins/Runtime/StateTree/Source/StateTreeEditorModule/Public/StateTreeEditorData.h
@@ -357,4 +357,11 @@ public:
 	 */
 	UPROPERTY(Transient)
 	TArray<FStateTreeEditorBreakpoint> Breakpoints;
+
+	/**
+	 * List of the previous compiled delegate dispatchers.
+	 * Saved in the editor data to be duplicated transient.
+	 */
+	UPROPERTY(DuplicateTransient)
+	TArray<FStateTreeEditorDelegateDispatcherCompiledBinding> CompiledDispatchers;
 };
diff --git a/Engine/Plugins/Runtime/StateTree/Source/StateTreeEditorModule/Public/StateTreeEditorTypes.h b/Engine/Plugins/Runtime/StateTree/Source/StateTreeEditorModule/Public/StateTreeEditorTypes.h
--- a/Engine/Plugins/Runtime/StateTree/Source/StateTreeEditorModule/Public/StateTreeEditorTypes.h
+++ b/Engine/Plugins/Runtime/StateTree/Source/StateTreeEditorModule/Public/StateTreeEditorTypes.h
@@ -5,6 +5,8 @@
 #include "Containers/UnrealString.h"
 #include "Math/Color.h"
 #include "Misc/Guid.h"
+#include "PropertyBindingPath.h"
+#include "StateTreeDelegate.h"
 #include "UObject/Class.h"
 #include "StateTreeEditorTypes.generated.h"
 
@@ -93,3 +95,19 @@ struct TStructOpsTypeTraits<FStateTreeEditorColor> : TStructOpsTypeTraitsBase2<F
 		WithExportTextItem = true,
 	};
 };
+
+/**
+ * Struct describing a compiled delegate dispatcher.
+ * To help with delta cooking/compilation. We want the same generated ID every time the asset compiles.
+ */
+USTRUCT()
+struct FStateTreeEditorDelegateDispatcherCompiledBinding
+{
+	GENERATED_BODY()
+
+	UPROPERTY()
+	FPropertyBindingPath DispatcherPath;
+
+	UPROPERTY()
+	FStateTreeDelegateDispatcher ID;
+};
\ No newline at end of file
diff --git a/Engine/Plugins/Runtime/StateTree/Source/StateTreeEditorModule/Public/StateTreePropertyBindingCompiler.h b/Engine/Plugins/Runtime/StateTree/Source/StateTreeEditorModule/Public/StateTreePropertyBindingCompiler.h
--- a/Engine/Plugins/Runtime/StateTree/Source/StateTreeEditorModule/Public/StateTreePropertyBindingCompiler.h
+++ b/Engine/Plugins/Runtime/StateTree/Source/StateTreeEditorModule/Public/StateTreePropertyBindingCompiler.h
@@ -2,6 +2,8 @@
 
 #pragma once
 
+#include "StateTreeDelegate.h"
+#include "StateTreeEditorTypes.h"
 #include "StateTreePropertyBindings.h"
 #include "StateTreePropertyBindingCompiler.generated.h"
 
@@ -37,16 +39,23 @@ struct STATETREEEDITORMODULE_API FStateTreePropertyBindingCompiler
 	[[nodiscard]] bool CompileBatch(const FStateTreeBindableStructDesc& TargetStruct, TConstArrayView<FStateTreePropertyPathBinding> PropertyBindings, FStateTreeIndex16 PropertyFuncsBegin, FStateTreeIndex16 PropertyFuncsEnd, int32& OutBatchIndex);
 
 	 /**
-	  * Compiles delegates for selected struct.
-	  * @param TargetStruct - Description of the structs which contains the delegate dispatchers or listeners.
-	  * @param DelegateSourceBindings - Array of bindings to dispatchers to compile.
-	  * @param DelegateTargetBindings - Array of bindings from listeners to compile.
+	  * Compiles delegates dispatcher for selected struct.
+	  * @param SourceStruct - Description of the structs which contains the delegate dispatchers.
+	  * @param PreviousDistpatchers - Array of previous dispatchers IDs.
+	  * @param DelegateSourceBindings - Array of bindings dispatcher to compile.
 	  * @param InstanceDataView - View to the instance data
-	  * @param TargetAccessibleStructs - View to all accessible structs.
-	  * @param OutBatchIndex - Resulting batch index, if index is INDEX_NONE, no bindings were found and no batch was generated.
 	  * @return True on success, false on failure.
 	 */
-	[[nodiscard]] bool CompileDelegates(const FStateTreeBindableStructDesc& TargetStruct, TConstArrayView<FStateTreePropertyPathBinding> DelegateSourceBindings, TConstArrayView<FStateTreePropertyPathBinding> DelegateTargetBindings, FStateTreeDataView InstanceDataView);
+	[[nodiscard]] bool CompileDelegateDispatchers(const FStateTreeBindableStructDesc& SourceStruct, TConstArrayView<FStateTreeEditorDelegateDispatcherCompiledBinding> PreviousDistpatchers, TConstArrayView<FStateTreePropertyPathBinding> DelegateSourceBindings, FStateTreeDataView InstanceDataView);
+
+	 /**
+	  * Compiles delegates listener for selected struct.
+	  * @param TargetStruct - Description of the structs which contains the delegate listeners.
+	  * @param DelegateTargetBindings - Array of bindings listeners to compile.
+	  * @param InstanceDataView - View to the instance data
+	  * @return True on success, false on failure.
+	 */
+	[[nodiscard]] bool CompileDelegateListeners(const FStateTreeBindableStructDesc& TargetStruct, TConstArrayView<FStateTreePropertyPathBinding> DelegateSourceBindings, FStateTreeDataView InstanceDataView);
 
 	/**
 	  * Compiles references for selected struct
@@ -68,8 +77,11 @@ struct STATETREEEDITORMODULE_API FStateTreePropertyBindingCompiler
 	  */
 	int32 AddSourceStruct(const FStateTreeBindableStructDesc& SourceStruct);
 
-	/** @return Index of delegate dispatcher residing behind given path. */
-	FStateTreeIndex16 GetDispatcherIDFromPath(const FPropertyBindingPath& PathToDispatcher) const;
+	/** @return delegate dispatcher residing behind given path. */
+	FStateTreeDelegateDispatcher GetDispatcherFromPath(const FPropertyBindingPath& PathToDispatcher) const;
+
+	/** @return the list of compiled delegate dispatcher. */
+	TArray<FStateTreeEditorDelegateDispatcherCompiledBinding> GetCompiledDelegateDispatchers() const;
 
 	const FStateTreeBindableStructDesc* GetSourceStructDescByID(const FGuid& ID) const
 	{
@@ -94,9 +106,7 @@ protected:
 	UPROPERTY()
 	TArray<FStateTreeBindableStructDesc> SourceStructs;
 
-	/**
-	 * Representation of compiled reference.
-	 */
+	/** Representation of compiled reference. */
 	struct FCompiledReference
 	{
 		FPropertyBindingPath Path;
@@ -105,14 +115,9 @@ protected:
 
 	TArray<FCompiledReference> CompiledReferences;
 
-	struct FCompiledDelegateDispatcher
-	{
-		FPropertyBindingPath Path;
-		FStateTreeIndex16 ID;
-	};
-
-	TArray<FCompiledDelegateDispatcher> CompiledDelegateDispatchers;
-	uint16 ListenersNum = 0u;
+	/** Dispatcher used by bindings. */
+	TArray<FStateTreeEditorDelegateDispatcherCompiledBinding> CompiledDelegateDispatchers;
+	int32 ListenersNum = 0;
 
 	FStateTreePropertyBindings* PropertyBindings = nullptr;
 
diff --git a/Engine/Plugins/Runtime/StateTree/Source/StateTreeModule/Private/Blueprint/StateTreeNodeBlueprintBase.cpp b/Engine/Plugins/Runtime/StateTree/Source/StateTreeModule/Private/Blueprint/StateTreeNodeBlueprintBase.cpp
--- a/Engine/Plugins/Runtime/StateTree/Source/StateTreeModule/Private/Blueprint/StateTreeNodeBlueprintBase.cpp
+++ b/Engine/Plugins/Runtime/StateTree/Source/StateTreeModule/Private/Blueprint/StateTreeNodeBlueprintBase.cpp
@@ -113,14 +113,6 @@ void UStateTreeNodeBlueprintBase::RequestTransition(const FStateTreeStateLink& T
 	}
 }
 
-void UStateTreeNodeBlueprintBase::BroadcastDelegate(FStateTreeDelegateDispatcher Dispatcher)
-{
-	if (!WeakExecutionContext.BroadcastDelegate(Dispatcher))
-	{
-		UE_VLOG_UELOG(this, LogStateTree, Error, TEXT("Failed to broadcast the delegate. The instance probably stopped."));
-	}
-}
-
 bool UStateTreeNodeBlueprintBase::IsPropertyRefValid(const FStateTreeBlueprintPropertyRef& PropertyRef) const
 {
 	FProperty* SourceProperty = nullptr;
diff --git a/Engine/Plugins/Runtime/StateTree/Source/StateTreeModule/Private/Blueprint/StateTreeTaskBlueprintBase.cpp b/Engine/Plugins/Runtime/StateTree/Source/StateTreeModule/Private/Blueprint/StateTreeTaskBlueprintBase.cpp
--- a/Engine/Plugins/Runtime/StateTree/Source/StateTreeModule/Private/Blueprint/StateTreeTaskBlueprintBase.cpp
+++ b/Engine/Plugins/Runtime/StateTree/Source/StateTreeModule/Private/Blueprint/StateTreeTaskBlueprintBase.cpp
@@ -6,6 +6,7 @@
 #include "BlueprintNodeHelpers.h"
 #include "Engine/World.h"
 #include "TimerManager.h"
+#include "VisualLogger/VisualLogger.h"
 
 #include UE_INLINE_GENERATED_CPP_BY_NAME(StateTreeTaskBlueprintBase)
 
@@ -112,6 +113,31 @@ void UStateTreeTaskBlueprintBase::FinishTask(const bool bSucceeded)
 	}
 }
 
+void UStateTreeTaskBlueprintBase::BroadcastDelegate(FStateTreeDelegateDispatcher Dispatcher)
+{
+	if (!GetWeakExecutionContext().BroadcastDelegate(Dispatcher)) 
+	{
+		UE_VLOG_UELOG(this, LogStateTree, Error, TEXT("Failed to broadcast the delegate. The instance probably stopped."));
+	}
+}
+
+void UStateTreeTaskBlueprintBase::BindDelegate(const FStateTreeDelegateListener& Listener, const FStateTreeDynamicDelegate& Delegate)
+{
+	FSimpleDelegate SimpleDelegate = FSimpleDelegate::CreateUFunction(const_cast<UObject*>(Delegate.GetUObject()), Delegate.GetFunctionName());
+	if (!GetWeakExecutionContext().BindDelegate(WeakTaskRef, Listener, MoveTemp(SimpleDelegate)))
+	{
+		UE_VLOG_UELOG(this, LogStateTree, Error, TEXT("Failed to bind the delegate. The instance probably stopped."));
+	}
+}
+
+void UStateTreeTaskBlueprintBase::UnbindDelegate(const FStateTreeDelegateListener& Listener)
+{
+	if (!GetWeakExecutionContext().UnbindDelegate(Listener))
+	{
+		UE_VLOG_UELOG(this, LogStateTree, Error, TEXT("Failed to unbind the delegate. The instance probably stopped."));
+	}
+}
+
 //----------------------------------------------------------------------//
 //  FStateTreeBlueprintTaskWrapper
 //----------------------------------------------------------------------//
diff --git a/Engine/Plugins/Runtime/StateTree/Source/StateTreeModule/Private/StateTreeExecutionContext.cpp b/Engine/Plugins/Runtime/StateTree/Source/StateTreeModule/Private/StateTreeExecutionContext.cpp
--- a/Engine/Plugins/Runtime/StateTree/Source/StateTreeModule/Private/StateTreeExecutionContext.cpp
+++ b/Engine/Plugins/Runtime/StateTree/Source/StateTreeModule/Private/StateTreeExecutionContext.cpp
@@ -43,23 +43,25 @@ namespace UE::StateTree::ExecutionContext::Private
 
 		for (FStateTreeStateHandle ActiveState : CurrentFrame.ActiveStates)
 		{
-			if (const FCompactStateTreeState* State = StateTree->GetStateFromHandle(ActiveState))
+			const FCompactStateTreeState* State = StateTree->GetStateFromHandle(ActiveState);
+			check(State);
+
+			if (!State->bHasDelegateTriggerTransitions)
 			{
-				if (!State->bHasDelegateTriggerTransitions)
-				{
-					continue;
-				}
+				continue;
+			}
 
-				for (uint16 TransitionIndex = State->TransitionsBegin; TransitionIndex < State->TransitionsBegin + State->TransitionsNum; ++TransitionIndex)
+
+			const int32 TransitionEnd = State->TransitionsBegin + State->TransitionsNum;
+			for (int32 TransitionIndex = State->TransitionsBegin; TransitionIndex <TransitionEnd; ++TransitionIndex)
+			{
+				const FCompactStateTransition* Transition = StateTree->GetTransitionFromIndex(FStateTreeIndex16(TransitionIndex));
+				check(Transition);
+				if (Transition->RequiredDelegateDispatcher == Dispatcher)
 				{
-					if (const FCompactStateTransition* Transition = StateTree->GetTransitionFromIndex(FStateTreeIndex16(TransitionIndex)))
-					{
-						if (Transition->DelegateDispatcherID == Dispatcher.GetDelegateID())
-						{
-							Storage.MarkDelegateAsBroadcasted(Dispatcher, CurrentFrame.StateTree);
-							return true;
-						}
-					}
+					ensureMsgf(EnumHasAnyFlags(Transition->Trigger, EStateTreeTransitionTrigger::OnDelegate), TEXT("The transition should have both (a valid dispatcher and the OnDelegate flag) or none."));
+					Storage.MarkDelegateAsBroadcasted(Dispatcher);
+					return true;
 				}
 			}
 		}
@@ -653,7 +655,7 @@ EStateTreeRunStatus FStateTreeExecutionContext::Start(FStartParameters Parameter
 		Exec.TreeRunStatus = EStateTreeRunStatus::Running;
 		Exec.LastTickStatus = EStateTreeRunStatus::Unset;
 
-		static const FStateTreeStateHandle RootState = FStateTreeStateHandle(0);
+		static const FStateTreeStateHandle RootState = FStateTreeStateHandle::Root;
 
 		FStateSelectionResult StateSelectionResult;
 		if (SelectState(InitFrame, RootState, StateSelectionResult))
@@ -1035,7 +1037,7 @@ void FStateTreeExecutionContext::TickTriggerTransitionsInternal()
 	}
 }
 
-void FStateTreeExecutionContext::BroadcastDelegate(FStateTreeDelegateDispatcher Dispatcher)
+void FStateTreeExecutionContext::BroadcastDelegate(const FStateTreeDelegateDispatcher& Dispatcher)
 {
 	if (!Dispatcher.IsValid())
 	{
@@ -1052,31 +1054,46 @@ void FStateTreeExecutionContext::BroadcastDelegate(FStateTreeDelegateDispatcher
 	const FStateTreeExecutionFrame* CurrentFrame = GetCurrentlyProcessedFrame();
 	check(CurrentFrame);
 
-	GetExecState().DelegateActiveListeners.BroadcastDelegate(Dispatcher, CurrentFrame->StateTree, GetExecState());
+	GetExecState().DelegateActiveListeners.BroadcastDelegate(Dispatcher, GetExecState());
 	if (UE::StateTree::ExecutionContext::Private::MarkDelegateAsBroadcasted(Dispatcher, *CurrentFrame, GetMutableInstanceData()->GetMutableStorage()))
 	{
 		ScheduleNextTick();
 	}
 }
 
-bool FStateTreeExecutionContext::AddDelegateListener(FStateTreeDelegateListener Listener, FSimpleDelegate Delegate)
+bool FStateTreeExecutionContext::AddDelegateListener(const FStateTreeDelegateListener& Listener, FSimpleDelegate Delegate)
+{
+	BindDelegate(Listener, MoveTemp(Delegate));
+	return true;
+}
+
+void FStateTreeExecutionContext::BindDelegate(const FStateTreeDelegateListener& Listener, FSimpleDelegate Delegate)
 {
 	if (!Listener.IsValid())
 	{
-		return true;
+		// The listener is not bound to a dispatcher. It will never trigger the delegate.
+		return;
 	}
 
 	if (!IsValid())
 	{
 		STATETREE_LOG(Warning, TEXT("%hs: StateTree context is not initialized properly ('%s' using StateTree '%s')"),
 			__FUNCTION__, *GetNameSafe(&Owner), *GetFullNameSafe(&RootStateTree));
-		return false;
+		return;
 	}
 
-	return GetExecState().DelegateActiveListeners.AddListener(Listener, MoveTemp(Delegate), CurrentlyProcessedFrame->StateTree, CurrentlyProcessedState.Root, CurrentlyProcessedState, FStateTreeIndex16(CurrentNodeDataHandle.GetIndex()));
+	const FStateTreeExecutionFrame* CurrentFrame = GetCurrentlyProcessedFrame();
+	check(CurrentFrame);
+
+	GetExecState().DelegateActiveListeners.Add(Listener, MoveTemp(Delegate), CurrentlyProcessedFrame->StateTree, CurrentlyProcessedState.Root, CurrentlyProcessedState, FStateTreeIndex16(CurrentNodeDataHandle.GetIndex()));
+}
+
+void FStateTreeExecutionContext::RemoveDelegateListener(const FStateTreeDelegateListener& Listener)
+{
+	UnbindDelegate(Listener);
 }
 
-void FStateTreeExecutionContext::RemoveDelegateListener(FStateTreeDelegateListener Listener)
+void FStateTreeExecutionContext::UnbindDelegate(const FStateTreeDelegateListener& Listener)
 {
 	if (!Listener.IsValid())
 	{
@@ -1093,7 +1110,7 @@ void FStateTreeExecutionContext::RemoveDelegateListener(FStateTreeDelegateListen
 	const FStateTreeExecutionFrame* CurrentFrame = GetCurrentlyProcessedFrame();
 	check(CurrentFrame);
 
-	GetExecState().DelegateActiveListeners.RemoveListener(Listener, CurrentFrame->StateTree);
+	GetExecState().DelegateActiveListeners.Remove(Listener);
 }
 
 void FStateTreeExecutionContext::RequestTransition(const FStateTreeTransitionRequest& Request)
@@ -2683,7 +2700,7 @@ void FStateTreeExecutionContext::ExitState(const FStateTreeTransitionResult& Tra
 
 void FStateTreeExecutionContext::RemoveDeactivatedDelegateListeners(const FStateTreeExecutionFrame* ChangedFrame, FStateTreeStateHandle ChangedState, TConstArrayView<FStateTreeExecutionFrame> DeactivatedFrames)
 {
-	GetExecState().DelegateActiveListeners.RemoveListenersByPredicate([&](FStateTreeStateHandle OwningState, FStateTreeIndex16 OwningNodeIndex, const UStateTree* FrameStateTree, FStateTreeStateHandle FrameRootState)
+	GetExecState().DelegateActiveListeners.RemoveByPredicate([&](FStateTreeStateHandle OwningState, FStateTreeIndex16 OwningNodeIndex, const UStateTree* FrameStateTree, FStateTreeStateHandle FrameRootState)
 		{
 			if (FrameStateTree == nullptr)
 			{
@@ -4099,11 +4116,11 @@ bool FStateTreeExecutionContext::TriggerTransitions()
 				}
 				else if (EnumHasAnyFlags(Transition.Trigger, EStateTreeTransitionTrigger::OnDelegate))
 				{
-					if (InstanceDataStorage->IsDelegateBroadcasted(Transition.DelegateDispatcherID, CurrentFrame))
-				{
-					// Dummy event to make sure we iterate to loop below once.
-					TransitionEvents.Emplace(nullptr);
-				}
+					if (InstanceDataStorage->IsDelegateBroadcasted(Transition.RequiredDelegateDispatcher))
+					{
+						// Dummy event to make sure we iterate to loop below once.
+						TransitionEvents.Emplace(nullptr);
+					}
 				}
 				else
 				{
@@ -4784,8 +4801,6 @@ bool FStateTreeExecutionContext::SelectStateInternal(
 					NewFrame.GlobalInstanceIndexBase = ExistingFrame->GlobalInstanceIndexBase;
 					NewFrame.StateParameterDataHandle = ExistingFrame->StateParameterDataHandle;
 					NewFrame.GlobalParameterDataHandle = ExistingFrame->GlobalParameterDataHandle;
-					NewFrame.ActiveListenersIndexBase = ExistingFrame->ActiveListenersIndexBase;
-					NewFrame.DelegateDispatcherIndexBase = ExistingFrame->DelegateDispatcherIndexBase;
 				}
 				else
 				{
@@ -4868,8 +4883,6 @@ bool FStateTreeExecutionContext::SelectStateInternal(
 				NewFrame.StateParameterDataHandle = ExistingFrame->StateParameterDataHandle;
 				NewFrame.GlobalParameterDataHandle = ExistingFrame->GlobalParameterDataHandle;
 				NewFrame.ExternalDataBaseIndex = ExistingFrame->ExternalDataBaseIndex;
-				NewFrame.ActiveListenersIndexBase = ExistingFrame->ActiveListenersIndexBase;
-				NewFrame.DelegateDispatcherIndexBase = ExistingFrame->DelegateDispatcherIndexBase;
 			}
 			else
 			{
@@ -4998,11 +5011,21 @@ bool FStateTreeExecutionContext::SelectStateInternal(
 						}
 					}
 				}
-				else if (EnumHasAnyFlags(Transition.Trigger, EStateTreeTransitionTrigger::OnTick)
-					|| (EnumHasAnyFlags(Transition.Trigger, EStateTreeTransitionTrigger::OnDelegate) && InstanceDataStorage->IsDelegateBroadcasted(Transition.DelegateDispatcherID, CurrentFrame)))
+				else if (EnumHasAnyFlags(Transition.Trigger, EStateTreeTransitionTrigger::OnTick))
 				{
 					SelectedStateTransitionEvents.Emplace();
 				}
+				else if (EnumHasAnyFlags(Transition.Trigger, EStateTreeTransitionTrigger::OnDelegate))
+				{
+					if (InstanceDataStorage->IsDelegateBroadcasted(Transition.RequiredDelegateDispatcher))
+					{
+						SelectedStateTransitionEvents.Emplace();
+					}
+				}
+				else
+				{
+					ensureMsgf(false, TEXT("Missing a transition trigger type."));
+				}
 
 				for (const FStateTreeSharedEvent* SelectedStateTransitionEvent : SelectedStateTransitionEvents)
 				{
@@ -5621,7 +5644,7 @@ bool FStateTreeWeakExecutionContext::RequestTransition(FStateTreeStateHandle Tar
 	return false;
 }
 
-bool FStateTreeWeakExecutionContext::BroadcastDelegate(FStateTreeDelegateDispatcher Dispatcher) const
+bool FStateTreeWeakExecutionContext::BroadcastDelegate(const FStateTreeDelegateDispatcher& Dispatcher) const
 {
 	if (!Dispatcher.IsValid())
 	{
@@ -5655,7 +5678,7 @@ bool FStateTreeWeakExecutionContext::BroadcastDelegate(FStateTreeDelegateDispatc
 				const FStateTreeExecutionFrame* CurrentFrame = FStateTreeExecutionContext::FindFrame(FrameStateTreePinned.Get(), FrameRootState, Exec.ActiveFrames, ParentFrame);
 				if (CurrentFrame)
 				{
-					Exec.DelegateActiveListeners.BroadcastDelegate(Dispatcher, CurrentFrame->StateTree, Exec);
+					Exec.DelegateActiveListeners.BroadcastDelegate(Dispatcher, Exec);
 					if (UE::StateTree::ExecutionContext::Private::MarkDelegateAsBroadcasted(Dispatcher, *CurrentFrame, *StoragePinned))
 					{
 						ScheduleNextTick(OwnerPinned.Get(), RootStateTreePinned.Get(), *StoragePinned);
@@ -5671,15 +5694,15 @@ bool FStateTreeWeakExecutionContext::BroadcastDelegate(FStateTreeDelegateDispatc
 	return bBroadcasted;
 }
 
-bool FStateTreeWeakExecutionContext::RemoveDelegateListener(FStateTreeDelegateListener Listener) const
+bool FStateTreeWeakExecutionContext::BindDelegate(const FStateTreeWeakTaskRef& Task, const FStateTreeDelegateListener& Listener, FSimpleDelegate Delegate) const
 {
 	if (!Listener.IsValid())
 	{
-		// Listener is not bound, not an error.
+		// Nothings binds to the delegate, not an error.
 		return true;
 	}
 
-	bool bRemoved = false;
+	bool bResult = false;
 
 	TStrongObjectPtr<UObject> OwnerPinned = Owner.Pin();
 	TStrongObjectPtr<const UStateTree> StateTreePinned = StateTree.Pin();
@@ -5693,22 +5716,66 @@ bool FStateTreeWeakExecutionContext::RemoveDelegateListener(FStateTreeDelegateLi
 		FStateTreeExecutionState& Exec = StoragePinned->GetMutableExecutionState();
 		if (Exec.CurrentExecutionContext)
 		{
-			Exec.CurrentExecutionContext->RemoveDelegateListener(Listener);
-			bRemoved = true;
+			Exec.CurrentExecutionContext->BindDelegate(Listener, MoveTemp(Delegate));
+			bResult = true;
 		}
-		else if(Exec.ActiveFrames.Num() > 0)
+		else if (Exec.ActiveFrames.Num() > 0)
 		{
+			FStateTreeStrongTaskRef TaskPin = Task.Pin();
+			const FStateTreeTaskBase* TaskPtr = TaskPin.GetTask();
 			TStrongObjectPtr<const UStateTree> FrameStateTreePinned = FrameStateTree.Pin();
-			if (FrameStateTreePinned)
+			if (TaskPtr != nullptr && FrameStateTreePinned.Get() != nullptr && TaskPin.GetStateTree() == FrameStateTreePinned.Get())
 			{
-				Exec.DelegateActiveListeners.RemoveListener(Listener, FrameStateTreePinned.Get());
-				bRemoved = true;
+				Exec.DelegateActiveListeners.Add(Listener, MoveTemp(Delegate), FrameStateTreePinned.Get(), FrameRootState, CurrentState, TaskPin.GetTaskIndex());
+				bResult = true;
 			}
 		}
 
 		StoragePinned->ReleaseWriteAccess();
 	}
 
+	return bResult;
+}
+
+bool FStateTreeWeakExecutionContext::RemoveDelegateListener(const FStateTreeDelegateListener& Listener) const
+{
+	return UnbindDelegate(Listener);
+}
+
+bool FStateTreeWeakExecutionContext::UnbindDelegate(const FStateTreeDelegateListener& Listener) const
+{
+	if (!Listener.IsValid())
+	{
+		// The listener is not bound to a dispatcher. It will never trigger the delegate. It is not an error.
+		return true;
+	}
+
+	bool bRemoved = false;
+
+	TStrongObjectPtr<UObject> OwnerPinned = Owner.Pin();
+	TStrongObjectPtr<const UStateTree> StateTreePinned = StateTree.Pin();
+	TSharedPtr<FStateTreeInstanceStorage> StoragePinned = Storage.Pin();
+	if (OwnerPinned && StateTreePinned && StoragePinned)
+	{
+		// Manually acquire and release the write access.
+		//The context may already exist and we want to confirm that we are on the same thread.
+		StoragePinned->AcquireWriteAccess();
+
+		FStateTreeExecutionState& Exec = StoragePinned->GetMutableExecutionState();
+		if (Exec.CurrentExecutionContext)
+		{
+			Exec.CurrentExecutionContext->UnbindDelegate(Listener);
+			bRemoved = true;
+		}
+		else if (Exec.ActiveFrames.Num() > 0)
+		{
+			Exec.DelegateActiveListeners.Remove(Listener);
+			bRemoved = true;
+		}
+
+		StoragePinned->ReleaseWriteAccess();
+	}
+
 	return bRemoved;
 }
 
diff --git a/Engine/Plugins/Runtime/StateTree/Source/StateTreeModule/Private/StateTreeExecutionTypes.cpp b/Engine/Plugins/Runtime/StateTree/Source/StateTreeModule/Private/StateTreeExecutionTypes.cpp
--- a/Engine/Plugins/Runtime/StateTree/Source/StateTreeModule/Private/StateTreeExecutionTypes.cpp
+++ b/Engine/Plugins/Runtime/StateTree/Source/StateTreeModule/Private/StateTreeExecutionTypes.cpp
@@ -201,13 +201,12 @@ float FStateTreeScheduledTick::GetTickRate() const
 //----------------------------------------------------------------------//
 // FStateTreeDelegateActiveListeners
 //----------------------------------------------------------------------//
-FStateTreeDelegateActiveListeners::FActiveListener::FActiveListener(FStateTreeDelegateListener InListener, FSimpleDelegate InDelegate, TNotNull<const UStateTree*> InFrameStateTree, FStateTreeStateHandle InFrameRootState, FStateTreeStateHandle InOwningState, FStateTreeIndex16 InOwningNodeIndex)
-	: FrameStateTree(InFrameStateTree)
+FStateTreeDelegateActiveListeners::FActiveListener::FActiveListener(const FStateTreeDelegateListener& InListener, FSimpleDelegate InDelegate, TNotNull<const UStateTree*> InFrameStateTree, FStateTreeStateHandle InFrameRootState, FStateTreeStateHandle InOwningState, FStateTreeIndex16 InOwningNodeIndex)
+	: Listener(InListener)
 	, Delegate(MoveTemp(InDelegate))
-	, DispatcherID(InListener.GetBoundDispatcherID())
-	, ListenerID(InListener.GetListenerID())
 	, OwningState(InOwningState)
 	, OwningNodeIndex(InOwningNodeIndex)
+	, FrameStateTree(InFrameStateTree)
 	, FrameRootState(InFrameRootState)
 {}
 
@@ -216,30 +215,20 @@ FStateTreeDelegateActiveListeners::~FStateTreeDelegateActiveListeners()
 	check(BroadcastingLockCount == 0);
 }
 
-bool FStateTreeDelegateActiveListeners::AddListener(FStateTreeDelegateListener Listener, FSimpleDelegate Delegate, TNotNull<const UStateTree*> FrameStateTree, FStateTreeStateHandle FrameRootState, FStateTreeStateHandle OwningState, FStateTreeIndex16 OwningNodeIndex)
+void FStateTreeDelegateActiveListeners::Add(const FStateTreeDelegateListener& Listener, FSimpleDelegate Delegate, TNotNull<const UStateTree*> FrameStateTree, FStateTreeStateHandle FrameRootState, FStateTreeStateHandle OwningState, FStateTreeIndex16 OwningNodeIndex)
 {
 	check(Listener.IsValid());
-
-	for (FActiveListener& ActiveListener : Listeners)
-	{
-		if (ActiveListener.ListenerID == Listener.GetListenerID())
-		{
-			return false;
-		}
-	}
-
+	Remove(Listener);
 	Listeners.Emplace(Listener, MoveTemp(Delegate), FrameStateTree, FrameRootState, OwningState, OwningNodeIndex);
-	return true;
 }
 
-void FStateTreeDelegateActiveListeners::RemoveListener(FStateTreeDelegateListener Listener, TNotNull<const UStateTree*> FrameStateTree)
+void FStateTreeDelegateActiveListeners::Remove(const FStateTreeDelegateListener& Listener)
 {
 	check(Listener.IsValid());
 
-	FObjectKey FrameStateTreeKey = FrameStateTree;
-	const int32 Index = Listeners.IndexOfByPredicate([FrameStateTreeKey, &Listener](const FActiveListener& ActiveListener)
+	const int32 Index = Listeners.IndexOfByPredicate([Listener](const FActiveListener& ActiveListener)
 		{
-			return FrameStateTreeKey == ActiveListener.FrameStateTree && ActiveListener.ListenerID == Listener.GetListenerID();
+			return ActiveListener.Listener == Listener;
 		});
 
 	if (Index == INDEX_NONE)
@@ -258,7 +247,7 @@ void FStateTreeDelegateActiveListeners::RemoveListener(FStateTreeDelegateListene
 	}
 }
 
-void FStateTreeDelegateActiveListeners::RemoveListenersByPredicate(TFunctionRef<bool(FStateTreeStateHandle OwningState, FStateTreeIndex16 OwningNodeIndex, const UStateTree* FrameStateTree, FStateTreeStateHandle FrameRootState)> Predicate)
+void FStateTreeDelegateActiveListeners::RemoveByPredicate(TFunctionRef<bool(FStateTreeStateHandle OwningState, FStateTreeIndex16 OwningNodeIndex, const UStateTree* FrameStateTree, FStateTreeStateHandle FrameRootState)> Predicate)
 {
 	check(BroadcastingLockCount == 0);
 	Listeners.RemoveAllSwap([&](const FActiveListener& Listener)
@@ -267,23 +256,28 @@ void FStateTreeDelegateActiveListeners::RemoveListenersByPredicate(TFunctionRef<
 		});
 }
 
-void FStateTreeDelegateActiveListeners::BroadcastDelegate(FStateTreeDelegateDispatcher Dispatcher, TNotNull<const UStateTree*> FrameStateTree, const FStateTreeExecutionState& Exec)
+void FStateTreeDelegateActiveListeners::BroadcastDelegate(const FStateTreeDelegateDispatcher& Dispatcher, const FStateTreeExecutionState& Exec)
 {
 	check(Dispatcher.IsValid());
 
 	++BroadcastingLockCount;
 
-
-	FObjectKey FrameStateTreeKey = FrameStateTree;
-	for (int32 i = Listeners.Num() - 1; i >= 0; --i)
+	const int32 NumListeners = Listeners.Num();
+	for (int32 Index = 0; Index < NumListeners; ++Index)
 	{
-		const FActiveListener& Listener = Listeners[i];
-		if (Listener.OwningNodeIndex.Get() <= Exec.EnterStateFailedTaskIndex.Get()
-		&& Listener.OwningNodeIndex.Get() < Exec.LastExitedNodeIndex.Get()
-		&& Listener.DispatcherID == Dispatcher.GetDelegateID()
-		&& Listener.FrameStateTree == FrameStateTreeKey)
+		FActiveListener& Listener = Listeners[Index];
+		if (Listener.Listener.GetDispatcher() == Dispatcher
+			&& Listener.OwningNodeIndex.Get() <= Exec.EnterStateFailedTaskIndex.Get()
+			&& Listener.OwningNodeIndex.Get() < Exec.LastExitedNodeIndex.Get())
 		{
-			Listener.Delegate.ExecuteIfBound();
+			if (Exec.ActiveFrames.ContainsByPredicate(
+				[&Listener](const FStateTreeExecutionFrame& Other)
+				{
+					return Other.RootState == Listener.FrameRootState && FObjectKey(Other.StateTree) == Listener.FrameStateTree;
+				}))
+			{
+				Listener.Delegate.ExecuteIfBound();
+			}
 		}
 	}
 
@@ -291,11 +285,11 @@ void FStateTreeDelegateActiveListeners::BroadcastDelegate(FStateTreeDelegateDisp
 
 	if (BroadcastingLockCount == 0)
 	{
-		RemoveUnboundListeners();
+		RemoveUnbounds();
 	}
 }
 
-void FStateTreeDelegateActiveListeners::RemoveUnboundListeners()
+void FStateTreeDelegateActiveListeners::RemoveUnbounds()
 {
 	check(BroadcastingLockCount == 0);
 	if (!bContainsUnboundListeners)
diff --git a/Engine/Plugins/Runtime/StateTree/Source/StateTreeModule/Private/StateTreeInstanceData.cpp b/Engine/Plugins/Runtime/StateTree/Source/StateTreeModule/Private/StateTreeInstanceData.cpp
--- a/Engine/Plugins/Runtime/StateTree/Source/StateTreeModule/Private/StateTreeInstanceData.cpp
+++ b/Engine/Plugins/Runtime/StateTree/Source/StateTreeModule/Private/StateTreeInstanceData.cpp
@@ -207,19 +207,20 @@ void FStateTreeInstanceStorage::AddTransitionRequest(const UObject* Owner, const
 	TransitionRequests.Add(Request);
 }
 
-void FStateTreeInstanceStorage::MarkDelegateAsBroadcasted(FStateTreeDelegateDispatcher Dispatcher, TNotNull<const UStateTree*> BelongingStateTree)
+void FStateTreeInstanceStorage::MarkDelegateAsBroadcasted(const FStateTreeDelegateDispatcher& Dispatcher)
 {
-	BroadcastedDelegates.Add({BelongingStateTree, Dispatcher.GetDelegateID()});
+	// The array is reset once the transitions are processed.
+	BroadcastedDelegates.AddUnique(Dispatcher);
 }
 
-bool FStateTreeInstanceStorage::IsDelegateBroadcasted(FStateTreeIndex16 DispatcherID, const FStateTreeExecutionFrame& BelongingFrame) const
+bool FStateTreeInstanceStorage::IsDelegateBroadcasted(const FStateTreeDelegateDispatcher& Dispatcher) const
 {
-	return BroadcastedDelegates.Contains(FStateTreeBroadcastedDelegate{BelongingFrame.StateTree, DispatcherID});
+	return BroadcastedDelegates.Contains(Dispatcher);
 }
 
 void FStateTreeInstanceStorage::ResetBroadcastedDelegates()
 {
-	BroadcastedDelegates.Reset();
+	BroadcastedDelegates.Empty();
 }
 
 bool FStateTreeInstanceStorage::HasBroadcastedDelegates() const
diff --git a/Engine/Plugins/Runtime/StateTree/Source/StateTreeModule/Public/Blueprint/StateTreeNodeBlueprintBase.h b/Engine/Plugins/Runtime/StateTree/Source/StateTreeModule/Public/Blueprint/StateTreeNodeBlueprintBase.h
--- a/Engine/Plugins/Runtime/StateTree/Source/StateTreeModule/Public/Blueprint/StateTreeNodeBlueprintBase.h
+++ b/Engine/Plugins/Runtime/StateTree/Source/StateTreeModule/Public/Blueprint/StateTreeNodeBlueprintBase.h
@@ -49,10 +49,6 @@ public:
 	UFUNCTION(BlueprintCallable, Category = "StateTree", meta = (HideSelfPin = "true", DisplayName = "StateTree Request Transition"))
 	void RequestTransition(const FStateTreeStateLink& TargetState, const EStateTreeTransitionPriority Priority = EStateTreeTransitionPriority::Normal);
 
-	/** Broadcasts provided delegate. */
-	UFUNCTION(BlueprintCallable, Category = "StateTree", meta = (HideSelfPin = "true"))
-	void BroadcastDelegate(FStateTreeDelegateDispatcher Dispatcher);
-
 	/** Returns a reference to selected property in State Tree. */
 	UFUNCTION(CustomThunk)
 	void GetPropertyReference(const FStateTreeBlueprintPropertyRef& PropertyRef) const;
diff --git a/Engine/Plugins/Runtime/StateTree/Source/StateTreeModule/Public/Blueprint/StateTreeTaskBlueprintBase.h b/Engine/Plugins/Runtime/StateTree/Source/StateTreeModule/Public/Blueprint/StateTreeTaskBlueprintBase.h
--- a/Engine/Plugins/Runtime/StateTree/Source/StateTreeModule/Public/Blueprint/StateTreeTaskBlueprintBase.h
+++ b/Engine/Plugins/Runtime/StateTree/Source/StateTreeModule/Public/Blueprint/StateTreeTaskBlueprintBase.h
@@ -10,6 +10,8 @@
 
 struct FStateTreeExecutionContext;
 
+DECLARE_DYNAMIC_DELEGATE(FStateTreeDynamicDelegate);
+
 /*
  * Base class for Blueprint based Tasks.
  */
@@ -82,6 +84,22 @@ protected:
 	UFUNCTION(BlueprintCallable, Category = "StateTree", meta = (HideSelfPin = "true", DisplayName = "Finish Task"))
 	void FinishTask(const bool bSucceeded = true);
 
+	/** Broadcasts the dispatcher. It will triggers bound callback (bound with BindDelegate) and bound transitions. */
+	UFUNCTION(BlueprintCallable, Category = "StateTree")
+	void BroadcastDelegate(FStateTreeDelegateDispatcher Dispatcher);
+
+	/**
+	 * Registers the callback to the listener.
+	 * If the listener was previously registered, then unregister it first before registering it again with the new delegate callback.
+	 * The listener is bound to a dispatcher in the editor.
+	 */
+	UFUNCTION(BlueprintCallable, Category = "StateTree")
+	void BindDelegate(const FStateTreeDelegateListener& Listener, const FStateTreeDynamicDelegate& Delegate);
+	
+	/** Unregisters the callback bound to the listener. */
+	UFUNCTION(BlueprintCallable, Category = "StateTree")
+	void UnbindDelegate(const FStateTreeDelegateListener& Listener);
+
 	/** Cached task while the node is active for async execution. */
 	mutable FStateTreeWeakTaskRef WeakTaskRef;
 
diff --git a/Engine/Plugins/Runtime/StateTree/Source/StateTreeModule/Public/StateTreeDelegate.h b/Engine/Plugins/Runtime/StateTree/Source/StateTreeModule/Public/StateTreeDelegate.h
--- a/Engine/Plugins/Runtime/StateTree/Source/StateTreeModule/Public/StateTreeDelegate.h
+++ b/Engine/Plugins/Runtime/StateTree/Source/StateTreeModule/Public/StateTreeDelegate.h
@@ -5,8 +5,24 @@
 #include "StateTreeIndexTypes.h"
 #include "StateTreeDelegate.generated.h"
 
+
+/**
+ * StateTree's delegates are used to send events through the state's hierarchy.
+ * Similar to events but more control. As a designer, you know exactly which delegate can be broadcast in the state tree asset.
+ * Delegates can be used in transitions (like events) and to trigger callbacks.
+ * They are primarily used in async tasks. Instead of ticking every frame (polling), you can register a delegate and wait for the callback to request a new transition.
+ * Delegates are designed to be bound in the editor (dispatcher and listener).
+ * By design, you cannot share delegates between state tree assets.
+ *
+ * StateTree's delegates are composed of a sender (dispatcher) and a receiver (listener).
+ * The listener is linked to a simple callback at execution with FStateTreeExecutionContext.BindDelegate.
+ * You can only bind a listener to a dispatcher. (Cannot bind listener to listener or dispatcher to dispatcher).
+ *
+ * See StateTreeDelegateTests.cpp for examples.
+ */
+
  /**
-  * StateTree's internal delegate dispatcher.
+  * StateTree's delegate dispatcher.
   */
 USTRUCT(BlueprintType)
 struct STATETREEMODULE_API FStateTreeDelegateDispatcher
@@ -14,73 +30,55 @@ struct STATETREEMODULE_API FStateTreeDelegateDispatcher
 	GENERATED_BODY()
 
 public:
-	/**
-	 * Used internally.
-	 * @return ID of the delegate.
-	 */
-	FStateTreeIndex16 GetDelegateID() const
-	{
-		return ID;
-	}
-
-	/**
-	 * Used internally.
-	 * @return if dispatcher is valid.
-	 */
+	/** @return if dispatcher is valid. */
 	bool IsValid() const
 	{
 		return ID.IsValid();
 	}
 
+	bool operator==(const FStateTreeDelegateDispatcher&) const = default;
+	bool operator!=(const FStateTreeDelegateDispatcher&) const = default;
+
 private:
 	UPROPERTY()
-	FStateTreeIndex16 ID = FStateTreeIndex16::Invalid;
+	FGuid ID;
 
 	friend struct FStateTreePropertyBindingCompiler;
 };
 
 /**
- * StateTree's internal delegate listener.
+ * The receiver of a delegate binding.
+ * Can be bound in the editor to a delegate dispatcher.
  */
-USTRUCT()
+USTRUCT(BlueprintType)
 struct STATETREEMODULE_API FStateTreeDelegateListener
 {
 	GENERATED_BODY()
 
 public:
-	/**
-	 * Used internally.
-	 * @return ID of the bound delegate.
-	 */
-	FStateTreeIndex16 GetBoundDispatcherID() const
-	{
-		return DispatcherID;
-	}
-
-	/**
-	 * Used internally.
-	 * @return ID of the listener.
-	 */
-	FStateTreeIndex16 GetListenerID() const
+	/** @return the bound delegate. */
+	FStateTreeDelegateDispatcher GetDispatcher() const
 	{
-		return ListenerID;
+		return Dispatcher;
 	}
 
-	/**
-	 * Used internally.
-	 * @return if the listener is valid.
-	 */
+	/** @return if the listener is valid. */
 	bool IsValid() const
 	{
-		return DispatcherID.IsValid() && ListenerID.IsValid();
+		return Dispatcher.IsValid() && ID != 0;
 	}
 
+	bool operator==(const FStateTreeDelegateListener&) const = default;
+	bool operator!=(const FStateTreeDelegateListener&) const = default;
+
+
 private:
+	/** ID of the dispatcher that listener is or will be bound to. */
 	UPROPERTY()
-	FStateTreeIndex16 DispatcherID = FStateTreeIndex16::Invalid;
+	FStateTreeDelegateDispatcher Dispatcher;
 
-	UPROPERTY()
-	FStateTreeIndex16 ListenerID = FStateTreeIndex16::Invalid;
+	/** The generated ID of the listener. */
+	int32 ID = 0;
 
 	friend struct FStateTreePropertyBindingCompiler;
-};
\ No newline at end of file
+};
diff --git a/Engine/Plugins/Runtime/StateTree/Source/StateTreeModule/Public/StateTreeExecutionContext.h b/Engine/Plugins/Runtime/StateTree/Source/StateTreeModule/Public/StateTreeExecutionContext.h
--- a/Engine/Plugins/Runtime/StateTree/Source/StateTreeModule/Public/StateTreeExecutionContext.h
+++ b/Engine/Plugins/Runtime/StateTree/Source/StateTreeModule/Public/StateTreeExecutionContext.h
@@ -314,14 +314,29 @@ public:
 	/** @return the names of all the active state. */
 	TArray<FName> GetActiveStateNames() const;
 
-	/** Broadcasts the delegate. */
-	void BroadcastDelegate(FStateTreeDelegateDispatcher Dispatcher);
+	/**
+	 * Broadcasts the delegate.
+	 * It executes bound delegates immediately and triggers bound transitions (when transitions are evaluated).
+	 */
+	void BroadcastDelegate(const FStateTreeDelegateDispatcher& Dispatcher);
 
+	UE_DEPRECATED(5.6, "Use BindDelegate")
 	/** Adds delegate listener. */
-	bool AddDelegateListener(FStateTreeDelegateListener Listener, FSimpleDelegate Delegate);
+	bool AddDelegateListener(const FStateTreeDelegateListener& Listener, FSimpleDelegate Delegate);
+
+	/**
+	 * Registers the delegate to the listener.
+	 * If the listener was previously registered, then unregister it first before registering it again with the new delegate callback.
+	 * The listener is bound to a dispatcher in the editor.
+	 */
+	void BindDelegate(const FStateTreeDelegateListener& Listener, FSimpleDelegate Delegate);
 
+	UE_DEPRECATED(5.6, "Use UnbindDelegate")
 	/** Removes delegate listener. */
-	void RemoveDelegateListener(FStateTreeDelegateListener Listener);
+	void RemoveDelegateListener(const FStateTreeDelegateListener& Listener);
+
+	/** Unregisters the callback bound to the listener. */
+	void UnbindDelegate(const FStateTreeDelegateListener& Listener);
 
 	/**
 	 * Iterates over all events.
@@ -1149,7 +1164,10 @@ public:
 		return StateTree.Pin();
 	}
 
-	/** Sends event for the StateTree. */
+	/**
+	 * Sends event for the StateTree.
+	 * @return false if the context is not valid and the event could not be sent.
+	 */
 	bool SendEvent(const FGameplayTag Tag, const FConstStructView Payload = FConstStructView(), const FName Origin = FName()) const;
 
 	/**
@@ -1160,14 +1178,34 @@ public:
 	 * @param TargetState The state to transition to.
 	 * @param Priority The priority of the transition.
 	 * @param Fallback of the transition if it fails to select the target state.
+	 * @return false if the context is not valid and the request failed.
 	 */
 	bool RequestTransition(FStateTreeStateHandle TargetState, EStateTreeTransitionPriority Priority = EStateTreeTransitionPriority::Normal, EStateTreeSelectionFallback Fallback = EStateTreeSelectionFallback::None) const;
 
-	/** Broadcasts the delegate. */
-	bool BroadcastDelegate(FStateTreeDelegateDispatcher Dispatcher) const;
+	/**
+	 * Broadcasts the delegate.
+	 * It executes bound delegates immediately and triggers bound transitions (when transitions are evaluated).
+	 * @return false if the context is not valid and the broadcast failed.
+	 */
+	bool BroadcastDelegate(const FStateTreeDelegateDispatcher& Dispatcher) const;
+
+	/**
+	 * Registers the delegate to the listener.
+	 * If the listener was previously registered, then unregister it first before registering it again with the new delegate callback.
+	 * The listener is bound to a dispatcher in the editor.
+	 * @return false if the context is not valid and the bind failed.
+	 */
+	bool BindDelegate(const FStateTreeWeakTaskRef& Task, const FStateTreeDelegateListener& Listener, FSimpleDelegate Delegate) const;
 
+	UE_DEPRECATED(5.6, "Use UnbindDelegate")
 	/** Removes Delegate Listener. */
-	bool RemoveDelegateListener(FStateTreeDelegateListener Listener) const;
+	bool RemoveDelegateListener(const FStateTreeDelegateListener& Listener) const;
+
+	/** 
+	 * Unregisters the callback bound to the listener.
+	 * @return false if the context is not valid and the unbind failed.
+	 */
+	bool UnbindDelegate(const FStateTreeDelegateListener& Listener) const;
 
 	/**
 	 * Finishes a task.
diff --git a/Engine/Plugins/Runtime/StateTree/Source/StateTreeModule/Public/StateTreeExecutionExtension.h b/Engine/Plugins/Runtime/StateTree/Source/StateTreeModule/Public/StateTreeExecutionExtension.h
--- a/Engine/Plugins/Runtime/StateTree/Source/StateTreeModule/Public/StateTreeExecutionExtension.h
+++ b/Engine/Plugins/Runtime/StateTree/Source/StateTreeModule/Public/StateTreeExecutionExtension.h
@@ -19,12 +19,14 @@ struct STATETREEMODULE_API FStateTreeExecutionExtension
 public:
 	struct FContextParameters
 	{
-		FContextParameters(UObject& Owner, const UStateTree& StateTree, const FStateTreeInstanceStorage& InstanceData)
-			: Owner(Owner), StateTree(StateTree), InstanceData(InstanceData)
+		FContextParameters(UObject& Owner, const UStateTree& StateTree, FStateTreeInstanceStorage& InstanceData)
+			: Owner(Owner)
+			, StateTree(StateTree)
+			, InstanceData(InstanceData)
 		{}
 		UObject& Owner;
 		const UStateTree& StateTree;
-		const FStateTreeInstanceStorage& InstanceData;
+		FStateTreeInstanceStorage& InstanceData;
 	};
 
 	virtual ~FStateTreeExecutionExtension() = default;
diff --git a/Engine/Plugins/Runtime/StateTree/Source/StateTreeModule/Public/StateTreeExecutionTypes.h b/Engine/Plugins/Runtime/StateTree/Source/StateTreeModule/Public/StateTreeExecutionTypes.h
--- a/Engine/Plugins/Runtime/StateTree/Source/StateTreeModule/Public/StateTreeExecutionTypes.h
+++ b/Engine/Plugins/Runtime/StateTree/Source/StateTreeModule/Public/StateTreeExecutionTypes.h
@@ -2,6 +2,7 @@
 
 #pragma once
 #include "StateTreeTypes.h"
+#include "StateTreeDelegate.h"
 #include "StateTreeEvents.h"
 #include "StateTreeExecutionExtension.h"
 #include "StateTreeStatePath.h"
@@ -12,8 +13,6 @@
 struct FStateTreeExecutionContext;
 class UStateTree;
 struct FStateTreeInstanceStorage;
-struct FStateTreeDelegateDispatcher;
-struct FStateTreeDelegateListener;
 
 /**
  * Enumeration for the different update phases.
@@ -635,47 +634,43 @@ struct FStateTreeDelayedFinishedTask
 
 
 /** Describes added delegate listeners. */
-USTRUCT()
-struct STATETREEMODULE_API FStateTreeDelegateActiveListeners
+struct FStateTreeDelegateActiveListeners
 {
-	GENERATED_BODY()
+	STATETREEMODULE_API FStateTreeDelegateActiveListeners() = default;
+	STATETREEMODULE_API ~FStateTreeDelegateActiveListeners();
 
-	~FStateTreeDelegateActiveListeners();
+	/** Adds a delegate bound in the editor to the list. Safe to be called during broadcasting.*/
+	STATETREEMODULE_API void Add(const FStateTreeDelegateListener& Listener, FSimpleDelegate Delegate, TNotNull<const UStateTree*> FrameStateTree, FStateTreeStateHandle FrameRootState, FStateTreeStateHandle OwningState, FStateTreeIndex16 OwningNodeIndex);
 
-	/** Adds the listener to the list. */
-	bool AddListener(FStateTreeDelegateListener Listener, FSimpleDelegate Delegate, TNotNull<const UStateTree*> FrameStateTree, FStateTreeStateHandle FrameRootState, FStateTreeStateHandle OwningState, FStateTreeIndex16 OwningNodeIndex);
-	
-	/** Removes the listener. Safe to be called during broadcasting. */
-	void RemoveListener(FStateTreeDelegateListener Listener, TNotNull<const UStateTree*> FrameStateTree);
+	/** Removes a delegate bound in the editor from the list. Safe to be called during broadcasting. */
+	STATETREEMODULE_API void Remove(const FStateTreeDelegateListener& Listener);
 
 	/** Removes the listener by predicate. */
-	void RemoveListenersByPredicate(TFunctionRef<bool(FStateTreeStateHandle OwningState, FStateTreeIndex16 OwningNodeIndex, const UStateTree* FrameStateTree, FStateTreeStateHandle FrameRootState)> Predicate);
+	STATETREEMODULE_API void RemoveByPredicate(TFunctionRef<bool(FStateTreeStateHandle OwningState, FStateTreeIndex16 OwningNodeIndex, const UStateTree* FrameStateTree, FStateTreeStateHandle FrameRootState)> Predicate);
 
 	/** Broadcasts the listener by predicate. */
-	void BroadcastDelegate(FStateTreeDelegateDispatcher Dispatcher, TNotNull<const UStateTree*> FrameStateTree, const FStateTreeExecutionState& Exec);
+	STATETREEMODULE_API void BroadcastDelegate(const FStateTreeDelegateDispatcher& Dispatcher, const FStateTreeExecutionState& Exec);
 
 private:
-	void RemoveUnboundListeners();
+	void RemoveUnbounds();
 
 	struct FActiveListener
 	{
 		FActiveListener() = default;
-		FActiveListener(FStateTreeDelegateListener Listener, FSimpleDelegate InDelegate, TNotNull<const UStateTree*> FrameStateTree, FStateTreeStateHandle FrameRootState, FStateTreeStateHandle OwningState, FStateTreeIndex16 OwningNodeIndex);
+		FActiveListener(const FStateTreeDelegateListener& Listener, FSimpleDelegate InDelegate, TNotNull<const UStateTree*> FrameStateTree, FStateTreeStateHandle FrameRootState, FStateTreeStateHandle OwningState, FStateTreeIndex16 OwningNodeIndex);
 
 		bool IsValid() const
 		{
-			return Delegate.IsBound();
+			return Listener.IsValid() && Delegate.IsBound();
 		}
 
-		FObjectKey FrameStateTree;
+		FStateTreeDelegateListener Listener;
 		FSimpleDelegate Delegate;
 
-		FStateTreeIndex16 DispatcherID = FStateTreeIndex16::Invalid;
-		FStateTreeIndex16 ListenerID = FStateTreeIndex16::Invalid;
-
 		FStateTreeStateHandle OwningState = FStateTreeStateHandle::Invalid;
 		FStateTreeIndex16 OwningNodeIndex = FStateTreeIndex16::Invalid;
 
+		FObjectKey FrameStateTree;
 		FStateTreeStateHandle FrameRootState;
 	};
 
@@ -684,6 +679,7 @@ private:
 	uint32 bContainsUnboundListeners : 1 = false;
 };
 
+
 /** Describes an active branch of a State Tree. */
 USTRUCT(BlueprintType)
 struct STATETREEMODULE_API FStateTreeExecutionFrame
@@ -725,14 +721,6 @@ struct STATETREEMODULE_API FStateTreeExecutionFrame
 	UPROPERTY()
 	FStateTreeIndex16 ActiveInstanceIndexBase = FStateTreeIndex16::Invalid;
 
-	/** Index within the active listeners to the first listener belonging to the frame. */
-	UPROPERTY()
-	FStateTreeIndex16 ActiveListenersIndexBase = FStateTreeIndex16::Invalid;
-
-	/** Index of the first delegate dispatcher belonging to the frame. */
-	UPROPERTY()
-	FStateTreeIndex16 DelegateDispatcherIndexBase = FStateTreeIndex16::Invalid;
-
 	/** Handle to the state parameter data, exists in ParentFrame. */
 	UPROPERTY()
 	FStateTreeDataHandle StateParameterDataHandle = FStateTreeDataHandle::Invalid; 
@@ -806,7 +794,6 @@ PRAGMA_ENABLE_DEPRECATION_WARNINGS
 	FRandomStream RandomStream;
 
 	/** Active delegate listeners. */
-	UPROPERTY()
 	FStateTreeDelegateActiveListeners DelegateActiveListeners;
 
 #if WITH_STATETREE_TRACE
diff --git a/Engine/Plugins/Runtime/StateTree/Source/StateTreeModule/Public/StateTreeInstanceData.h b/Engine/Plugins/Runtime/StateTree/Source/StateTreeModule/Public/StateTreeInstanceData.h
--- a/Engine/Plugins/Runtime/StateTree/Source/StateTreeModule/Public/StateTreeInstanceData.h
+++ b/Engine/Plugins/Runtime/StateTree/Source/StateTreeModule/Public/StateTreeInstanceData.h
@@ -61,21 +61,6 @@ struct STATETREEMODULE_API FStateTreeTemporaryInstanceData
 	FInstancedStruct Instance;
 };
 
-/**
- * Represents a delegate that was broadcasted.
- */
-USTRUCT()
-struct FStateTreeBroadcastedDelegate
-{
-	GENERATED_BODY()
-
-	bool operator==(const FStateTreeBroadcastedDelegate&) const = default;
-
-	FObjectKey  FrameStateTree;
-
-	UPROPERTY()
-	FStateTreeIndex16 DelegateID = FStateTreeIndex16::Invalid;
-};
 
 struct STATETREEMODULE_API FStateTreeInstanceStorageCustomVersion
 {
@@ -156,11 +141,11 @@ struct STATETREEMODULE_API FStateTreeInstanceStorage
 	*/
 	void AddTransitionRequest(const UObject* Owner, const FStateTreeTransitionRequest& Request);
 
-	/** Marks delegate as broadcasted. */
-	void MarkDelegateAsBroadcasted(FStateTreeDelegateDispatcher Dispatcher, TNotNull<const UStateTree*> BelongingStateTree);
+	/** Marks delegate as broadcasted. Use for transitions. */
+	void MarkDelegateAsBroadcasted(const FStateTreeDelegateDispatcher& Dispatcher);
 
-	/** @return true if the delegate of the given index was broadcasted. */
-	bool IsDelegateBroadcasted(FStateTreeIndex16 DispatcherID, const FStateTreeExecutionFrame& BelongingFrame) const;
+	/** @return true if a delegate was broadcasted. */
+	bool IsDelegateBroadcasted(const FStateTreeDelegateDispatcher& Dispatcher) const;
 
 	/** Resets the list of broadcasted delegates. */
 	void ResetBroadcastedDelegates();
@@ -307,7 +292,7 @@ protected:
 	TSharedRef<FStateTreeEventQueue> EventQueue = MakeShared<FStateTreeEventQueue>();
 	
 	/** Array of broadcasted delegates. */
-	TArray<FStateTreeBroadcastedDelegate> BroadcastedDelegates;
+	TArray<FStateTreeDelegateDispatcher> BroadcastedDelegates;
 
 	/** Requested transitions */
 	UPROPERTY(Transient)
diff --git a/Engine/Plugins/Runtime/StateTree/Source/StateTreeModule/Public/StateTreeTypes.h b/Engine/Plugins/Runtime/StateTree/Source/StateTreeModule/Public/StateTreeTypes.h
--- a/Engine/Plugins/Runtime/StateTree/Source/StateTreeModule/Public/StateTreeTypes.h
+++ b/Engine/Plugins/Runtime/StateTree/Source/StateTreeModule/Public/StateTreeTypes.h
@@ -9,6 +9,7 @@
 #include "PropertyBindingDataView.h"
 #include "StructUtils/PropertyBag.h"
 #include "GameplayTagContainer.h"
+#include "StateTreeDelegate.h"
 #include "StateTreeIndexTypes.h"
 #include "StateTreeTypes.generated.h"
 
@@ -622,13 +623,13 @@ struct STATETREEMODULE_API FCompactStateTransition
 	UPROPERTY()
 	FCompactEventDesc RequiredEvent;
 
-	/** Index to first condition to test */
+	/** Delegate dispatcher the transition is waiting for. */
 	UPROPERTY()
-	uint16 ConditionsBegin = 0;
+	FStateTreeDelegateDispatcher RequiredDelegateDispatcher;
 
-	/** ID of the delegate dispatcher the transition is listening to. */
+	/** Index to first condition to test */
 	UPROPERTY()
-	FStateTreeIndex16 DelegateDispatcherID = FStateTreeIndex16::Invalid;
+	uint16 ConditionsBegin = 0;
 
 	/** Target state of the transition */
 	UPROPERTY()
diff --git a/Engine/Plugins/Runtime/StateTree/Source/StateTreeTestSuite/Private/StateTreeDelegateTest.cpp b/Engine/Plugins/Runtime/StateTree/Source/StateTreeTestSuite/Private/StateTreeDelegateTest.cpp
--- a/Engine/Plugins/Runtime/StateTree/Source/StateTreeTestSuite/Private/StateTreeDelegateTest.cpp
+++ b/Engine/Plugins/Runtime/StateTree/Source/StateTreeTestSuite/Private/StateTreeDelegateTest.cpp
@@ -225,10 +225,10 @@ struct FStateTreeTest_Delegate_SelfRemoval : FStateTreeTestBase
 		uint32 TriggersCounter = 0;
 
 		EditorData.AddPropertyBinding(DispatcherTask, GET_MEMBER_NAME_STRING_CHECKED(FTestTask_BroadcastDelegate_InstanceData, OnTickDelegate), CustomFuncTask, GET_MEMBER_NAME_STRING_CHECKED(FTestTask_CustomFuncOnDelegate_InstanceData, Listener));
-		CustomFuncTask.GetNode().CustomFunc = [&](const FStateTreeWeakExecutionContext& WeakContext, FStateTreeDelegateListener Listener)
+		CustomFuncTask.GetNode().CustomFunc = [&TriggersCounter](const FStateTreeWeakExecutionContext& WeakContext, FStateTreeDelegateListener Listener)
 			{
 				++TriggersCounter;
-				WeakContext.RemoveDelegateListener(Listener);
+				WeakContext.UnbindDelegate(Listener);
 			};
 
 		FStateTreeCompilerLog Log;
diff --git a/Engine/Plugins/Runtime/StateTree/Source/StateTreeTestSuite/Private/StateTreeTestTypes.h b/Engine/Plugins/Runtime/StateTree/Source/StateTreeTestSuite/Private/StateTreeTestTypes.h
--- a/Engine/Plugins/Runtime/StateTree/Source/StateTreeTestSuite/Private/StateTreeTestTypes.h
+++ b/Engine/Plugins/Runtime/StateTree/Source/StateTreeTestSuite/Private/StateTreeTestTypes.h
@@ -881,7 +881,7 @@ struct FTestTask_ListenDelegate : public FStateTreeTaskBase
 	{
 		FTestTask_ListenDelegate_InstanceData& InstanceData = Context.GetInstanceData<FTestTask_ListenDelegate_InstanceData>(*this);
 		FStateTreeTestLog& TestLog = Context.GetExternalData(LogHandle);
-		Context.AddDelegateListener(InstanceData.Listener, FSimpleDelegate::CreateLambda([&TestLog, this, InstanceDataRef = Context.GetInstanceDataStructRef(*this)]()
+		Context.BindDelegate(InstanceData.Listener, FSimpleDelegate::CreateLambda([&TestLog, this, InstanceDataRef = Context.GetInstanceDataStructRef(*this)]()
 			{
 				if (FInstanceDataType* InstanceData = InstanceDataRef.GetPtr())
 				{
@@ -904,7 +904,7 @@ struct FTestTask_ListenDelegate : public FStateTreeTaskBase
 		if (bRemoveOnExit)
 		{
 			FTestTask_ListenDelegate_InstanceData& InstanceData = Context.GetInstanceData<FTestTask_ListenDelegate_InstanceData>(*this);
-			Context.RemoveDelegateListener(InstanceData.Listener);
+			Context.UnbindDelegate(InstanceData.Listener);
 		}
 	}
 	
@@ -943,7 +943,7 @@ struct FTestTask_RebroadcastDelegate : public FStateTreeTaskBase
 	{
 		FTestTask_RebroadcastDelegate_InstanceData& InstanceData = Context.GetInstanceData<FTestTask_RebroadcastDelegate_InstanceData>(*this);
 
-		Context.AddDelegateListener(InstanceData.Listener, FSimpleDelegate::CreateLambda([Dispatcher = InstanceData.Dispatcher, WeakContext = Context.MakeWeakExecutionContext()]()
+		Context.BindDelegate(InstanceData.Listener, FSimpleDelegate::CreateLambda([Dispatcher = InstanceData.Dispatcher, WeakContext = Context.MakeWeakExecutionContext()]()
 			{
 				WeakContext.BroadcastDelegate(Dispatcher);
 			}));
@@ -954,7 +954,7 @@ struct FTestTask_RebroadcastDelegate : public FStateTreeTaskBase
 	virtual void ExitState(FStateTreeExecutionContext& Context, const FStateTreeTransitionResult& Transition) const override
 	{
 		FTestTask_RebroadcastDelegate_InstanceData& InstanceData = Context.GetInstanceData<FTestTask_RebroadcastDelegate_InstanceData>(*this);
-		Context.RemoveDelegateListener(InstanceData.Listener);
+		Context.UnbindDelegate(InstanceData.Listener);
 	}
 };
 
@@ -986,7 +986,7 @@ struct FTestTask_CustomFuncOnDelegate : public FStateTreeTaskBase
 	{
 		FTestTask_CustomFuncOnDelegate_InstanceData& InstanceData = Context.GetInstanceData<FTestTask_CustomFuncOnDelegate_InstanceData>(*this);
 
-		Context.AddDelegateListener(InstanceData.Listener, FSimpleDelegate::CreateLambda([Func = CustomFunc, Listener = InstanceData.Listener, WeakContext = Context.MakeWeakExecutionContext()]()
+		Context.BindDelegate(InstanceData.Listener, FSimpleDelegate::CreateLambda([Func = CustomFunc, Listener = InstanceData.Listener, WeakContext = Context.MakeWeakExecutionContext()]()
 			{
 				Func(WeakContext, Listener);
 			}));
@@ -998,7 +998,7 @@ struct FTestTask_CustomFuncOnDelegate : public FStateTreeTaskBase
 		if (bRemoveOnExit)
 		{
 			FTestTask_CustomFuncOnDelegate_InstanceData& InstanceData = Context.GetInstanceData<FTestTask_CustomFuncOnDelegate_InstanceData>(*this);
-			Context.RemoveDelegateListener(InstanceData.Listener);
+			Context.UnbindDelegate(InstanceData.Listener);
 		}
 	}
 
