StateTree: Updated the DebugText task to not tick.

diff --git a/Engine/Plugins/Runtime/StateTree/Source/StateTreeModule/Private/Tasks/StateTreeDebugTextTask.cpp b/Engine/Plugins/Runtime/StateTree/Source/StateTreeModule/Private/Tasks/StateTreeDebugTextTask.cpp
--- a/Engine/Plugins/Runtime/StateTree/Source/StateTreeModule/Private/Tasks/StateTreeDebugTextTask.cpp
+++ b/Engine/Plugins/Runtime/StateTree/Source/StateTreeModule/Private/Tasks/StateTreeDebugTextTask.cpp
@@ -9,6 +9,14 @@
 
 #define LOCTEXT_NAMESPACE "StateTree"
 
+FStateTreeDebugTextTask::FStateTreeDebugTextTask()
+{
+	bShouldCallTick = false;
+	// We do not want to change the ReferenceActor if it's bound.
+	bShouldCopyBoundPropertiesOnTick = false;
+	bShouldCopyBoundPropertiesOnExitState = false;
+}
+
 EStateTreeRunStatus FStateTreeDebugTextTask::EnterState(FStateTreeExecutionContext& Context, const FStateTreeTransitionResult& Transition) const
 {
 	if (!bEnabled)
diff --git a/Engine/Plugins/Runtime/StateTree/Source/StateTreeModule/Private/Tasks/StateTreeDebugTextTask.h b/Engine/Plugins/Runtime/StateTree/Source/StateTreeModule/Private/Tasks/StateTreeDebugTextTask.h
--- a/Engine/Plugins/Runtime/StateTree/Source/StateTreeModule/Private/Tasks/StateTreeDebugTextTask.h
+++ b/Engine/Plugins/Runtime/StateTree/Source/StateTreeModule/Private/Tasks/StateTreeDebugTextTask.h
@@ -28,7 +28,7 @@ struct STATETREEMODULE_API FStateTreeDebugTextTask : public FStateTreeTaskCommon
 
 	using FInstanceDataType = FStateTreeDebugTextTaskInstanceData;
 	
-	FStateTreeDebugTextTask() = default;
+	FStateTreeDebugTextTask();
 
 	virtual const UStruct* GetInstanceDataType() const override { return FInstanceDataType::StaticStruct(); }
 
