StateTree: Display the "bHasRequiredEventToEnter" as a condition in the state tree view.

diff --git a/Engine/Plugins/Runtime/StateTree/Source/StateTreeEditorModule/Private/SStateTreeView.cpp b/Engine/Plugins/Runtime/StateTree/Source/StateTreeEditorModule/Private/SStateTreeView.cpp
--- a/Engine/Plugins/Runtime/StateTree/Source/StateTreeEditorModule/Private/SStateTreeView.cpp
+++ b/Engine/Plugins/Runtime/StateTree/Source/StateTreeEditorModule/Private/SStateTreeView.cpp
@@ -449,11 +449,13 @@ void SStateTreeView::HandleModelStatesChanged(const TSet<UStateTreeState*>& Affe
 	// and Paste or Duplicate an element in the array), The TreeView needs to be rebuilt because new elements came in or old elements have gone or both.
 	// This will not rebuild the list when we change an inner property in a condition or in a task node because of InstanceStruct wrapper
 	if (PropertyChangedEvent.MemberProperty->GetFName() == GET_MEMBER_NAME_CHECKED(UStateTreeState, Tasks)
-		|| PropertyChangedEvent.MemberProperty->GetFName() == GET_MEMBER_NAME_CHECKED(UStateTreeState, EnterConditions))
+		|| PropertyChangedEvent.MemberProperty->GetFName() == GET_MEMBER_NAME_CHECKED(UStateTreeState, EnterConditions)
+		|| PropertyChangedEvent.MemberProperty->GetFName() == GET_MEMBER_NAME_CHECKED(UStateTreeState, bHasRequiredEventToEnter)
+		|| PropertyChangedEvent.MemberProperty->GetFName() == GET_MEMBER_NAME_CHECKED(UStateTreeState, RequiredEventToEnter))
 	{
 		TreeView->RebuildList();
 	}
-	// todo: find a better way to fetch/update Transition information so we don't have to hardcode the rules here
+	//@todo: find a better way to fetch/update Transition information so we don't have to hardcode the rules here
 	// We cache the information of Transitions when making SStateTreeViewRow and using cached info to display
 	// So we need to rebuild the list whenever the relevant information changes to update the view
 	// For State Transitions, this will cover any changes made on the Array and on each Transition
diff --git a/Engine/Plugins/Runtime/StateTree/Source/StateTreeEditorModule/Private/SStateTreeViewRow.cpp b/Engine/Plugins/Runtime/StateTree/Source/StateTreeEditorModule/Private/SStateTreeViewRow.cpp
--- a/Engine/Plugins/Runtime/StateTree/Source/StateTreeEditorModule/Private/SStateTreeViewRow.cpp
+++ b/Engine/Plugins/Runtime/StateTree/Source/StateTreeEditorModule/Private/SStateTreeViewRow.cpp
@@ -662,226 +662,286 @@ TSharedRef<SWidget> SStateTreeViewRow::MakeConditionsWidget(const TSharedPtr<SSc
 		return SNullWidget::NullWidget;
 	}
 
-	const TSharedRef<SWrapBox> ConditionsBox = SNew(SWrapBox)
-	.PreferredSize_Lambda([WeakOwnerViewBox = ViewBox.ToWeakPtr()]()
-	{
-		// Captured as weak ptr so we don't prevent our parent widget from being destroyed (circular pointer reference).
-		if (const TSharedPtr<SScrollBox> OwnerViewBox = WeakOwnerViewBox.Pin())
-		{
-			return FMath::Max(300, OwnerViewBox->GetTickSpaceGeometry().GetLocalSize().X - 200);
-		}
-		return 0.f;
-	});
-
-	if (State->EnterConditions.IsEmpty())
+	if (!State->bHasRequiredEventToEnter && State->EnterConditions.IsEmpty())
 	{
 		return SNullWidget::NullWidget;
 	}
 
-	const int32 NumConditions = State->EnterConditions.Num();
-
-	for (int32 ConditionIndex = 0; ConditionIndex < NumConditions; ConditionIndex++)
+	TSharedRef<SVerticalBox> VerticalBox = SNew(SVerticalBox);
+	if (State->bHasRequiredEventToEnter)
 	{
-		const FStateTreeEditorNode& ConditionNode = State->EnterConditions[ConditionIndex];
-		if (const FStateTreeConditionBase* Condition = ConditionNode.Node.GetPtr<FStateTreeConditionBase>())
-		{
-			const FGuid ConditionId = ConditionNode.ID;
-
-			auto IsConditionEnabledFunc = [WeakState = WeakState]
+		auto IsConditionEnabledFunc = [WeakState = WeakState]
 			{
 				const UStateTreeState* State = WeakState.Get();
 				return State && State->bEnabled;
 			};
 
-			auto IsForcedConditionVisibleFunc = [WeakState = WeakState, ConditionIndex]()
+		const FName PayloadStructName = State->RequiredEventToEnter.PayloadStruct ? State->RequiredEventToEnter.PayloadStruct->GetFName() : FName();
+		const FText Description = FText::Format(LOCTEXT("Condition", "<b>Tag(</>{0}<b>) Payload(</>{1}<b>)</>"), FText::FromName(State->RequiredEventToEnter.Tag.GetTagName()), FText::FromName(PayloadStructName));
+
+		VerticalBox->AddSlot()
+		[
+			SNew(SBorder)
+			.VAlign(VAlign_Center)
+			.BorderImage(FAppStyle::GetNoBrush())
+			.Padding(FMargin(4.0f, 2.0f, 4.0f, 0.0f))
+			.IsEnabled_Lambda(IsConditionEnabledFunc)
+			.Padding(0)
+			[
+				SNew(SBox)
+				.HeightOverride(GetDefault<UStateTreeEditorUserSettings>()->GetStatesViewNodeRowHeight())
+				[
+					SNew(SHorizontalBox)
+					// Icon
+					+ SHorizontalBox::Slot()
+					.VAlign(VAlign_Center)
+					.HAlign(HAlign_Left)
+					.Padding(FMargin(4.0f, 0.0f, 4.0f, 0.0f))
+					.AutoWidth()
+					[
+						SNew(SImage)
+						.Image(FStateTreeEditorStyle::Get().GetBrush(FName("StateTreeEditor.Conditions")))
+					]
+					// Desc
+					+ SHorizontalBox::Slot()
+					.VAlign(VAlign_Center)
+					.HAlign(HAlign_Left)
+					.AutoWidth()
+					[
+						SNew(SRichTextBlock)
+						.Text(Description)
+						.TextStyle(&FStateTreeEditorStyle::Get().GetWidgetStyle<FTextBlockStyle>("StateTree.Task.Title"))
+						.OverflowPolicy(ETextOverflowPolicy::Ellipsis)
+						.Clipping(EWidgetClipping::OnDemand)
+						+ SRichTextBlock::Decorator(FTextStyleDecorator::Create(TEXT(""), FStateTreeEditorStyle::Get().GetWidgetStyle<FTextBlockStyle>("StateTree.Task.Title")))
+						+ SRichTextBlock::Decorator(FTextStyleDecorator::Create(TEXT("b"), FStateTreeEditorStyle::Get().GetWidgetStyle<FTextBlockStyle>("StateTree.Task.Title.Bold")))
+					]
+				]
+			]
+		];
+	}
+
+	if (!State->EnterConditions.IsEmpty())
+	{
+		TSharedRef<SWrapBox> ConditionsBox = SNew(SWrapBox)
+			.PreferredSize_Lambda([WeakOwnerViewBox = ViewBox.ToWeakPtr()]()
 				{
-					const UStateTreeState* State = WeakState.Get();
-					if (State != nullptr && State->EnterConditions.IsValidIndex(ConditionIndex))
+					// Captured as weak ptr so we don't prevent our parent widget from being destroyed (circular pointer reference).
+					if (const TSharedPtr<SScrollBox> OwnerViewBox = WeakOwnerViewBox.Pin())
 					{
-						if (const FStateTreeConditionBase* Condition = State->EnterConditions[ConditionIndex].Node.GetPtr<FStateTreeConditionBase>())
-						{
-							return Condition->EvaluationMode != EStateTreeConditionEvaluationMode::Evaluated ? EVisibility::Visible : EVisibility::Hidden;
-						}
+						return FMath::Max(300, OwnerViewBox->GetTickSpaceGeometry().GetLocalSize().X - 200);
 					}
-					return EVisibility::Hidden;
-				};
+					return 0.f;
+				});
 
-			auto GetForcedConditionTooltipFunc = [WeakState = WeakState, ConditionIndex]()
+		const int32 NumConditions = State->EnterConditions.Num();
+		for (int32 ConditionIndex = 0; ConditionIndex < NumConditions; ConditionIndex++)
+		{
+			const FStateTreeEditorNode& ConditionNode = State->EnterConditions[ConditionIndex];
+			if (const FStateTreeConditionBase* Condition = ConditionNode.Node.GetPtr<FStateTreeConditionBase>())
+			{
+				const FGuid ConditionId = ConditionNode.ID;
+
+				auto IsConditionEnabledFunc = [WeakState = WeakState]
 				{
 					const UStateTreeState* State = WeakState.Get();
-					if (State != nullptr && State->EnterConditions.IsValidIndex(ConditionIndex))
+					return State && State->bEnabled;
+				};
+
+				auto IsForcedConditionVisibleFunc = [WeakState = WeakState, ConditionIndex]()
 					{
-						if (const FStateTreeConditionBase* Condition = State->EnterConditions[ConditionIndex].Node.GetPtr<FStateTreeConditionBase>())
+						const UStateTreeState* State = WeakState.Get();
+						if (State != nullptr && State->EnterConditions.IsValidIndex(ConditionIndex))
 						{
-							if (Condition->EvaluationMode == EStateTreeConditionEvaluationMode::ForcedTrue)
-							{
-								return LOCTEXT("ForcedTrueConditionTooltip", "This condition is not evaluated and result forced to 'true'.");
-							}
-							if (Condition->EvaluationMode == EStateTreeConditionEvaluationMode::ForcedFalse)
+							if (const FStateTreeConditionBase* Condition = State->EnterConditions[ConditionIndex].Node.GetPtr<FStateTreeConditionBase>())
 							{
-								return LOCTEXT("ForcedFalseConditionTooltip", "This condition is not evaluated and result forced to 'false'.");
+								return Condition->EvaluationMode != EStateTreeConditionEvaluationMode::Evaluated ? EVisibility::Visible : EVisibility::Hidden;
 							}
 						}
-					}
-					return FText::GetEmpty();
-				};
+						return EVisibility::Hidden;
+					};
 
-			auto GetForcedConditionImageFunc = [WeakState = WeakState, ConditionIndex]() -> const FSlateBrush*
-				{
-					const UStateTreeState* State = WeakState.Get();
-					if (State != nullptr && State->EnterConditions.IsValidIndex(ConditionIndex))
+				auto GetForcedConditionTooltipFunc = [WeakState = WeakState, ConditionIndex]()
 					{
-						if (const FStateTreeConditionBase* Condition = State->EnterConditions[ConditionIndex].Node.GetPtr<FStateTreeConditionBase>())
+						const UStateTreeState* State = WeakState.Get();
+						if (State != nullptr && State->EnterConditions.IsValidIndex(ConditionIndex))
 						{
-							if (Condition->EvaluationMode == EStateTreeConditionEvaluationMode::ForcedTrue)
+							if (const FStateTreeConditionBase* Condition = State->EnterConditions[ConditionIndex].Node.GetPtr<FStateTreeConditionBase>())
 							{
-								return FStateTreeEditorStyle::Get().GetBrush("StateTreeEditor.Debugger.Condition.Passed");
+								if (Condition->EvaluationMode == EStateTreeConditionEvaluationMode::ForcedTrue)
+								{
+									return LOCTEXT("ForcedTrueConditionTooltip", "This condition is not evaluated and result forced to 'true'.");
+								}
+								if (Condition->EvaluationMode == EStateTreeConditionEvaluationMode::ForcedFalse)
+								{
+									return LOCTEXT("ForcedFalseConditionTooltip", "This condition is not evaluated and result forced to 'false'.");
+								}
 							}
-							if (Condition->EvaluationMode == EStateTreeConditionEvaluationMode::ForcedFalse)
+						}
+						return FText::GetEmpty();
+					};
+
+				auto GetForcedConditionImageFunc = [WeakState = WeakState, ConditionIndex]() -> const FSlateBrush*
+					{
+						const UStateTreeState* State = WeakState.Get();
+						if (State != nullptr && State->EnterConditions.IsValidIndex(ConditionIndex))
+						{
+							if (const FStateTreeConditionBase* Condition = State->EnterConditions[ConditionIndex].Node.GetPtr<FStateTreeConditionBase>())
 							{
-								return FStateTreeEditorStyle::Get().GetBrush("StateTreeEditor.Debugger.Condition.Failed");
+								if (Condition->EvaluationMode == EStateTreeConditionEvaluationMode::ForcedTrue)
+								{
+									return FStateTreeEditorStyle::Get().GetBrush("StateTreeEditor.Debugger.Condition.Passed");
+								}
+								if (Condition->EvaluationMode == EStateTreeConditionEvaluationMode::ForcedFalse)
+								{
+									return FStateTreeEditorStyle::Get().GetBrush("StateTreeEditor.Debugger.Condition.Failed");
+								}
 							}
 						}
-					}
-					return nullptr;
-				};
+						return nullptr;
+					};
 
-			ConditionsBox->AddSlot()
-				[
-					SNew(SBorder)
-					.VAlign(VAlign_Center)
-					.BorderImage(FAppStyle::GetNoBrush())
-					.IsEnabled_Lambda(IsConditionEnabledFunc)
-					.Padding(0)
+				ConditionsBox->AddSlot()
 					[
-						SNew(SBox)
-						.HeightOverride(GetDefault<UStateTreeEditorUserSettings>()->GetStatesViewNodeRowHeight())
+						SNew(SBorder)
+						.VAlign(VAlign_Center)
+						.BorderImage(FAppStyle::GetNoBrush())
+						.IsEnabled_Lambda(IsConditionEnabledFunc)
+						.Padding(0)
 						[
-							SNew(SHorizontalBox)
-
-							// Operand
-							+ SHorizontalBox::Slot()
-							.AutoWidth()
-							.VAlign(VAlign_Center)
+							SNew(SBox)
+							.HeightOverride(GetDefault<UStateTreeEditorUserSettings>()->GetStatesViewNodeRowHeight())
 							[
-								SNew(SBox)
-								.Padding(FMargin(4, 2, 4, 0))
+								SNew(SHorizontalBox)
+
+								// Operand
+								+ SHorizontalBox::Slot()
+								.AutoWidth()
 								.VAlign(VAlign_Center)
 								[
-									SNew(STextBlock)
-									.TextStyle(FStateTreeEditorStyle::Get(), "StateTree.Node.Operand")
-									.Text(this, &SStateTreeViewRow::GetOperandText, ConditionIndex)
+									SNew(SBox)
+									.Padding(FMargin(4, 2, 4, 0))
+									.VAlign(VAlign_Center)
+									[
+										SNew(STextBlock)
+										.TextStyle(FStateTreeEditorStyle::Get(), "StateTree.Node.Operand")
+										.Text(this, &SStateTreeViewRow::GetOperandText, ConditionIndex)
+									]
 								]
-							]
-							// Open parens
-							+ SHorizontalBox::Slot()
-							.AutoWidth()
-							.VAlign(VAlign_Center)
-							[
-								SNew(SBox)
-								.Padding(FMargin(FMargin(0.0f, 1.0f, 0.0f, 0.0f)))
+								// Open parens
+								+ SHorizontalBox::Slot()
+								.AutoWidth()
+								.VAlign(VAlign_Center)
 								[
-									SNew(STextBlock)
-									.TextStyle(FStateTreeEditorStyle::Get(), "StateTree.Task.Title")
-									.Text(this, &SStateTreeViewRow::GetOpenParens, ConditionIndex)
+									SNew(SBox)
+									.Padding(FMargin(FMargin(0.0f, 1.0f, 0.0f, 0.0f)))
+									[
+										SNew(STextBlock)
+										.TextStyle(FStateTreeEditorStyle::Get(), "StateTree.Task.Title")
+										.Text(this, &SStateTreeViewRow::GetOpenParens, ConditionIndex)
+									]
 								]
-							]
-							// Open parens
-							+ SHorizontalBox::Slot() 
-							.AutoWidth()
-							.VAlign(VAlign_Center)
-							[
-								SNew(SOverlay)
-								+ SOverlay::Slot()
+								// Open parens
+								+ SHorizontalBox::Slot() 
+								.AutoWidth()
+								.VAlign(VAlign_Center)
 								[
+									SNew(SOverlay)
+									+ SOverlay::Slot()
+									[
 
-									SNew(SButton)
-									.ButtonStyle(FAppStyle::Get(), "SimpleButton")
-									.ContentPadding(FMargin(2.f, 0.f))
-									.OnClicked_Lambda([StateTreeViewModel = StateTreeViewModel, WeakState = WeakState, ConditionId]()
-									{
-										StateTreeViewModel->BringNodeToFocus(WeakState.Get(), ConditionId);
-										return FReply::Handled();
-										})
+										SNew(SButton)
+										.ButtonStyle(FAppStyle::Get(), "SimpleButton")
+										.ContentPadding(FMargin(2.f, 0.f))
+										.OnClicked_Lambda([StateTreeViewModel = StateTreeViewModel, WeakState = WeakState, ConditionId]()
+										{
+											StateTreeViewModel->BringNodeToFocus(WeakState.Get(), ConditionId);
+											return FReply::Handled();
+											})
+										[
+											SNew(SHorizontalBox)
+										
+											// Icon
+											+ SHorizontalBox::Slot()
+											.VAlign(VAlign_Center)
+											.HAlign(HAlign_Left)
+											.AutoWidth()
+											[
+												SNew(SBox)
+												.Padding(FMargin(0.f, 0.f, 2.f, 0.f))
+												.Visibility(this, &SStateTreeViewRow::GetConditionIconVisibility, ConditionId)
+												[
+													SNew(SImage)
+													.Image(this, &SStateTreeViewRow::GetConditionIcon, ConditionId)
+													.ColorAndOpacity(this, &SStateTreeViewRow::GetConditionIconColor, ConditionId)
+												]
+											]
+											// Desc
+											+ SHorizontalBox::Slot()
+											.VAlign(VAlign_Center)
+											.HAlign(HAlign_Left)
+											.AutoWidth()
+											[
+												SNew(SRichTextBlock)
+												.Text(this, &SStateTreeViewRow::GetConditionDesc, ConditionId, EStateTreeNodeFormatting::RichText)
+												.ToolTipText(this, &SStateTreeViewRow::GetConditionDesc, ConditionId, EStateTreeNodeFormatting::Text)
+												.TextStyle(&FStateTreeEditorStyle::Get().GetWidgetStyle<FTextBlockStyle>("StateTree.Task.Title"))
+												.OverflowPolicy(ETextOverflowPolicy::Ellipsis)
+												.Clipping(EWidgetClipping::OnDemand)
+												+SRichTextBlock::Decorator(FTextStyleDecorator::Create(TEXT(""), FStateTreeEditorStyle::Get().GetWidgetStyle<FTextBlockStyle>("StateTree.Task.Title")))
+												+SRichTextBlock::Decorator(FTextStyleDecorator::Create(TEXT("b"), FStateTreeEditorStyle::Get().GetWidgetStyle<FTextBlockStyle>("StateTree.Task.Title.Bold")))
+												+SRichTextBlock::Decorator(FTextStyleDecorator::Create(TEXT("s"), FStateTreeEditorStyle::Get().GetWidgetStyle<FTextBlockStyle>("StateTree.Task.Title.Subdued")))
+											]
+										]
+									]
+
+									+ SOverlay::Slot()
 									[
+										// Condition override box
 										SNew(SHorizontalBox)
-										
-										// Icon
 										+ SHorizontalBox::Slot()
-										.VAlign(VAlign_Center)
+										.VAlign(VAlign_Top)
 										.HAlign(HAlign_Left)
 										.AutoWidth()
 										[
 											SNew(SBox)
-											.Padding(FMargin(0.f, 0.f, 2.f, 0.f))
-											.Visibility(this, &SStateTreeViewRow::GetConditionIconVisibility, ConditionId)
+											.Padding(FMargin(-2.0f, -2.0f, 0.0f, 0.0f))
 											[
 												SNew(SImage)
-												.Image(this, &SStateTreeViewRow::GetConditionIcon, ConditionId)
-												.ColorAndOpacity(this, &SStateTreeViewRow::GetConditionIconColor, ConditionId)
+												.DesiredSizeOverride(FVector2D(16.f, 16.f))
+												.Image_Lambda(GetForcedConditionImageFunc)
+												.Visibility_Lambda(IsForcedConditionVisibleFunc)
+												.ToolTipText_Lambda(GetForcedConditionTooltipFunc)
 											]
 										]
-										// Desc
-										+ SHorizontalBox::Slot()
-										.VAlign(VAlign_Center)
-										.HAlign(HAlign_Left)
-										.AutoWidth()
-										[
-											SNew(SRichTextBlock)
-											.Text(this, &SStateTreeViewRow::GetConditionDesc, ConditionId, EStateTreeNodeFormatting::RichText)
-											.ToolTipText(this, &SStateTreeViewRow::GetConditionDesc, ConditionId, EStateTreeNodeFormatting::Text)
-											.TextStyle(&FStateTreeEditorStyle::Get().GetWidgetStyle<FTextBlockStyle>("StateTree.Task.Title"))
-											.OverflowPolicy(ETextOverflowPolicy::Ellipsis)
-											.Clipping(EWidgetClipping::OnDemand)
-											+SRichTextBlock::Decorator(FTextStyleDecorator::Create(TEXT(""), FStateTreeEditorStyle::Get().GetWidgetStyle<FTextBlockStyle>("StateTree.Task.Title")))
-											+SRichTextBlock::Decorator(FTextStyleDecorator::Create(TEXT("b"), FStateTreeEditorStyle::Get().GetWidgetStyle<FTextBlockStyle>("StateTree.Task.Title.Bold")))
-											+SRichTextBlock::Decorator(FTextStyleDecorator::Create(TEXT("s"), FStateTreeEditorStyle::Get().GetWidgetStyle<FTextBlockStyle>("StateTree.Task.Title.Subdued")))
-										]
 									]
 								]
-
-								+ SOverlay::Slot()
+							
+								// Close parens
+								+ SHorizontalBox::Slot()
+								.AutoWidth()
+								.VAlign(VAlign_Center)
 								[
-									// Condition override box
-									SNew(SHorizontalBox)
-									+ SHorizontalBox::Slot()
-									.VAlign(VAlign_Top)
-									.HAlign(HAlign_Left)
-									.AutoWidth()
+									SNew(SBox)
+									.Padding(FMargin(0.0f, 1.0f, 0.0f, 0.0f))
 									[
-										SNew(SBox)
-										.Padding(FMargin(-2.0f, -2.0f, 0.0f, 0.0f))
-										[
-											SNew(SImage)
-											.DesiredSizeOverride(FVector2D(16.f, 16.f))
-											.Image_Lambda(GetForcedConditionImageFunc)
-											.Visibility_Lambda(IsForcedConditionVisibleFunc)
-											.ToolTipText_Lambda(GetForcedConditionTooltipFunc)
-										]
+										SNew(STextBlock)
+										.TextStyle(FStateTreeEditorStyle::Get(), "StateTree.Task.Title")
+										.Text(this, &SStateTreeViewRow::GetCloseParens, ConditionIndex)
 									]
 								]
 							]
-							
-							// Close parens
-							+ SHorizontalBox::Slot()
-							.AutoWidth()
-							.VAlign(VAlign_Center)
-							[
-								SNew(SBox)
-								.Padding(FMargin(0.0f, 1.0f, 0.0f, 0.0f))
-								[
-									SNew(STextBlock)
-									.TextStyle(FStateTreeEditorStyle::Get(), "StateTree.Task.Title")
-									.Text(this, &SStateTreeViewRow::GetCloseParens, ConditionIndex)
-								]
-							]
 						]
-					]
-				];
+					];
+			}
 		}
+		VerticalBox->AddSlot()
+		[
+			ConditionsBox
+		];
 	}
 
-	return ConditionsBox;
+	return VerticalBox;
 }
 
 void SStateTreeViewRow::MakeFlagsWidget()
