StateTree: Add scheduled tick request. A task can request a slower tick rate (for better performance). See StateTreeDelayTask for an example. Use bConsideredForScheduling to inform the scheduled tick system to not take the task into account when calculating the next tick rate.

diff --git a/Engine/Plugins/Runtime/StateTree/Source/StateTreeEditorModule/Private/Customizations/StateTreeEditorNodeDetails.cpp b/Engine/Plugins/Runtime/StateTree/Source/StateTreeEditorModule/Private/Customizations/StateTreeEditorNodeDetails.cpp
--- a/Engine/Plugins/Runtime/StateTree/Source/StateTreeEditorModule/Private/Customizations/StateTreeEditorNodeDetails.cpp
+++ b/Engine/Plugins/Runtime/StateTree/Source/StateTreeEditorModule/Private/Customizations/StateTreeEditorNodeDetails.cpp
@@ -724,9 +724,12 @@ void FStateTreeEditorNodeDetails::MakeFlagsWidget()
 			{
 				if (const FStateTreeTaskBase* CompiledTask = StateTreePtr->GetNode(StateTreePtr->GetNodeIndexFromId(EditorNode->ID).AsInt32()).GetPtr<const FStateTreeTaskBase>())
 				{
-					bShowCallTick = bShowCallTick || CompiledTask->bShouldCallTick;
-					bShouldCallTickOnlyOnEvents = bShouldCallTickOnlyOnEvents || CompiledTask->bShouldCallTickOnlyOnEvents;
-					bHasTransitionTick = bHasTransitionTick || CompiledTask->bShouldAffectTransitions;
+					if (CompiledTask->bConsideredForScheduling)
+					{
+						bShowCallTick = bShowCallTick || CompiledTask->bShouldCallTick;
+						bShouldCallTickOnlyOnEvents = bShouldCallTickOnlyOnEvents || CompiledTask->bShouldCallTickOnlyOnEvents;
+						bHasTransitionTick = bHasTransitionTick || CompiledTask->bShouldAffectTransitions;
+					}
 					bUseEditorData = false;
 				}
 			}
@@ -735,9 +738,12 @@ void FStateTreeEditorNodeDetails::MakeFlagsWidget()
 			{
 				if (const FStateTreeTaskBase* TreeTaskNodePtr = EditorNode->Node.GetPtr<FStateTreeTaskBase>())
 				{
-					bShowCallTick = bShowCallTick || TreeTaskNodePtr->bShouldCallTick;
-					bShouldCallTickOnlyOnEvents = bShouldCallTickOnlyOnEvents || TreeTaskNodePtr->bShouldCallTickOnlyOnEvents;
-					bHasTransitionTick = bHasTransitionTick || TreeTaskNodePtr->bShouldAffectTransitions;
+					if (TreeTaskNodePtr->bConsideredForScheduling)
+					{
+						bShowCallTick = bShowCallTick || TreeTaskNodePtr->bShouldCallTick;
+						bShouldCallTickOnlyOnEvents = bShouldCallTickOnlyOnEvents || TreeTaskNodePtr->bShouldCallTickOnlyOnEvents;
+						bHasTransitionTick = bHasTransitionTick || TreeTaskNodePtr->bShouldAffectTransitions;
+					}
 				}
 			}
 		}
diff --git a/Engine/Plugins/Runtime/StateTree/Source/StateTreeEditorModule/Private/SStateTreeViewRow.cpp b/Engine/Plugins/Runtime/StateTree/Source/StateTreeEditorModule/Private/SStateTreeViewRow.cpp
--- a/Engine/Plugins/Runtime/StateTree/Source/StateTreeEditorModule/Private/SStateTreeViewRow.cpp
+++ b/Engine/Plugins/Runtime/StateTree/Source/StateTreeEditorModule/Private/SStateTreeViewRow.cpp
@@ -958,7 +958,9 @@ void SStateTreeViewRow::MakeFlagsWidget()
 	{
 		if (const FCompactStateTreeState* RuntimeState = StateTree->GetStateFromHandle(StateTree->GetStateHandleFromId(State->ID)))
 		{
-			if (RuntimeState->bHasTickTasks || RuntimeState->bHasTickTasksOnlyOnEvents || RuntimeState->bHasTransitionTasks || RuntimeState->bHasCustomTickRate)
+			const bool bHasEvents = true;
+			const bool bHasBroadcastedDelegates = true;
+			if (RuntimeState->DoesRequestTickTasks(bHasEvents) || RuntimeState->bHasCustomTickRate || RuntimeState->ShouldTickTransitions(bHasEvents, bHasBroadcastedDelegates))
 			{
 				TSharedRef<SHorizontalBox> Box = SNew(SHorizontalBox);
 				if (RuntimeState->bHasCustomTickRate)
diff --git a/Engine/Plugins/Runtime/StateTree/Source/StateTreeModule/Private/StateTree.cpp b/Engine/Plugins/Runtime/StateTree/Source/StateTreeModule/Private/StateTree.cpp
--- a/Engine/Plugins/Runtime/StateTree/Source/StateTreeModule/Private/StateTree.cpp
+++ b/Engine/Plugins/Runtime/StateTree/Source/StateTreeModule/Private/StateTree.cpp
@@ -191,6 +191,8 @@ void UStateTree::ResetCompiled()
 	bHasGlobalTransitionTasks = false;
 	bHasGlobalTickTasks = false;
 	bHasGlobalTickTasksOnlyOnEvents = false;
+	bCachedRequestGlobalTick = false;
+	bCachedRequestGlobalTickOnlyOnEvents = false;
 	bScheduledTickAllowed = false;
 	
 	ResetLinked();
@@ -701,6 +703,8 @@ void UStateTree::UpdateRuntimeFlags()
 		State.bHasTickTasks = false;
 		State.bHasTickTasksOnlyOnEvents = false;
 		State.bHasTransitionTasks = false;
+		State.bCachedRequestTick = false;
+		State.bCachedRequestTickOnlyOnEvents = false;
 		for (int32 TaskIndex = State.TasksBegin; TaskIndex < (State.TasksBegin + State.TasksNum); ++TaskIndex)
 		{
 			const FStateTreeTaskBase& Task = Nodes[TaskIndex].Get<const FStateTreeTaskBase>();
@@ -709,6 +713,11 @@ void UStateTree::UpdateRuntimeFlags()
 				State.bHasTickTasks |= Task.bShouldCallTick;
 				State.bHasTickTasksOnlyOnEvents |= Task.bShouldCallTickOnlyOnEvents;
 				State.bHasTransitionTasks |= Task.bShouldAffectTransitions;
+				if (Task.bConsideredForScheduling)
+				{
+					State.bCachedRequestTick |= Task.bShouldCallTick || Task.bShouldAffectTransitions;
+					State.bCachedRequestTickOnlyOnEvents |= Task.bShouldCallTickOnlyOnEvents;
+				}
 			}
 		}
 	}
@@ -718,6 +727,8 @@ void UStateTree::UpdateRuntimeFlags()
 		bHasGlobalTickTasks = false;
 		bHasGlobalTickTasksOnlyOnEvents = false;
 		bHasGlobalTransitionTasks = false;
+		bCachedRequestGlobalTick = false;
+		bCachedRequestGlobalTickOnlyOnEvents = false;
 		for (int32 TaskIndex = GlobalTasksBegin; TaskIndex < (GlobalTasksBegin + GlobalTasksNum); ++TaskIndex)
 		{
 			const FStateTreeTaskBase& Task = Nodes[TaskIndex].Get<const FStateTreeTaskBase>();
@@ -726,6 +737,11 @@ void UStateTree::UpdateRuntimeFlags()
 				bHasGlobalTickTasks |= Task.bShouldCallTick;
 				bHasGlobalTickTasksOnlyOnEvents |= Task.bShouldCallTickOnlyOnEvents;
 				bHasGlobalTransitionTasks |= Task.bShouldAffectTransitions;
+				if (Task.bConsideredForScheduling)
+				{
+					bCachedRequestGlobalTick |= Task.bShouldCallTick || Task.bShouldAffectTransitions;
+					bCachedRequestGlobalTickOnlyOnEvents |= Task.bShouldCallTickOnlyOnEvents;
+				}
 			}
 		}
 	}
diff --git a/Engine/Plugins/Runtime/StateTree/Source/StateTreeModule/Private/StateTreeExecutionContext.cpp b/Engine/Plugins/Runtime/StateTree/Source/StateTreeModule/Private/StateTreeExecutionContext.cpp
--- a/Engine/Plugins/Runtime/StateTree/Source/StateTreeModule/Private/StateTreeExecutionContext.cpp
+++ b/Engine/Plugins/Runtime/StateTree/Source/StateTreeModule/Private/StateTreeExecutionContext.cpp
@@ -199,7 +199,7 @@ FStateTreeScheduledTick FStateTreeMinimalExecutionContext::GetNextScheduledTick(
 
 	// We wish to return in order: EveryFrames, then NextFrame, then CustomTickRate, then Sleep.
 	// Do we have a state that requires a tick or is waiting for an event.
-	TOptional<float> StateCustomTickRate;
+	TOptional<float> CustomTickRate;
 	{
 		bool bHasTaskWithEveryFramesTick = false;
 		for (int32 FrameIndex = 0; FrameIndex < Exec.ActiveFrames.Num(); ++FrameIndex)
@@ -208,7 +208,7 @@ FStateTreeScheduledTick FStateTreeMinimalExecutionContext::GetNextScheduledTick(
 			const UStateTree* CurrentStateTree = CurrentFrame.StateTree;
 
 			// Test global tasks.
-			if (CurrentStateTree->ShouldTickGlobalTasks(bHasEvents))
+			if (CurrentStateTree->DoesRequestTickGlobalTasks(bHasEvents))
 			{
 				bHasTaskWithEveryFramesTick = true;
 			}
@@ -222,12 +222,11 @@ FStateTreeScheduledTick FStateTreeMinimalExecutionContext::GetNextScheduledTick(
 				{
 					if (State.bHasCustomTickRate)
 					{
-						StateCustomTickRate = StateCustomTickRate.IsSet() ? FMath::Min(StateCustomTickRate.GetValue(), State.CustomTickRate) : State.CustomTickRate;
+						CustomTickRate = CustomTickRate.IsSet() ? FMath::Min(CustomTickRate.GetValue(), State.CustomTickRate) : State.CustomTickRate;
 					}
-					else if (!StateCustomTickRate.IsSet())
+					else if (!CustomTickRate.IsSet())
 					{
-						if (State.ShouldTickTasks(bHasEvents)
-							|| State.bHasTransitionTasks
+						if (State.DoesRequestTickTasks(bHasEvents)
 							|| State.ShouldTickTransitions(bHasEvents, bHasBroadcastedDelegates)
 							)
 						{
@@ -239,20 +238,33 @@ FStateTreeScheduledTick FStateTreeMinimalExecutionContext::GetNextScheduledTick(
 			}
 		}
 
-		if (!StateCustomTickRate.IsSet() && bHasTaskWithEveryFramesTick)
+		if (!CustomTickRate.IsSet() && bHasTaskWithEveryFramesTick)
 		{
 			return FStateTreeScheduledTick::MakeEveryFrames();
 		}
 
 		// If one state has a custom tick rate, then it overrides the tick rate for all states.
 		//Only return the CustomTickRate if it's > than NextFrame, the custom tick rate will be processed at the end.
-		if (StateCustomTickRate.IsSet() && StateCustomTickRate.GetValue() <= 0.0f)
+		if (CustomTickRate.IsSet() && CustomTickRate.GetValue() <= 0.0f)
 		{
 			// A state might override the custom tick rate with > 0, then another state overrides it again with 0 to tick back every frame.
 			return FStateTreeScheduledTick::MakeEveryFrames();
 		}
 	}
 
+	// Requests
+	if (Exec.HasScheduledTickRequests())
+	{
+		// The ScheduledTickRequests loop value is cached. Returns every frame or next frame. CustomTime needs to wait after the other tests.
+		const FStateTreeScheduledTick ScheduledTickRequest = Exec.GetScheduledTickRequest();
+		if (ScheduledTickRequest.ShouldTickEveryFrames() || ScheduledTickRequest.ShouldTickOnceNextFrame())
+		{
+			return ScheduledTickRequest;
+		}
+		const float CachedTickRate = ScheduledTickRequest.GetTickRate();
+		CustomTickRate = CustomTickRate.IsSet() ? FMath::Min(CustomTickRate.GetValue(), CachedTickRate) : CachedTickRate;
+	}
+
 	// Transitions
 	if (Storage.GetTransitionRequests().Num() > 0)
 	{
@@ -276,19 +288,63 @@ FStateTreeScheduledTick FStateTreeMinimalExecutionContext::GetNextScheduledTick(
 	{
 		for (const FStateTreeTransitionDelayedState& Transition : Exec.DelayedTransitions)
 		{
-			StateCustomTickRate = StateCustomTickRate.IsSet() ? FMath::Min(StateCustomTickRate.GetValue(), Transition.TimeLeft) : Transition.TimeLeft;
+			CustomTickRate = CustomTickRate.IsSet() ? FMath::Min(CustomTickRate.GetValue(), Transition.TimeLeft) : Transition.TimeLeft;
 		}
 	}
 
 	// Custom tick rate for tasks and transitions.
-	if (StateCustomTickRate.IsSet())
+	if (CustomTickRate.IsSet())
 	{
-		return FStateTreeScheduledTick::MakeCustomTickRate(StateCustomTickRate.GetValue());
+		return FStateTreeScheduledTick::MakeCustomTickRate(CustomTickRate.GetValue());
 	}
 
 	return FStateTreeScheduledTick::MakeSleep();
 }
 
+UE::StateTree::FScheduledTickHandle FStateTreeMinimalExecutionContext::AddScheduledTickRequest(FStateTreeScheduledTick ScheduledTick)
+{
+	if (!IsValid())
+	{
+		STATETREE_LOG(Warning, TEXT("%hs: StateTree context is not initialized properly ('%s' using StateTree '%s')"),
+			__FUNCTION__, *GetNameSafe(&Owner), *GetFullNameSafe(&RootStateTree));
+		return UE::StateTree::FScheduledTickHandle();
+	}
+
+	UE::StateTree::FScheduledTickHandle Result = Storage.GetMutableExecutionState().AddScheduledTickRequest(ScheduledTick);
+	ScheduleNextTick();
+	return Result;
+}
+
+void FStateTreeMinimalExecutionContext::UpdateScheduledTickRequest(UE::StateTree::FScheduledTickHandle Handle, FStateTreeScheduledTick ScheduledTick)
+{
+	if (!IsValid())
+	{
+		STATETREE_LOG(Warning, TEXT("%hs: StateTree context is not initialized properly ('%s' using StateTree '%s')"),
+			__FUNCTION__, *GetNameSafe(&Owner), *GetFullNameSafe(&RootStateTree));
+		return;
+	}
+
+	if (Storage.GetMutableExecutionState().UpdateScheduledTickRequest(Handle, ScheduledTick))
+	{
+		ScheduleNextTick();
+	}
+}
+
+void FStateTreeMinimalExecutionContext::RemoveScheduledTickRequest(UE::StateTree::FScheduledTickHandle Handle)
+{
+	if (!IsValid())
+	{
+		STATETREE_LOG(Warning, TEXT("%hs: StateTree context is not initialized properly ('%s' using StateTree '%s')"),
+			__FUNCTION__, *GetNameSafe(&Owner), *GetFullNameSafe(&RootStateTree));
+		return;
+	}
+
+	if (Storage.GetMutableExecutionState().RemoveScheduledTickRequest(Handle))
+	{
+		ScheduleNextTick();
+	}
+}
+
 void FStateTreeMinimalExecutionContext::SendEvent(const FGameplayTag Tag, const FConstStructView Payload, const FName Origin)
 {
 	CSV_SCOPED_TIMING_STAT_EXCLUSIVE(StateTree_SendEvent);
@@ -5826,6 +5882,41 @@ bool FStateTreeWeakExecutionContext::FinishTask(const FStateTreeWeakTaskRef& Tas
 	return bCompleted;
 }
 
+
+bool FStateTreeWeakExecutionContext::UpdateScheduledTickRequest(UE::StateTree::FScheduledTickHandle Handle, FStateTreeScheduledTick ScheduledTick) const
+{
+	bool bCompleted = false;
+
+	TStrongObjectPtr<UObject> OwnerPinned = Owner.Pin();
+	TStrongObjectPtr<const UStateTree> StateTreePinned = StateTree.Pin();
+	TSharedPtr<FStateTreeInstanceStorage> StoragePinned = Storage.Pin();
+	if (OwnerPinned && StateTreePinned && StoragePinned)
+	{
+		// Manually acquire and release the write access.
+		//The context may already exist and we want to confirm that we are on the same thread.
+		StoragePinned->AcquireWriteAccess();
+
+		bCompleted = true;
+
+		FStateTreeExecutionState& Exec = StoragePinned->GetMutableExecutionState();
+		if (Exec.CurrentExecutionContext)
+		{
+			Exec.CurrentExecutionContext->UpdateScheduledTickRequest(Handle, ScheduledTick);
+		}
+		else
+		{
+			if (Exec.UpdateScheduledTickRequest(Handle, ScheduledTick))
+			{
+				ScheduleNextTick(OwnerPinned.Get(), StateTreePinned.Get(), *StoragePinned);
+			}
+		}
+
+		StoragePinned->ReleaseWriteAccess();
+	}
+
+	return bCompleted;
+}
+
 void FStateTreeWeakExecutionContext::ScheduleNextTick(TNotNull<UObject*> Owner, TNotNull<const UStateTree*> RootStateTree, FStateTreeInstanceStorage& Storage)
 {
 	TInstancedStruct<FStateTreeExecutionExtension>& ExecutionExtension = Storage.GetMutableExecutionState().ExecutionExtension;
diff --git a/Engine/Plugins/Runtime/StateTree/Source/StateTreeModule/Private/StateTreeExecutionTypes.cpp b/Engine/Plugins/Runtime/StateTree/Source/StateTreeModule/Private/StateTreeExecutionTypes.cpp
--- a/Engine/Plugins/Runtime/StateTree/Source/StateTreeModule/Private/StateTreeExecutionTypes.cpp
+++ b/Engine/Plugins/Runtime/StateTree/Source/StateTreeModule/Private/StateTreeExecutionTypes.cpp
@@ -122,6 +122,79 @@ UE::StateTree::FActiveStatePath FStateTreeExecutionState::GetActiveStatePath() c
 	return UE::StateTree::FActiveStatePath(ActiveFrames[0].StateTree, MoveTemp(Elements));
 }
 
+UE::StateTree::FScheduledTickHandle FStateTreeExecutionState::AddScheduledTickRequest(FStateTreeScheduledTick ScheduledTick)
+{
+	UE::StateTree::FScheduledTickHandle Result = UE::StateTree::FScheduledTickHandle::GenerateNewHandle();
+	ScheduledTickRequests.Emplace(Result, ScheduledTick);
+	CacheScheduledTickRequest();
+	return Result;
+}
+
+bool FStateTreeExecutionState::UpdateScheduledTickRequest(UE::StateTree::FScheduledTickHandle Handle, FStateTreeScheduledTick ScheduledTick)
+{
+	FScheduledTickRequest* Found = ScheduledTickRequests.FindByPredicate([Handle](const FScheduledTickRequest& Other) { return Other.Handle == Handle; });
+	if (Found && Found->ScheduledTick != ScheduledTick)
+	{
+		Found->ScheduledTick = ScheduledTick;
+		CacheScheduledTickRequest();
+		return true;
+	}
+	return false;
+}
+
+bool FStateTreeExecutionState::RemoveScheduledTickRequest(UE::StateTree::FScheduledTickHandle Handle)
+{
+	const int32 IndexOf = ScheduledTickRequests.IndexOfByPredicate([Handle](const FScheduledTickRequest& Other) { return Other.Handle == Handle; });
+	if (IndexOf != INDEX_NONE)
+	{
+		ScheduledTickRequests.RemoveAtSwap(IndexOf);
+		CacheScheduledTickRequest();
+	}
+	return IndexOf != INDEX_NONE;
+}
+
+void FStateTreeExecutionState::CacheScheduledTickRequest()
+{
+	auto GetBestRequest = [](TConstArrayView<FScheduledTickRequest> ScheduledTickRequests)
+		{
+			const int32 ScheduledTickRequestsNum = ScheduledTickRequests.Num();
+			if (ScheduledTickRequestsNum == 0)
+			{
+				return FStateTreeScheduledTick();
+			}
+			if (ScheduledTickRequestsNum == 1)
+			{
+				return ScheduledTickRequests[0].ScheduledTick;
+			}
+
+			for (const FScheduledTickRequest& Request : ScheduledTickRequests)
+			{
+				if (Request.ScheduledTick.ShouldTickEveryFrames())
+				{
+					return Request.ScheduledTick;
+				}
+			}
+
+			for (const FScheduledTickRequest& Request : ScheduledTickRequests)
+			{
+				if (Request.ScheduledTick.ShouldTickOnceNextFrame())
+				{
+					return Request.ScheduledTick;
+				}
+			}
+
+			TOptional<float> CustomTickRate;
+			for (const FScheduledTickRequest& Request : ScheduledTickRequests)
+			{
+				const float CachedTickRate = Request.ScheduledTick.GetTickRate();
+				CustomTickRate = CustomTickRate.IsSet() ? FMath::Min(CustomTickRate.GetValue(), CachedTickRate) : CachedTickRate;
+			}
+			return FStateTreeScheduledTick::MakeCustomTickRate(CustomTickRate.GetValue());
+		};
+
+	CachedScheduledTickRequest = GetBestRequest(ScheduledTickRequests);
+}
+
 //----------------------------------------------------------------------//
 // FStateTreeExecutionFrame
 //----------------------------------------------------------------------//
@@ -197,6 +270,28 @@ float FStateTreeScheduledTick::GetTickRate() const
 	return NextDeltaTime;
 }
 
+//----------------------------------------------------------------------//
+// FScheduledTickHandle
+//----------------------------------------------------------------------//
+UE::StateTree::FScheduledTickHandle UE::StateTree::FScheduledTickHandle::GenerateNewHandle()
+{
+	static std::atomic<uint32> Value = 0;
+
+	uint32 Result = 0;
+	UE_AUTORTFM_OPEN
+	{
+		Result = ++Value;
+
+		// Check that we wrap round to 0, because we reserve 0 for invalid.
+		if (Result == 0)
+		{
+			Result = ++Value;
+		}
+	};
+
+	return FScheduledTickHandle(Result);
+}
+
 
 //----------------------------------------------------------------------//
 // FStateTreeDelegateActiveListeners
diff --git a/Engine/Plugins/Runtime/StateTree/Source/StateTreeModule/Private/Tasks/StateTreeDelayTask.cpp b/Engine/Plugins/Runtime/StateTree/Source/StateTreeModule/Private/Tasks/StateTreeDelayTask.cpp
--- a/Engine/Plugins/Runtime/StateTree/Source/StateTreeModule/Private/Tasks/StateTreeDelayTask.cpp
+++ b/Engine/Plugins/Runtime/StateTree/Source/StateTreeModule/Private/Tasks/StateTreeDelayTask.cpp
@@ -7,6 +7,11 @@
 
 #define LOCTEXT_NAMESPACE "StateTree"
 
+FStateTreeDelayTask::FStateTreeDelayTask()
+{
+	bConsideredForScheduling = false;
+}
+
 EStateTreeRunStatus FStateTreeDelayTask::EnterState(FStateTreeExecutionContext& Context, const FStateTreeTransitionResult& Transition) const
 {
 	FInstanceDataType& InstanceData = Context.GetInstanceData(*this);
@@ -15,6 +20,8 @@ EStateTreeRunStatus FStateTreeDelayTask::EnterState(FStateTreeExecutionContext&
 	{
 		InstanceData.RemainingTime = FMath::FRandRange(
 			FMath::Max(0.0f, InstanceData.Duration - InstanceData.RandomDeviation), (InstanceData.Duration + InstanceData.RandomDeviation));
+
+		InstanceData.ScheduledTickHandle = Context.AddScheduledTickRequest(FStateTreeScheduledTick::MakeCustomTickRate(InstanceData.RemainingTime));
 	}
 	
 	return EStateTreeRunStatus::Running;
@@ -30,13 +37,21 @@ EStateTreeRunStatus FStateTreeDelayTask::Tick(FStateTreeExecutionContext& Contex
 
 		if (InstanceData.RemainingTime <= 0.f)
 		{
+			Context.RemoveScheduledTickRequest(InstanceData.ScheduledTickHandle);
 			return EStateTreeRunStatus::Succeeded;
 		}
+		Context.UpdateScheduledTickRequest(InstanceData.ScheduledTickHandle, FStateTreeScheduledTick::MakeCustomTickRate(InstanceData.RemainingTime));
 	}
 	
 	return EStateTreeRunStatus::Running;
 }
 
+void FStateTreeDelayTask::ExitState(FStateTreeExecutionContext& Context, const FStateTreeTransitionResult& Transition) const
+{
+	FInstanceDataType& InstanceData = Context.GetInstanceData(*this);
+	Context.RemoveScheduledTickRequest(InstanceData.ScheduledTickHandle);
+}
+
 #if WITH_EDITOR
 FText FStateTreeDelayTask::GetDescription(const FGuid& ID, FStateTreeDataView InstanceDataView, const IStateTreeBindingLookup& BindingLookup, EStateTreeNodeFormatting Formatting) const
 {
diff --git a/Engine/Plugins/Runtime/StateTree/Source/StateTreeModule/Private/Tasks/StateTreeDelayTask.h b/Engine/Plugins/Runtime/StateTree/Source/StateTreeModule/Private/Tasks/StateTreeDelayTask.h
--- a/Engine/Plugins/Runtime/StateTree/Source/StateTreeModule/Private/Tasks/StateTreeDelayTask.h
+++ b/Engine/Plugins/Runtime/StateTree/Source/StateTreeModule/Private/Tasks/StateTreeDelayTask.h
@@ -27,6 +27,9 @@ struct STATETREEMODULE_API FStateTreeDelayTaskInstanceData
 
 	/** Internal countdown in seconds. */
 	float RemainingTime = 0.f;
+
+	/** The handle of the scheduled tick request. */
+	UE::StateTree::FScheduledTickHandle ScheduledTickHandle;
 };
 
 /**
@@ -39,12 +42,14 @@ struct STATETREEMODULE_API FStateTreeDelayTask : public FStateTreeTaskCommonBase
 
 	using FInstanceDataType = FStateTreeDelayTaskInstanceData;
 	
-	FStateTreeDelayTask() = default;
+	FStateTreeDelayTask();
 
 	virtual const UStruct* GetInstanceDataType() const override { return FInstanceDataType::StaticStruct(); }
 
 	virtual EStateTreeRunStatus EnterState(FStateTreeExecutionContext& Context, const FStateTreeTransitionResult& Transition) const override;
 	virtual EStateTreeRunStatus Tick(FStateTreeExecutionContext& Context, const float DeltaTime) const override;
+	virtual void ExitState(FStateTreeExecutionContext& Context, const FStateTreeTransitionResult& Transition) const override;
+
 #if WITH_EDITOR
 	virtual FText GetDescription(const FGuid& ID, FStateTreeDataView InstanceDataView, const IStateTreeBindingLookup& BindingLookup, EStateTreeNodeFormatting Formatting = EStateTreeNodeFormatting::Text) const override;
 	virtual FName GetIconName() const override
diff --git a/Engine/Plugins/Runtime/StateTree/Source/StateTreeModule/Private/Tasks/StateTreeRunParallelStateTreeTask.cpp b/Engine/Plugins/Runtime/StateTree/Source/StateTreeModule/Private/Tasks/StateTreeRunParallelStateTreeTask.cpp
--- a/Engine/Plugins/Runtime/StateTree/Source/StateTreeModule/Private/Tasks/StateTreeRunParallelStateTreeTask.cpp
+++ b/Engine/Plugins/Runtime/StateTree/Source/StateTreeModule/Private/Tasks/StateTreeRunParallelStateTreeTask.cpp
@@ -1,16 +1,25 @@
 // Copyright Epic Games, Inc. All Rights Reserved.
 
 #include "Tasks/StateTreeRunParallelStateTreeTask.h"
-
 #include "StateTreeExecutionContext.h"
+#include "StateTreeExecutionTypes.h"
+#include "StateTreeInstanceData.h"
 
 #define LOCTEXT_NAMESPACE "StateTree"
 
+void FStateTreeRunParallelStateTreeExecutionExtension::ScheduleNextTick(const FContextParameters& Context)
+{
+	const FStateTreeMinimalExecutionContext ExecutionContext(Context.Owner, Context.StateTree, Context.InstanceData);
+	const FStateTreeScheduledTick ScheduledTick = ExecutionContext.GetNextScheduledTick();
+	WeakExecutionContext.UpdateScheduledTickRequest(ScheduledTickHandle, ScheduledTick);
+}
+
 FStateTreeRunParallelStateTreeTask::FStateTreeRunParallelStateTreeTask()
 {
 	bShouldCopyBoundPropertiesOnTick = false;
 	bShouldCopyBoundPropertiesOnExitState = false;
 	bShouldAffectTransitions = true;
+	bConsideredForScheduling = false;
 }
 
 EStateTreeRunStatus FStateTreeRunParallelStateTreeTask::EnterState(FStateTreeExecutionContext& Context, const FStateTreeTransitionResult& Transitions) const
@@ -22,6 +31,11 @@ EStateTreeRunStatus FStateTreeRunParallelStateTreeTask::EnterState(FStateTreeExe
 		return EStateTreeRunStatus::Failed;
 	}
 
+	if (InstanceData.ScheduledTickHandle.IsValid())
+	{
+		Context.RemoveScheduledTickRequest(InstanceData.ScheduledTickHandle);
+	}
+
 	// Share event queue with parent tree.
 	if (FStateTreeInstanceData* OuterInstanceData = Context.GetMutableInstanceData())
 	{
@@ -35,7 +49,18 @@ EStateTreeRunStatus FStateTreeRunParallelStateTreeTask::EnterState(FStateTreeExe
 		return EStateTreeRunStatus::Failed;
 	}
 
-	return ParallelTreeContext.Start(&StateTreeToRun.GetParameters());
+	FStateTreeRunParallelStateTreeExecutionExtension Extension;
+	Extension.WeakExecutionContext = Context.MakeWeakExecutionContext();
+	const EStateTreeRunStatus RunStatus = ParallelTreeContext.Start(FStateTreeExecutionContext::FStartParameters
+		{
+			.GlobalParameters = &StateTreeToRun.GetParameters(),
+			.ExecutionExtension = TInstancedStruct<FStateTreeRunParallelStateTreeExecutionExtension>::Make(MoveTemp(Extension))
+		});
+
+	InstanceData.ScheduledTickHandle = Context.AddScheduledTickRequest(ParallelTreeContext.GetNextScheduledTick());
+	InstanceData.TreeInstanceData.GetMutableExecutionState()->ExecutionExtension.GetMutable<FStateTreeRunParallelStateTreeExecutionExtension>().ScheduledTickHandle = InstanceData.ScheduledTickHandle;
+
+	return RunStatus;
 }
 
 EStateTreeRunStatus FStateTreeRunParallelStateTreeTask::Tick(FStateTreeExecutionContext& Context, const float DeltaTime) const
@@ -52,7 +77,9 @@ EStateTreeRunStatus FStateTreeRunParallelStateTreeTask::Tick(FStateTreeExecution
 		return EStateTreeRunStatus::Failed;
 	}
 
-	return ParallelTreeContext.TickUpdateTasks(DeltaTime);
+	const EStateTreeRunStatus RunStatus = ParallelTreeContext.TickUpdateTasks(DeltaTime);
+	Context.UpdateScheduledTickRequest(InstanceData.ScheduledTickHandle, ParallelTreeContext.GetNextScheduledTick());
+	return RunStatus;
 }
 
 void FStateTreeRunParallelStateTreeTask::TriggerTransitions(FStateTreeExecutionContext& Context) const
@@ -70,6 +97,7 @@ void FStateTreeRunParallelStateTreeTask::TriggerTransitions(FStateTreeExecutionC
 	}
 
 	ParallelTreeContext.TickTriggerTransitions();
+	Context.UpdateScheduledTickRequest(InstanceData.ScheduledTickHandle, ParallelTreeContext.GetNextScheduledTick());
 }
 
 void FStateTreeRunParallelStateTreeTask::ExitState(FStateTreeExecutionContext& Context, const FStateTreeTransitionResult& Transition) const
@@ -87,6 +115,7 @@ void FStateTreeRunParallelStateTreeTask::ExitState(FStateTreeExecutionContext& C
 	}
 
 	ParallelTreeContext.Stop();
+	Context.RemoveScheduledTickRequest(InstanceData.ScheduledTickHandle);
 }
 
 const FStateTreeReference& FStateTreeRunParallelStateTreeTask::GetStateTreeToRun(FStateTreeExecutionContext& Context, FInstanceDataType& InstanceData) const
diff --git a/Engine/Plugins/Runtime/StateTree/Source/StateTreeModule/Public/StateTree.h b/Engine/Plugins/Runtime/StateTree/Source/StateTreeModule/Public/StateTree.h
--- a/Engine/Plugins/Runtime/StateTree/Source/StateTreeModule/Public/StateTree.h
+++ b/Engine/Plugins/Runtime/StateTree/Source/StateTreeModule/Public/StateTree.h
@@ -165,6 +165,12 @@ public:
 	const FStateTreePropertyBindings& GetPropertyBindings() const { return PropertyBindings; }
 
 	/** @return True if there is any global tasks need ticking. */
+	bool DoesRequestTickGlobalTasks(bool bHasEvents) const
+	{
+		return bCachedRequestGlobalTick || (bHasEvents && bCachedRequestGlobalTickOnlyOnEvents);
+	}
+	
+	/** @return True if there is any global tasks that ticks. */
 	bool ShouldTickGlobalTasks(bool bHasEvents) const
 	{
 		return bHasGlobalTickTasks || (bHasEvents && bHasGlobalTickTasksOnlyOnEvents);
@@ -359,6 +365,15 @@ private:
 	 * Not ticking implies no property copy.
 	 */
 	uint8 bHasGlobalTickTasksOnlyOnEvents : 1 = false;
+	
+	/** True if any global tasks request a tick every frame. */
+	uint8 bCachedRequestGlobalTick : 1 = false;
+
+	/**
+	 * True if any global tasks request a tick every frame but only if there are events.
+	 * No effect if bCachedRequestGlobalTick is true.
+	 */
+	uint8 bCachedRequestGlobalTickOnlyOnEvents : 1 = false;
 
 	/** True when the scheduled tick is allowed by the schema. */
 	uint8 bScheduledTickAllowed : 1 = false;
diff --git a/Engine/Plugins/Runtime/StateTree/Source/StateTreeModule/Public/StateTreeExecutionContext.h b/Engine/Plugins/Runtime/StateTree/Source/StateTreeModule/Public/StateTreeExecutionContext.h
--- a/Engine/Plugins/Runtime/StateTree/Source/StateTreeModule/Public/StateTreeExecutionContext.h
+++ b/Engine/Plugins/Runtime/StateTree/Source/StateTreeModule/Public/StateTreeExecutionContext.h
@@ -86,6 +86,19 @@ public:
 	/** @return the delta time for the next execution context tick. */
 	FStateTreeScheduledTick GetNextScheduledTick() const;
 
+	/** 
+	 * Adds a scheduled tick request.
+	 * The result of GetNextScheduledTick is affected by the request.
+	 * This allows a specific task to control when the tree ticks.
+	 */
+	UE::StateTree::FScheduledTickHandle AddScheduledTickRequest(FStateTreeScheduledTick ScheduledTick);
+
+	/** Updates the scheduled tick of a previous request. */
+	void UpdateScheduledTickRequest(UE::StateTree::FScheduledTickHandle Handle, FStateTreeScheduledTick ScheduledTick);
+
+	/** Removes a scheduled tick request. */
+	void RemoveScheduledTickRequest(UE::StateTree::FScheduledTickHandle Handle);
+
 	/** Sends event for the StateTree. */
 	void SendEvent(const FGameplayTag Tag, const FConstStructView Payload = FConstStructView(), const FName Origin = FName());
 
@@ -1213,6 +1226,12 @@ public:
 	 * If called outside of the tick processing, then the request is buffered and handled on the next tick.
 	 */
 	bool FinishTask(const FStateTreeWeakTaskRef& Task, EStateTreeFinishTaskType FinishType) const;
+	
+	/**
+	 * Updates the scheduled tick of a previous request.
+	 * @return false if the context is not valid.
+	 */
+	bool UpdateScheduledTickRequest(UE::StateTree::FScheduledTickHandle Handle, FStateTreeScheduledTick ScheduledTick) const;
 
 private:
 	static void ScheduleNextTick(TNotNull<UObject*> Owner, TNotNull<const UStateTree*> RootStateTree, FStateTreeInstanceStorage& Storage);
diff --git a/Engine/Plugins/Runtime/StateTree/Source/StateTreeModule/Public/StateTreeExecutionTypes.h b/Engine/Plugins/Runtime/StateTree/Source/StateTreeModule/Public/StateTreeExecutionTypes.h
--- a/Engine/Plugins/Runtime/StateTree/Source/StateTreeModule/Public/StateTreeExecutionTypes.h
+++ b/Engine/Plugins/Runtime/StateTree/Source/StateTreeModule/Public/StateTreeExecutionTypes.h
@@ -633,6 +633,82 @@
 };
 
 
+/*
+ * Information on how a state tree should tick next.
+ */
+USTRUCT()
+struct STATETREEMODULE_API FStateTreeScheduledTick
+{
+	GENERATED_BODY()
+
+public:
+	FStateTreeScheduledTick() = default;
+
+	bool operator==(const FStateTreeScheduledTick&) const = default;
+	bool operator!=(const FStateTreeScheduledTick&) const = default;
+
+	/** Make a scheduled tick that returns Sleep. */
+	static FStateTreeScheduledTick MakeSleep();
+	/** Make a scheduled tick that returns EveryFrame. */
+	static FStateTreeScheduledTick MakeEveryFrames();
+	/** Make a scheduled tick that returns NextFrame. */
+	static FStateTreeScheduledTick MakeNextFrame();
+	/** Make a scheduled tick that returns a tick rate. The value needs to be greater than zero. */
+	static FStateTreeScheduledTick MakeCustomTickRate(float DeltaTime);
+
+public:
+	/** @return true if it doesn't need to tick until an event/delegate/transition/... occurs. */
+	bool ShouldSleep() const;
+	/** @return true if it the needs to tick every frame. */
+	bool ShouldTickEveryFrames() const;
+	/** @return true if it usually doesn't need to tick but it needs to tick once next frame. */
+	bool ShouldTickOnceNextFrame() const;
+	/** @return true if it has a custom tick rate. */
+	bool HasCustomTickRate() const;
+	/** @return the delay in seconds between ticks. */
+	float GetTickRate() const;
+
+private:
+	FStateTreeScheduledTick(float DeltaTime)
+		: NextDeltaTime(DeltaTime)
+	{}
+
+	UPROPERTY()
+	float NextDeltaTime = 0.0f;
+};
+
+
+namespace UE::StateTree
+{
+/*
+ * ID of a scheduled tick request.
+ */
+struct FScheduledTickHandle
+{
+public:
+	FScheduledTickHandle() = default;
+	FScheduledTickHandle(const FScheduledTickHandle&) = default;
+
+	STATETREEMODULE_API static FScheduledTickHandle GenerateNewHandle();
+
+	bool IsValid() const
+	{
+		return Value != 0;
+	}
+
+	bool operator==(const FScheduledTickHandle& Other) const = default;
+	bool operator!=(const FScheduledTickHandle& Other) const = default;
+
+private:
+	explicit FScheduledTickHandle(uint32 InValue)
+		: Value(InValue)
+	{}
+
+	uint32 Value = 0;
+};
+} // namespace UE::StateTree
+
+
 /** Describes added delegate listeners. */
 struct FStateTreeDelegateActiveListeners
 {
@@ -751,6 +827,7 @@
 	FStateTreeExecutionState& operator=(FStateTreeExecutionState&&) = default;
 PRAGMA_ENABLE_DEPRECATION_WARNINGS
 	
+public:
 	void Reset()
 	{
 		*this = FStateTreeExecutionState();
@@ -762,6 +839,7 @@
 	UE_DEPRECATED(5.5, "Use FindAndRemoveExpiredDelayedTransitions() instead or search through DelayedTransitions directly.")
 	FStateTreeTransitionDelayedState* FindDelayedTransition(const UStateTree* OwnerStateTree, const FStateTreeIndex16 TransitionIndex) { return nullptr; }
 
+	UE_DEPRECATED(5.6, "FindAndRemoveExpiredDelayedTransitions is not used anymore.")
 	/** Finds all delayed transition states for a specific transition and removes them. Returns their copies. */
 	TArray<FStateTreeTransitionDelayedState, TInlineAllocator<8>> FindAndRemoveExpiredDelayedTransitions(const UStateTree* OwnerStateTree, const FStateTreeIndex16 TransitionIndex)
 	{
@@ -778,6 +856,31 @@
 		return Result;
 	}
 
+	/** @return whether it contains any shcedule tick requests. */
+	bool HasScheduledTickRequests() const
+	{
+		return ScheduledTickRequests.Num() > 0;
+	}
+
+	/** @return the best/smallest scheduled tick request of all the requests. */
+	FStateTreeScheduledTick GetScheduledTickRequest() const
+	{
+		return HasScheduledTickRequests() ? CachedScheduledTickRequest : FStateTreeScheduledTick();
+	}
+
+	/** Adds a scheduled tick request. */
+	UE::StateTree::FScheduledTickHandle AddScheduledTickRequest(FStateTreeScheduledTick ScheduledTick);
+
+	/** Updates the scheduled tick of a previous request. */
+	bool UpdateScheduledTickRequest(UE::StateTree::FScheduledTickHandle Handle, FStateTreeScheduledTick ScheduledTick);
+
+	/** Removes a request. */
+	bool RemoveScheduledTickRequest(UE::StateTree::FScheduledTickHandle Handle);
+
+private:
+	void CacheScheduledTickRequest();
+
+public:
 	/** Currently active frames (and states) */
 	UPROPERTY()
 	TArray<FStateTreeExecutionFrame> ActiveFrames;
@@ -796,6 +899,19 @@
 	/** Active delegate listeners. */
 	FStateTreeDelegateActiveListeners DelegateActiveListeners;
 
+private:
+	/** ScheduledTick */
+	struct FScheduledTickRequest
+	{
+		UE::StateTree::FScheduledTickHandle Handle;
+		FStateTreeScheduledTick ScheduledTick;
+	};
+	TArray<FScheduledTickRequest> ScheduledTickRequests;
+
+	/** The current computed value from ScheduledTickRequests. Only valid when ScheduledTickRequests is not empty. */
+	FStateTreeScheduledTick CachedScheduledTickRequest;
+
+public:
 #if WITH_STATETREE_TRACE
 	/** Id for the active instance used for debugging. */
 	mutable FStateTreeInstanceDebugId InstanceDebugId;
@@ -1013,44 +1129,3 @@
 	UPROPERTY()
 	FStateTreeStateHandle SourceRootState = FStateTreeStateHandle::Invalid;
 };
-
-/*
- * Information on how a state tree should tick next.
- */
-USTRUCT()
-struct STATETREEMODULE_API FStateTreeScheduledTick
-{
-	GENERATED_BODY()
-
-public:
-	FStateTreeScheduledTick() = default;
-
-	/** Make a scheduled tick that returns Sleep. */
-	static FStateTreeScheduledTick MakeSleep();
-	/** Make a scheduled tick that returns EveryFrame. */
-	static FStateTreeScheduledTick MakeEveryFrames();
-	/** Make a scheduled tick that returns NextFrame. */
-	static FStateTreeScheduledTick MakeNextFrame();
-	/** Make a scheduled tick that returns a tick rate. The value needs to be greater than zero. */
-	static FStateTreeScheduledTick MakeCustomTickRate(float DeltaTime);
-
-public:
-	/** @return true if it doesn't need to tick until an event/delegate/transition/... occurs. */
-	bool ShouldSleep() const;
-	/** @return true if it the needs to tick every frame. */
-	bool ShouldTickEveryFrames() const;
-	/** @return true if it usually doesn't need to tick but it needs to tick once next frame. */
-	bool ShouldTickOnceNextFrame() const;
-	/** @return true if it has a custom tick rate. */
-	bool HasCustomTickRate() const;
-	/** @return the delay in seconds between ticks. */
-	float GetTickRate() const;
-
-private:
-	FStateTreeScheduledTick(float DeltaTime)
-		: NextDeltaTime(DeltaTime)
-	{}
-
-	UPROPERTY()
-	float NextDeltaTime = 0.0f;
-};
diff --git a/Engine/Plugins/Runtime/StateTree/Source/StateTreeModule/Public/StateTreeTaskBase.h b/Engine/Plugins/Runtime/StateTree/Source/StateTreeModule/Public/StateTreeTaskBase.h
--- a/Engine/Plugins/Runtime/StateTree/Source/StateTreeModule/Public/StateTreeTaskBase.h
+++ b/Engine/Plugins/Runtime/StateTree/Source/StateTreeModule/Public/StateTreeTaskBase.h
@@ -24,6 +24,7 @@ struct STATETREEMODULE_API FStateTreeTaskBase : public FStateTreeNodeBase
 		, bShouldCopyBoundPropertiesOnTick(true)
 		, bShouldCopyBoundPropertiesOnExitState(true)
 		, bShouldAffectTransitions(false)
+		, bConsideredForScheduling(true)
 		, bTaskEnabled(true)
 	{
 	}
@@ -117,6 +118,13 @@ struct STATETREEMODULE_API FStateTreeTaskBase : public FStateTreeNodeBase
 	/** If set to true, TriggerTransitions() is called during transition handling. Default false. */
 	uint8 bShouldAffectTransitions : 1;
 
+	/**
+	 * If set to true, the task is considered for scheduled tick. It will use these flags: bShouldCallTick, bShouldCallTickOnlyOnEvents, and bShouldAffectTransitions.
+	 * It doesn't affect how the task ticks.
+	 * Default true.
+	 */
+	uint8 bConsideredForScheduling : 1;
+
 	/** True if the node is Enabled (i.e. not explicitly disabled in the asset). */
 	UPROPERTY()
 	uint8 bTaskEnabled : 1;
diff --git a/Engine/Plugins/Runtime/StateTree/Source/StateTreeModule/Public/StateTreeTypes.h b/Engine/Plugins/Runtime/StateTree/Source/StateTreeModule/Public/StateTreeTypes.h
--- a/Engine/Plugins/Runtime/StateTree/Source/StateTreeModule/Public/StateTreeTypes.h
+++ b/Engine/Plugins/Runtime/StateTree/Source/StateTreeModule/Public/StateTreeTypes.h
@@ -685,6 +685,12 @@ struct STATETREEMODULE_API FCompactStateTreeState
 	}
 
 	/** @return True if the state has any tasks that need ticking. */
+	bool DoesRequestTickTasks(bool bHasEvent) const
+	{
+		return bCachedRequestTick || (bHasEvent && bCachedRequestTickOnlyOnEvents);
+	}
+
+	/** @return True if the state has any tasks that ticks. */
 	bool ShouldTickTasks(bool bHasEvent) const
 	{
 		return bHasTickTasks || (bHasEvent && bHasTickTasksOnlyOnEvents);
@@ -828,9 +834,19 @@ struct STATETREEMODULE_API FCompactStateTreeState
 
 	/**
 	 * True if any task has bShouldCallTickOnlyOnEvents.
-	 * No effect if bHasGlobalTickTasks is true.
+	 * No effect if bHasTickTasks is true.
 	 */
 	uint8 bHasTickTasksOnlyOnEvents : 1 = false;
+	
+	
+	/** True if any global tasks request a tick every frame. */
+	uint8 bCachedRequestTick : 1 = false;
+
+	/**
+	 * True if any global tasks request a tick every frame but only if there are events.
+	 * No effect if bCachedRequestTick is true.
+	 */
+	uint8 bCachedRequestTickOnlyOnEvents : 1 = false;
 
 	/** True if the state contains transitions with tick trigger. */
 	UPROPERTY()
diff --git a/Engine/Plugins/Runtime/StateTree/Source/StateTreeModule/Public/Tasks/StateTreeRunParallelStateTreeTask.h b/Engine/Plugins/Runtime/StateTree/Source/StateTreeModule/Public/Tasks/StateTreeRunParallelStateTreeTask.h
--- a/Engine/Plugins/Runtime/StateTree/Source/StateTreeModule/Public/Tasks/StateTreeRunParallelStateTreeTask.h
+++ b/Engine/Plugins/Runtime/StateTree/Source/StateTreeModule/Public/Tasks/StateTreeRunParallelStateTreeTask.h
@@ -2,10 +2,10 @@
 
 #pragma once
 
+#include "StateTreeExecutionContext.h"
 #include "StateTreeInstanceData.h"
 #include "StateTreeReference.h"
 #include "StateTreeTaskBase.h"
-
 #include "StateTreeRunParallelStateTreeTask.generated.h"
 
 USTRUCT()
@@ -22,6 +22,21 @@ struct FStateTreeRunParallelStateTreeTaskInstanceData
 
 	UPROPERTY(Transient)
 	TObjectPtr<const UStateTree> RunningStateTree = nullptr;
+
+	/** The handle of the scheduled tick. */
+	UE::StateTree::FScheduledTickHandle ScheduledTickHandle;
+};
+
+USTRUCT()
+struct FStateTreeRunParallelStateTreeExecutionExtension : public FStateTreeExecutionExtension
+{
+	GENERATED_BODY()
+
+public:
+	virtual void ScheduleNextTick(const FContextParameters& Context) override;
+
+	FStateTreeWeakExecutionContext WeakExecutionContext;
+	UE::StateTree::FScheduledTickHandle ScheduledTickHandle;
 };
 
 /**
