Fix build on mac

#rnx

diff --git a/Engine/Plugins/Runtime/StateTree/Source/StateTreeModule/Private/StateTreeExecutionTypes.cpp b/Engine/Plugins/Runtime/StateTree/Source/StateTreeModule/Private/StateTreeExecutionTypes.cpp
--- a/Engine/Plugins/Runtime/StateTree/Source/StateTreeModule/Private/StateTreeExecutionTypes.cpp
+++ b/Engine/Plugins/Runtime/StateTree/Source/StateTreeModule/Private/StateTreeExecutionTypes.cpp
@@ -125,7 +125,7 @@ UE::StateTree::FActiveStatePath FStateTreeExecutionState::GetActiveStatePath() c
 UE::StateTree::FScheduledTickHandle FStateTreeExecutionState::AddScheduledTickRequest(FStateTreeScheduledTick ScheduledTick)
 {
 	UE::StateTree::FScheduledTickHandle Result = UE::StateTree::FScheduledTickHandle::GenerateNewHandle();
-	ScheduledTickRequests.Emplace(Result, ScheduledTick);
+	ScheduledTickRequests.Add(FScheduledTickRequest{.Handle = Result, .ScheduledTick = ScheduledTick});
 	CacheScheduledTickRequest();
 	return Result;
 }
@@ -155,19 +155,19 @@ bool FStateTreeExecutionState::RemoveScheduledTickRequest(UE::StateTree::FSchedu
 
 void FStateTreeExecutionState::CacheScheduledTickRequest()
 {
-	auto GetBestRequest = [](TConstArrayView<FScheduledTickRequest> ScheduledTickRequests)
+	auto GetBestRequest = [](TConstArrayView<FScheduledTickRequest> Requests)
 		{
-			const int32 ScheduledTickRequestsNum = ScheduledTickRequests.Num();
+			const int32 ScheduledTickRequestsNum = Requests.Num();
 			if (ScheduledTickRequestsNum == 0)
 			{
 				return FStateTreeScheduledTick();
 			}
 			if (ScheduledTickRequestsNum == 1)
 			{
-				return ScheduledTickRequests[0].ScheduledTick;
+				return Requests[0].ScheduledTick;
 			}
 
-			for (const FScheduledTickRequest& Request : ScheduledTickRequests)
+			for (const FScheduledTickRequest& Request : Requests)
 			{
 				if (Request.ScheduledTick.ShouldTickEveryFrames())
 				{
@@ -175,7 +175,7 @@ void FStateTreeExecutionState::CacheScheduledTickRequest()
 				}
 			}
 
-			for (const FScheduledTickRequest& Request : ScheduledTickRequests)
+			for (const FScheduledTickRequest& Request : Requests)
 			{
 				if (Request.ScheduledTick.ShouldTickOnceNextFrame())
 				{
@@ -184,7 +184,7 @@ void FStateTreeExecutionState::CacheScheduledTickRequest()
 			}
 
 			TOptional<float> CustomTickRate;
-			for (const FScheduledTickRequest& Request : ScheduledTickRequests)
+			for (const FScheduledTickRequest& Request : Requests)
 			{
 				const float CachedTickRate = Request.ScheduledTick.GetTickRate();
 				CustomTickRate = CustomTickRate.IsSet() ? FMath::Min(CustomTickRate.GetValue(), CachedTickRate) : CachedTickRate;
