StateTree: Fix crash when the child state of an already active parent state is activated. The parent state was not properly detected as "already activated" and the event data was reset.

diff --git a/Engine/Plugins/Runtime/StateTree/Source/StateTreeModule/Private/StateTreeExecutionContext.cpp b/Engine/Plugins/Runtime/StateTree/Source/StateTreeModule/Private/StateTreeExecutionContext.cpp
--- a/Engine/Plugins/Runtime/StateTree/Source/StateTreeModule/Private/StateTreeExecutionContext.cpp
+++ b/Engine/Plugins/Runtime/StateTree/Source/StateTreeModule/Private/StateTreeExecutionContext.cpp
@@ -1645,7 +1645,7 @@ void FStateTreeExecutionContext::UpdateInstanceData(TConstArrayView<FStateTreeEx
 	}
 	
 	// Common section should match.
-	// @todo: put this behind a define when enough testing has been done.
+#if WITH_STATETREE_DEBUG
 	for (int32 Index = 0; Index < NumCommonInstanceData; Index++)
 	{
 		check(Index < InstanceData.Num());
@@ -1663,6 +1663,7 @@ void FStateTreeExecutionContext::UpdateInstanceData(TConstArrayView<FStateTreeEx
 			check(ExistingWrapper->InstanceObject->GetClass() == NewWrapper->InstanceObject->GetClass());
 		}
 	}
+#endif
 
 	// Remove instance data that was not common.
 	InstanceData.ShrinkTo(NumCommonInstanceData);
@@ -1794,9 +1795,12 @@ FStateTreeDataView FStateTreeExecutionContext::GetDataViewFromInstanceStorage(FS
 		{
 			// Return FStateTreeEvent from shared event.
 			FStateTreeSharedEvent& SharedEvent = InstanceDataStorage.GetMutableStruct(CurrentFrame.ActiveInstanceIndexBase.Get() + Handle.GetIndex()).Get<FStateTreeSharedEvent>();
-			check(SharedEvent.IsValid());
-			// Events are read only, but we cannot express that in FStateTreeDataView.
-			return FStateTreeDataView(FStructView::Make(*SharedEvent.GetMutable()));
+			if (ensure(SharedEvent.IsValid()))
+			{
+				// Events are read only, but we cannot express that in FStateTreeDataView.
+				return FStateTreeDataView(FStructView::Make(*SharedEvent.GetMutable()));
+			}
+			return {};
 		}
 
 	case EStateTreeDataSourceType::ContextData:
@@ -2176,12 +2180,13 @@ void FStateTreeExecutionContext::CaptureNewStateEvents(TConstArrayView<FStateTre
 		{
 			const FStateTreeExecutionFrame& PrevFrame = PrevFrames[FrameIndex];
 
-			if (PrevFrame.RootState == NewFrame.RootState
-				&& PrevFrame.StateTree == NewFrame.StateTree)
+			if (PrevFrame.FrameId == NewFrame.FrameId)
 			{
+				checkf(PrevFrame.RootState == NewFrame.RootState && PrevFrame.StateTree == NewFrame.StateTree, TEXT("If the Id matches, then the root and the tree must also match."));
+
 				for (int32 StateIndex = 0; StateIndex < NewFrame.ActiveStates.Num(); ++StateIndex)
 				{
-					if (!PrevFrame.ActiveStates.IsValidIndex(StateIndex) || PrevFrame.ActiveStates[StateIndex] != NewFrame.ActiveStates[StateIndex])
+					if (!PrevFrame.ActiveStates.IsValidIndex(StateIndex) || PrevFrame.ActiveStates.StateIds[StateIndex] != NewFrame.ActiveStates.StateIds[StateIndex])
 					{
 						UniqueStates = TConstArrayView<FStateTreeStateHandle>(&NewFrame.ActiveStates[StateIndex], NewFrame.ActiveStates.Num() - StateIndex);
 						break;
