StateTree: Fix runtime crash with unexisting property. The property was not generated for cooked platforms.

diff --git a/Engine/Plugins/Runtime/StateTree/Source/StateTreeEditorModule/Private/K2Node_StateTreeBlueprintPropertyRef.cpp b/Engine/Plugins/Runtime/StateTree/Source/StateTreeEditorModule/Private/K2Node_StateTreeBlueprintPropertyRef.cpp
--- a/Engine/Plugins/Runtime/StateTree/Source/StateTreeEditorModule/Private/K2Node_StateTreeBlueprintPropertyRef.cpp
+++ b/Engine/Plugins/Runtime/StateTree/Source/StateTreeEditorModule/Private/K2Node_StateTreeBlueprintPropertyRef.cpp
@@ -128,14 +128,11 @@ void FKCHandler_StateTreeBlueprintPropertyRefGet::RegisterNets(FKismetFunctionCo
 	OutValueTerminal->CopyFromPin(OutPin, Context.NetNameMap->MakeValidName(OutPin));
 	Context.NetMap.Add(OutPin, OutValueTerminal);
 
-	FBPTerminal*& IsValidTerminal = TemporaryBoolTerminals.FindOrAdd(Node);
-	if (!IsValidTerminal)
-	{
-		IsValidTerminal = Context.CreateLocalTerminal();
-		IsValidTerminal->bPassedByReference = true;
-		IsValidTerminal->Source = Node;
-		IsValidTerminal->Type.PinCategory = UEdGraphSchema_K2::PC_Boolean;
-	}	
+	FBPTerminal* BoolTerm = Context.CreateLocalTerminal();
+	BoolTerm->Type.PinCategory = UEdGraphSchema_K2::PC_Boolean;
+	BoolTerm->Source = Node;
+	BoolTerm->Name = Context.NetNameMap->MakeValidName(Node, TEXT("ReferenceValid"));
+	TemporaryBoolTerminals.Add(Node, BoolTerm);
 }
 
 void FKCHandler_StateTreeBlueprintPropertyRefGet::Compile(FKismetFunctionContext& Context, UEdGraphNode* Node)
