StateTree: Use the temporary buffer on evaluators and global tasks when ticked in the selection process.
[RN]

diff --git a/Engine/Plugins/Runtime/StateTree/Source/StateTreeModule/Private/StateTreeExecutionContext.cpp b/Engine/Plugins/Runtime/StateTree/Source/StateTreeModule/Private/StateTreeExecutionContext.cpp
--- a/Engine/Plugins/Runtime/StateTree/Source/StateTreeModule/Private/StateTreeExecutionContext.cpp
+++ b/Engine/Plugins/Runtime/StateTree/Source/StateTreeModule/Private/StateTreeExecutionContext.cpp
@@ -3256,7 +3256,7 @@ EStateTreeRunStatus FStateTreeExecutionContext::StartTemporaryEvaluatorsAndGloba
 		// Copy bound properties.
 		if (Eval.BindingsBatch.IsValid())
 		{
-			CopyBatchOnActiveInstances(CurrentParentFrame, CurrentFrame, EvalInstanceView, Eval.BindingsBatch);
+			CopyBatchWithValidation(CurrentParentFrame, CurrentFrame, EvalInstanceView, Eval.BindingsBatch);
 		}
 
 		STATETREE_LOG(Verbose, TEXT("  Start: '%s'"), *Eval.Name.ToString());
@@ -3295,7 +3295,7 @@ EStateTreeRunStatus FStateTreeExecutionContext::StartTemporaryEvaluatorsAndGloba
 		// Copy bound properties.
 		if (Task.BindingsBatch.IsValid())
 		{
-			CopyBatchOnActiveInstances(CurrentParentFrame, CurrentFrame, TaskDataView, Task.BindingsBatch);
+			CopyBatchWithValidation(CurrentParentFrame, CurrentFrame, TaskDataView, Task.BindingsBatch);
 		}
 
 		STATETREE_LOG(Verbose, TEXT("  Start: '%s'"), *Task.Name.ToString());
