Fix CIS issue 744047: fixed redirectors

diff --git a/Engine/Plugins/Runtime/StateTree/Config/BaseStateTree.ini b/Engine/Plugins/Runtime/StateTree/Config/BaseStateTree.ini
--- a/Engine/Plugins/Runtime/StateTree/Config/BaseStateTree.ini
+++ b/Engine/Plugins/Runtime/StateTree/Config/BaseStateTree.ini
@@ -26,9 +26,9 @@
 +PropertyRedirects=(OldName="StateTree.Items",NewName="StateTree.Nodes")
 +PropertyRedirects=(OldName="StateTreeEditorNode.Item",NewName="StateTreeEditorNode.Node")
 
-+EnumRedirects=(OldName="EStateTreePropertyCopyType",ValueChanges=(("Plain","CopyPlain"),("Complex","CopyComplex"),("Bool","CopyBool"),("Struct","CopyStruct"),("Object","CopyObject"),("Name","CopyName"),("FixedArray","CopyFixedArray")))
 +EnumRedirects=(OldName="EStateTreePropertyAccessType",NewName="/Script/PropertyBindingUtilsModule.EPropertyBindingAccessType")
-+EnumRedirects=(OldName="EStateTreePropertyCopyType",NewName="/Script/PropertyBindingUtilsModule.EPropertyCopyType")
++EnumRedirects=(OldName="EStateTreePropertyCopyType",NewName="/Script/PropertyBindingUtilsModule.EPropertyCopyType",ValueChanges=(("Plain","CopyPlain"),("Complex","CopyComplex"),("Bool","CopyBool"),("Struct","CopyStruct"),("Object","CopyObject"),("Name","CopyName"),("FixedArray","CopyFixedArray")))
+
 +EnumRedirects=(OldName="/script/StateTreeModule.EPropertyBindingAccessType",ValueChanges=(("Array","IndexArray"),("EStateTreePropertyAccessType::Offset","EPropertyBindingAccessType::Offset")))
 
 +StructRedirects=(OldName="StateTreeCondition_CompareIntInstanceData",NewName="/script/StateTreeModule.StateTreeCompareIntConditionInstanceData")
