StateTree: Prevent execution of the same state tree in a parallel task or in a linked asset. Add unit tests for parallel and linked asset. Create unit test for weak context.

diff --git a/Engine/Plugins/Runtime/StateTree/Source/StateTreeEditorModule/Public/StateTreeEditorData.h b/Engine/Plugins/Runtime/StateTree/Source/StateTreeEditorModule/Public/StateTreeEditorData.h
--- a/Engine/Plugins/Runtime/StateTree/Source/StateTreeEditorModule/Public/StateTreeEditorData.h
+++ b/Engine/Plugins/Runtime/StateTree/Source/StateTreeEditorModule/Public/StateTreeEditorData.h
@@ -233,7 +233,7 @@ public:
 		FStateTreeEditorNode& EditorNode = GlobalTasks.AddDefaulted_GetRef();
 		EditorNode.ID = FGuid::NewGuid();
 		EditorNode.Node.InitializeAs<T>(Forward<TArgs>(InArgs)...);
-		T& Task = EditorNode.Node.GetMutable<T>();
+		const FStateTreeNodeBase& Task = EditorNode.Node.Get<FStateTreeNodeBase>();
 		if (const UScriptStruct* InstanceType = Cast<const UScriptStruct>(Task.GetInstanceDataType()))
 		{
 			EditorNode.Instance.InitializeAs(InstanceType);
diff --git a/Engine/Plugins/Runtime/StateTree/Source/StateTreeModule/Private/Tasks/StateTreeRunParallelStateTreeTask.cpp b/Engine/Plugins/Runtime/StateTree/Source/StateTreeModule/Private/Tasks/StateTreeRunParallelStateTreeTask.cpp
--- a/Engine/Plugins/Runtime/StateTree/Source/StateTreeModule/Private/Tasks/StateTreeRunParallelStateTreeTask.cpp
+++ b/Engine/Plugins/Runtime/StateTree/Source/StateTreeModule/Private/Tasks/StateTreeRunParallelStateTreeTask.cpp
@@ -31,6 +31,14 @@ EStateTreeRunStatus FStateTreeRunParallelStateTreeTask::EnterState(FStateTreeExe
 		return EStateTreeRunStatus::Failed;
 	}
 
+	// Find if it's a recursive call. The detection is not perfect. For example: StateTrees with a parallel task that links to each other cannot be detected.
+	const bool bInParentContext = Context.GetActiveFrames().ContainsByPredicate([NewTree = StateTreeToRun.GetStateTree()](const FStateTreeExecutionFrame& Frame){ return Frame.StateTree == NewTree; });
+	if (bInParentContext)
+	{
+		UE_LOG(LogStateTree, Warning, TEXT("Trying to start a new parallel tree from the same tree '%s'"), *StateTreeToRun.GetStateTree()->GetName());
+		return EStateTreeRunStatus::Failed;
+	}
+
 	if (InstanceData.ScheduledTickHandle.IsValid())
 	{
 		Context.RemoveScheduledTickRequest(InstanceData.ScheduledTickHandle);
diff --git a/Engine/Plugins/Runtime/StateTree/Source/StateTreeTestSuite/Private/StateTreeLinkedStateTest.cpp b/Engine/Plugins/Runtime/StateTree/Source/StateTreeTestSuite/Private/StateTreeLinkedStateTest.cpp
--- a/Engine/Plugins/Runtime/StateTree/Source/StateTreeTestSuite/Private/StateTreeLinkedStateTest.cpp
+++ b/Engine/Plugins/Runtime/StateTree/Source/StateTreeTestSuite/Private/StateTreeLinkedStateTest.cpp
@@ -607,6 +607,140 @@ struct FStateTreeTest_EmptyStateWithTickTransitionLinkedAsset : FStateTreeTestBa
 };
 IMPLEMENT_AI_INSTANT_TEST(FStateTreeTest_EmptyStateWithTickTransitionLinkedAsset, "System.StateTree.LinkedAsset.EmptyStateWithTickTransition");
 
+struct FStateTreeTest_RecursiveLinkedAsset : FStateTreeTestBase
+{
+	virtual bool InstantTest() override
+	{
+		//Tree 1
+		//	Root
+		//		StateLinkedTree1 (Tree2) -> Next
+		//		StateA -> Succeeded
+		//Tree 2
+		//	Root
+		//		StateLinkedTreeA (Tree1) -> Next
+		//		StateA -> Succeeded
+
+		UStateTree& StateTree1 = NewStateTree();
+		UStateTreeState* Root1 = nullptr;
+		// Asset 1 definition
+		{
+			UStateTreeEditorData& EditorData1 = *Cast<UStateTreeEditorData>(StateTree1.EditorData);
+			Root1 = &EditorData1.AddSubTree("Tree1StateRoot");
+
+			FStateTreeCompilerLog Log;
+			FStateTreeCompiler Compiler(Log);
+			const bool bResult = Compiler.Compile(StateTree1);
+			AITEST_TRUE("StateTree 1 should get compiled", bResult);
+		}
+
+		UStateTree& StateTree2 = NewStateTree();
+		UStateTreeState* Root2 = nullptr;
+		// Asset 2 definition
+		{
+			UStateTreeEditorData& EditorData2 = *Cast<UStateTreeEditorData>(StateTree2.EditorData);
+			Root2 = &EditorData2.AddSubTree("Tree2StateRoot");
+
+			FStateTreeCompilerLog Log;
+			FStateTreeCompiler Compiler(Log);
+			const bool bResult = Compiler.Compile(StateTree2);
+			AITEST_TRUE("StateTree 2 should get compiled", bResult);
+		}
+		// Asset 1 implementation
+		{
+			{
+				TStateTreeEditorNode<FTestTask_PrintValue>& Task1 = Root1->AddTask<FTestTask_PrintValue>("Tree1StateRootTaskA");
+				Task1.GetInstanceData().Value = 101;
+			}
+			{
+				UStateTreeState& C1 = Root1->AddChildState("Tree1StateLinkedTree1", EStateTreeStateType::LinkedAsset);
+				C1.Tag = GetTestTag1();
+				C1.SetLinkedStateAsset(&StateTree2);
+				C1.AddTransition(EStateTreeTransitionTrigger::OnStateSucceeded, EStateTreeTransitionType::NextState);
+			}
+			{
+				UStateTreeState& StateA = Root1->AddChildState("Tree1StateA", EStateTreeStateType::State);
+				TStateTreeEditorNode<FTestTask_PrintValue>& Task = StateA.AddTask<FTestTask_PrintValue>("Tree1StateA");
+				Task.GetInstanceData().Value = 102;
+				FStateTreeTransition& Transition = StateA.AddTransition(EStateTreeTransitionTrigger::OnTick, EStateTreeTransitionType::GotoState, Root1);
+				Transition.bDelayTransition = true;
+				Transition.DelayDuration = 1.0f;
+			}
+
+			FStateTreeCompilerLog Log;
+			FStateTreeCompiler Compiler(Log);
+			const bool bResult = Compiler.Compile(StateTree1);
+			AITEST_TRUE("StateTree should get compiled", bResult);
+		}
+		// Asset 2 implementation
+		UStateTreeState* Tree2StateLinkedTree1 = nullptr;
+		{
+			{
+				TStateTreeEditorNode<FTestTask_PrintValue>& Task1 = Root2->AddTask<FTestTask_PrintValue>("Tree2StateRootTaskA");
+				Task1.GetInstanceData().Value = 201;
+			}
+			{
+				UStateTreeState& C1 = Root2->AddChildState("Tree2StateLinkedTree1", EStateTreeStateType::LinkedAsset);
+				C1.Tag = GetTestTag2();
+				C1.SetLinkedStateAsset(&StateTree2);
+				C1.AddTransition(EStateTreeTransitionTrigger::OnStateSucceeded, EStateTreeTransitionType::NextState);
+				Tree2StateLinkedTree1 = &C1;
+			}
+			{
+				UStateTreeState& StateD = Root2->AddChildState("Tree2StateA", EStateTreeStateType::State);
+				TStateTreeEditorNode<FTestTask_PrintValue>& Task = StateD.AddTask<FTestTask_PrintValue>("Tree2StateA");
+				Task.GetInstanceData().Value = 202;
+				FStateTreeTransition& Transition = StateD.AddTransition(EStateTreeTransitionTrigger::OnTick, EStateTreeTransitionType::GotoState, Root2);
+				Transition.bDelayTransition = true;
+				Transition.DelayDuration = 1.0f;
+			}
+
+			// circular dependency detected
+			FStateTreeCompilerLog Log;
+			FStateTreeCompiler Compiler(Log);
+			const bool bResult = Compiler.Compile(StateTree2);
+			AITEST_FALSE("StateTree should not compiled", bResult);
+		}
+		// Fix circular dependency
+		{
+			Tree2StateLinkedTree1->SetLinkedStateAsset(&StateTree1);
+			FStateTreeCompilerLog Log;
+			FStateTreeCompiler Compiler(Log);
+			const bool bResult = Compiler.Compile(StateTree2);
+			AITEST_TRUE("StateTree should get compiled", bResult);
+		}
+		// Run test
+		{
+			FStateTreeInstanceData InstanceData;
+			FTestStateTreeExecutionContext Exec(StateTree1, StateTree1, InstanceData);
+			{
+				const bool bInitSucceeded = Exec.IsValid();
+				AITEST_TRUE("StateTree should init", bInitSucceeded);
+			}
+
+			{
+				GetTestRunner().AddExpectedError(TEXT("Trying to recursively enter subtree"), EAutomationExpectedErrorFlags::Contains, 1);
+
+				EStateTreeRunStatus Status = Exec.Start();
+				AITEST_EQUAL(TEXT("Start should complete with Running"), Status, EStateTreeRunStatus::Running);
+				FTestStateTreeExecutionContext::FLogOrder LogOrder = Exec.Expect("Tree1StateRootTaskA", TEXT("EnterState101"));
+				AITEST_TRUE(TEXT("Start should enter Tree1StateRootTaskA"), LogOrder);
+				LogOrder = LogOrder.Then("Tree2StateRootTaskA", TEXT("EnterState201"));
+				AITEST_TRUE(TEXT("Start should enter Tree2StateRootTaskA"), LogOrder);
+				LogOrder = LogOrder.Then("Tree2StateA", TEXT("EnterState202"));
+				AITEST_TRUE(TEXT("Start should enter Tree2StateA"), LogOrder);
+				Exec.LogClear();
+				AITEST_TRUE("Doesn't have the expected error message.", GetTestRunner().HasMetExpectedMessages());
+			}
+			Exec.Stop();
+		}
+
+		return true;
+	}
+};
+
+
+IMPLEMENT_AI_INSTANT_TEST(FStateTreeTest_RecursiveLinkedAsset, "System.StateTree.LinkedAsset.RecursiveLinkedAsset");
+
 } // namespace UE::StateTree::Tests
 
 UE_ENABLE_OPTIMIZATION_SHIP
diff --git a/Engine/Plugins/Runtime/StateTree/Source/StateTreeTestSuite/Private/StateTreeTest.cpp b/Engine/Plugins/Runtime/StateTree/Source/StateTreeTestSuite/Private/StateTreeTest.cpp
--- a/Engine/Plugins/Runtime/StateTree/Source/StateTreeTestSuite/Private/StateTreeTest.cpp
+++ b/Engine/Plugins/Runtime/StateTree/Source/StateTreeTestSuite/Private/StateTreeTest.cpp
@@ -905,27 +905,91 @@ struct FStateTreeTest_StateRequiringEvent : FStateTreeTestBase
 
 		EStateTreeRunStatus Status = EStateTreeRunStatus::Unset;
 		FStateTreeInstanceData InstanceData;
-		FTestStateTreeExecutionContext Exec(StateTree, StateTree, InstanceData);
-		const bool bInitSucceeded = Exec.IsValid();
-		AITEST_TRUE("StateTree should init", bInitSucceeded);
 
 		const FString EnterStateStr(TEXT("EnterState"));
 
-		Status = Exec.Start();
-		AITEST_TRUE("StateTree TaskInitial should enter state", Exec.Expect(TaskInitial.GetName(), EnterStateStr));
-		Exec.LogClear();
+		auto Test=[&](FTestStateTreeExecutionContext& Exec)
+		{
+			AITEST_FALSE("StateTree TaskA should not enter state", Exec.Expect(TaskA.GetName(), EnterStateStr));
+			AITEST_FALSE("StateTree TaskB should not enter state", Exec.Expect(TaskB.GetName(), EnterStateStr));
+			AITEST_FALSE("StateTree TaskC should not enter state", Exec.Expect(TaskC.GetName(), EnterStateStr));
+			AITEST_FALSE("StateTree TaskD should not enter state", Exec.Expect(TaskD.GetName(), EnterStateStr));
+			AITEST_TRUE("StateTree TaskE should enter state", Exec.Expect(TaskE.GetName(), EnterStateStr));
+			return true;
+		};
 
-		Exec.SendEvent(ValidTag, FConstStructView::Make(FValidPayload()));
-		Status = Exec.Tick(0.1f);
+		{
+			FTestStateTreeExecutionContext Exec(StateTree, StateTree, InstanceData);
+			const bool bInitSucceeded = Exec.IsValid();
+			AITEST_TRUE("StateTree should init", bInitSucceeded);
 
-		AITEST_FALSE("StateTree TaskA should not enter state", Exec.Expect(TaskA.GetName(), EnterStateStr));
-		AITEST_FALSE("StateTree TaskB should not enter state", Exec.Expect(TaskB.GetName(), EnterStateStr));
-		AITEST_FALSE("StateTree TaskC should not enter state", Exec.Expect(TaskC.GetName(), EnterStateStr));
-		AITEST_FALSE("StateTree TaskD should not enter state", Exec.Expect(TaskD.GetName(), EnterStateStr));
-		AITEST_TRUE("StateTree TaskE should enter state", Exec.Expect(TaskE.GetName(), EnterStateStr));
-		Exec.LogClear();
-		
-		Exec.Stop();
+			Status = Exec.Start();
+			AITEST_TRUE("StateTree TaskInitial should enter state", Exec.Expect(TaskInitial.GetName(), EnterStateStr));
+			Exec.LogClear();
+
+			Exec.SendEvent(ValidTag, FConstStructView::Make(FValidPayload()));
+			Status = Exec.Tick(0.1f);
+
+			if (!Test(Exec))
+			{
+				return false;
+			}
+			Exec.LogClear();
+
+			Exec.Stop();
+		}
+		// Same test but event sent with weak context while the FTestStateTreeExecutionContext still exist
+		{
+			FTestStateTreeExecutionContext Exec(StateTree, StateTree, InstanceData);
+			const bool bInitSucceeded = Exec.IsValid();
+			AITEST_TRUE("StateTree should init", bInitSucceeded);
+
+			Status = Exec.Start();
+			AITEST_TRUE("StateTree TaskInitial should enter state", Exec.Expect(TaskInitial.GetName(), EnterStateStr));
+			Exec.LogClear();
+
+			FStateTreeWeakExecutionContext WeakExec = Exec.MakeWeakExecutionContext();
+			WeakExec.SendEvent(ValidTag, FConstStructView::Make(FValidPayload()));
+			Status = Exec.Tick(0.1f);
+
+			if (!Test(Exec))
+			{
+				return false;
+			}
+			Exec.LogClear();
+
+			Exec.Stop();
+		}
+		// Same test but event sent with weak context
+		{
+			FStateTreeWeakExecutionContext WeakExec;
+			{
+				FTestStateTreeExecutionContext Exec(StateTree, StateTree, InstanceData);
+				const bool bInitSucceeded = Exec.IsValid();
+				AITEST_TRUE("StateTree should init", bInitSucceeded);
+
+				Status = Exec.Start();
+				AITEST_TRUE("StateTree TaskInitial should enter state", Exec.Expect(TaskInitial.GetName(), EnterStateStr));
+				Exec.LogClear();
+
+				WeakExec = Exec.MakeWeakExecutionContext();
+			}
+			{
+				WeakExec.SendEvent(ValidTag, FConstStructView::Make(FValidPayload()));
+			}
+			{
+				FTestStateTreeExecutionContext Exec(StateTree, StateTree, InstanceData);
+				Status = Exec.Tick(0.1f);
+
+				if (!Test(Exec))
+				{
+					return false;
+				}
+				Exec.LogClear();
+
+				Exec.Stop();
+			}
+		}
 
 		return true;
 	}
@@ -1638,6 +1702,51 @@ struct FStateTreeTest_NestedOverride : FStateTreeTestBase
 };
 IMPLEMENT_AI_INSTANT_TEST(FStateTreeTest_NestedOverride, "System.StateTree.NestedOverride");
 
+// Test parallel tree event priority handling.
+struct FStateTreeTest_RecursiveParallelTask : FStateTreeTestBase
+{
+	virtual bool InstantTest() override
+	{
+		//Tree 1
+		//	Root (with task that runs Tree 1)
+
+		UStateTree& StateTree1 = NewStateTree();
+		{
+			UStateTreeEditorData& EditorData1 = *Cast<UStateTreeEditorData>(StateTree1.EditorData);
+			UStateTreeState* Root1 = &EditorData1.AddSubTree("Tree1StateRoot");
+
+			TStateTreeEditorNode<FStateTreeRunParallelStateTreeTask>& GlobalTask = EditorData1.AddGlobalTask<FStateTreeRunParallelStateTreeTask>();
+			GlobalTask.GetInstanceData().StateTree.SetStateTree(&StateTree1);
+
+			TStateTreeEditorNode<FTestTask_PrintValue>& RootTask = Root1->AddTask<FTestTask_PrintValue>();
+			RootTask.GetInstanceData().Value = 101;
+		}
+		{
+			FStateTreeCompilerLog Log;
+			FStateTreeCompiler Compiler(Log);
+			const bool bResult = Compiler.Compile(StateTree1);
+			AITEST_TRUE("StateTreePar should get compiled", bResult);
+		}
+		{
+			FStateTreeInstanceData InstanceData;
+			FTestStateTreeExecutionContext Exec(StateTree1, StateTree1, InstanceData);
+			const bool bInitSucceeded = Exec.IsValid();
+			AITEST_TRUE("StateTree should init", bInitSucceeded);
+
+			{
+				GetTestRunner().AddExpectedError(TEXT("Trying to start a new parallel tree from the same tree"), EAutomationExpectedErrorFlags::Contains, 1);
+
+				EStateTreeRunStatus Status = Exec.Start();
+				AITEST_EQUAL(TEXT("Start should complete with failed"), Status, EStateTreeRunStatus::Failed);
+				AITEST_TRUE(TEXT(""), GetTestRunner().HasMetExpectedMessages());
+			}
+		}
+		return true;
+	}
+};
+
+IMPLEMENT_AI_INSTANT_TEST(FStateTreeTest_RecursiveParallelTask, "System.StateTree.RecursiveParallelTask");
+
 // Test parallel tree event priority handling.
 struct FStateTreeTest_ParallelEventPriority : FStateTreeTestBase
 {
