StateTree: Add NOT to enter condition in the object category.

diff --git a/Engine/Plugins/Runtime/StateTree/Source/StateTreeModule/Private/Conditions/StateTreeObjectConditions.cpp b/Engine/Plugins/Runtime/StateTree/Source/StateTreeModule/Private/Conditions/StateTreeObjectConditions.cpp
--- a/Engine/Plugins/Runtime/StateTree/Source/StateTreeModule/Private/Conditions/StateTreeObjectConditions.cpp
+++ b/Engine/Plugins/Runtime/StateTree/Source/StateTreeModule/Private/Conditions/StateTreeObjectConditions.cpp
@@ -2,6 +2,7 @@
 
 #include "Conditions/StateTreeObjectConditions.h"
 #include "StateTreeExecutionContext.h"
+#include "StateTreeNodeDescriptionHelpers.h"
 
 #include UE_INLINE_GENERATED_CPP_BY_NAME(StateTreeObjectConditions)
 
@@ -19,6 +20,18 @@ bool FStateTreeObjectIsValidCondition::TestCondition(FStateTreeExecutionContext&
 	return bResult ^ bInvert;
 }
 
+#if WITH_EDITOR
+FText FStateTreeObjectIsValidCondition::GetDescription(const FGuid& ID, FStateTreeDataView InstanceDataView, const IStateTreeBindingLookup& BindingLookup, EStateTreeNodeFormatting Formatting) const
+{
+	const FText InvertText = UE::StateTree::DescHelpers::GetInvertText(bInvert, Formatting);
+	const FText Format = (Formatting == EStateTreeNodeFormatting::RichText)
+		? LOCTEXT("ObjectIsValidConditionsRich", "{EmptyOrNot}<s>Is Object Valid</>")
+		: LOCTEXT("ObjectIsValidCondition", "{EmptyOrNot}Is Object Valid");
+	return FText::FormatNamed(Format,
+		TEXT("EmptyOrNot"), InvertText);
+}
+#endif
+
 //----------------------------------------------------------------------//
 //  FStateTreeCondition_ObjectEquals
 //----------------------------------------------------------------------//
@@ -30,6 +43,18 @@ bool FStateTreeObjectEqualsCondition::TestCondition(FStateTreeExecutionContext&
 	return bResult ^ bInvert;
 }
 
+#if WITH_EDITOR
+FText FStateTreeObjectEqualsCondition::GetDescription(const FGuid& ID, FStateTreeDataView InstanceDataView, const IStateTreeBindingLookup& BindingLookup, EStateTreeNodeFormatting Formatting) const
+{
+	const FText InvertText = UE::StateTree::DescHelpers::GetInvertText(bInvert, Formatting);
+	const FText Format = (Formatting == EStateTreeNodeFormatting::RichText)
+		? LOCTEXT("ObjectEqualsConditionRich", "{EmptyOrNot}<s>Is Object Equals</>")
+		: LOCTEXT("ObjectEqualsCondition", "{EmptyOrNot}Is Object Equals");
+	return FText::FormatNamed(Format,
+		TEXT("EmptyOrNot"), InvertText);
+}
+#endif
+
 //----------------------------------------------------------------------//
 //  FStateTreeCondition_ObjectIsChildOfClass
 //----------------------------------------------------------------------//
@@ -42,4 +67,16 @@ bool FStateTreeObjectIsChildOfClassCondition::TestCondition(FStateTreeExecutionC
 	return bResult ^ bInvert;
 }
 
+#if WITH_EDITOR
+FText FStateTreeObjectIsChildOfClassCondition::GetDescription(const FGuid& ID, FStateTreeDataView InstanceDataView, const IStateTreeBindingLookup& BindingLookup, EStateTreeNodeFormatting Formatting) const
+{
+	const FText InvertText = UE::StateTree::DescHelpers::GetInvertText(bInvert, Formatting);
+	const FText Format = (Formatting == EStateTreeNodeFormatting::RichText)
+		? LOCTEXT("ObjectIsValidConditionsRich", "{EmptyOrNot}<s>Is Child Of Class</>")
+		: LOCTEXT("ObjectIsValidCondition", "{EmptyOrNot}Is Child Of Class");
+	return FText::FormatNamed(Format,
+		TEXT("EmptyOrNot"), InvertText);
+}
+#endif
+
 #undef LOCTEXT_NAMESPACE
diff --git a/Engine/Plugins/Runtime/StateTree/Source/StateTreeModule/Public/Conditions/StateTreeObjectConditions.h b/Engine/Plugins/Runtime/StateTree/Source/StateTreeModule/Public/Conditions/StateTreeObjectConditions.h
--- a/Engine/Plugins/Runtime/StateTree/Source/StateTreeModule/Public/Conditions/StateTreeObjectConditions.h
+++ b/Engine/Plugins/Runtime/StateTree/Source/StateTreeModule/Public/Conditions/StateTreeObjectConditions.h
@@ -29,9 +29,16 @@ struct STATETREEMODULE_API FStateTreeObjectIsValidCondition : public FStateTreeC
 		: bInvert(InInverts == EStateTreeCompare::Invert)
 	{}
 
-	virtual const UStruct* GetInstanceDataType() const override { return FInstanceDataType::StaticStruct(); }
+	virtual const UStruct* GetInstanceDataType() const override
+	{
+		return FInstanceDataType::StaticStruct();
+	}
 	virtual bool TestCondition(FStateTreeExecutionContext& Context) const override;
 
+#if WITH_EDITOR
+	virtual FText GetDescription(const FGuid& ID, FStateTreeDataView InstanceDataView, const IStateTreeBindingLookup& BindingLookup, EStateTreeNodeFormatting Formatting = EStateTreeNodeFormatting::Text) const override;
+#endif
+
 	UPROPERTY(EditAnywhere, Category = Parameter)
 	bool bInvert = false;
 };
@@ -64,9 +71,16 @@ struct STATETREEMODULE_API FStateTreeObjectEqualsCondition : public FStateTreeCo
 		: bInvert(InInverts == EStateTreeCompare::Invert)
 	{}
 
-	virtual const UStruct* GetInstanceDataType() const override { return FInstanceDataType::StaticStruct(); }
+	virtual const UStruct* GetInstanceDataType() const override
+	{
+		return FInstanceDataType::StaticStruct();
+	}
 	virtual bool TestCondition(FStateTreeExecutionContext& Context) const override;
 
+#if WITH_EDITOR
+	virtual FText GetDescription(const FGuid& ID, FStateTreeDataView InstanceDataView, const IStateTreeBindingLookup& BindingLookup, EStateTreeNodeFormatting Formatting = EStateTreeNodeFormatting::Text) const override;
+#endif
+
 	UPROPERTY(EditAnywhere, Category = "Parameter")
 	bool bInvert = false;
 };
@@ -99,9 +113,16 @@ struct STATETREEMODULE_API FStateTreeObjectIsChildOfClassCondition : public FSta
 		: bInvert(InInverts == EStateTreeCompare::Invert)
 	{}
 
-	virtual const UStruct* GetInstanceDataType() const override { return FInstanceDataType::StaticStruct(); }
+	virtual const UStruct* GetInstanceDataType() const override
+	{
+		return FInstanceDataType::StaticStruct();
+	}
 	virtual bool TestCondition(FStateTreeExecutionContext& Context) const override;
 
+#if WITH_EDITOR
+		virtual FText GetDescription(const FGuid& ID, FStateTreeDataView InstanceDataView, const IStateTreeBindingLookup& BindingLookup, EStateTreeNodeFormatting Formatting = EStateTreeNodeFormatting::Text) const override;
+#endif
+
 	UPROPERTY(EditAnywhere, Category = "Parameter")
 	bool bInvert = false;
 };
