[StateTree] updated old redirector to the new type

diff --git a/Engine/Plugins/Runtime/StateTree/Config/BaseStateTree.ini b/Engine/Plugins/Runtime/StateTree/Config/BaseStateTree.ini
--- a/Engine/Plugins/Runtime/StateTree/Config/BaseStateTree.ini
+++ b/Engine/Plugins/Runtime/StateTree/Config/BaseStateTree.ini
@@ -64,7 +64,7 @@
 +PropertyRedirects=(OldName="/Script/StateTreeModule.StateTreeTransition.GateDelay",NewName="/Script/StateTreeModule.StateTreeTransition.DelayDuration")
 
 +StructRedirects=(OldName="StateTreeEditorPropertyBinding",NewName="/script/StateTreeModule.StateTreePropertyPathBinding")
-+StructRedirects=(OldName="StateTreePropCopyBatch",NewName="/script/StateTreeModule.StateTreePropertyCopyBatch")
++StructRedirects=(OldName="StateTreePropCopyBatch",NewName="/script/PropertyBindingUtilsModule.PropertyBindingCopyInfoBatch")
 +StructRedirects=(OldName="/script/StateTreeModule.StateTreePropertyCopyBatch",NewName="/script/PropertyBindingUtilsModule.PropertyBindingCopyInfoBatch")
 
 +EnumRedirects=(OldName="/Script/StateTreeModule.EStateTreeDataSourceType",ValueChanges=(("LinkedStateParameterData", "StateParameterData")))
