StateTree: Use the frame id to identify the frame in the instance data. A tree that has 2 states with the same linked asset but with different condition would use the same buffer, and only initialialized once.

diff --git a/Engine/Plugins/Runtime/StateTree/Source/StateTreeModule/Private/StateTreeExecutionContext.cpp b/Engine/Plugins/Runtime/StateTree/Source/StateTreeModule/Private/StateTreeExecutionContext.cpp
--- a/Engine/Plugins/Runtime/StateTree/Source/StateTreeModule/Private/StateTreeExecutionContext.cpp
+++ b/Engine/Plugins/Runtime/StateTree/Source/StateTreeModule/Private/StateTreeExecutionContext.cpp
@@ -1448,7 +1448,7 @@ void FStateTreeExecutionContext::UpdateInstanceData(TConstArrayView<FStateTreeEx
 	{
 		FStateTreeTemporaryInstanceData* TempData = TempInstances.FindByPredicate([&Frame, &DataHandle](const FStateTreeTemporaryInstanceData& Data)
 		{
-			return Data.StateTree == Frame.StateTree && Data.RootState == Frame.RootState && Data.DataHandle == DataHandle; 
+			return Data.FrameId == Frame.FrameId && Data.DataHandle == DataHandle;
 		});
 		return TempData ? &TempData->Instance : nullptr;
 	};
@@ -3253,7 +3253,7 @@ EStateTreeRunStatus FStateTreeExecutionContext::StartTemporaryEvaluatorsAndGloba
 		}
 		
 		FNodeInstanceDataScope DataScope(*this, &Eval, EvalIndex, Eval.InstanceDataHandle, EvalInstanceView);
-		// Copy bound properties.
+		// Copy bound properties with the temporary buffer.
 		if (Eval.BindingsBatch.IsValid())
 		{
 			CopyBatchWithValidation(CurrentParentFrame, CurrentFrame, EvalInstanceView, Eval.BindingsBatch);
@@ -3292,7 +3292,7 @@ EStateTreeRunStatus FStateTreeExecutionContext::StartTemporaryEvaluatorsAndGloba
 
 		FNodeInstanceDataScope DataScope(*this, &Task, TaskIndex, Task.InstanceDataHandle, TaskDataView);
 
-		// Copy bound properties.
+		// Copy bound properties with the temporary buffer.
 		if (Task.BindingsBatch.IsValid())
 		{
 			CopyBatchWithValidation(CurrentParentFrame, CurrentFrame, TaskDataView, Task.BindingsBatch);
@@ -3332,7 +3332,7 @@ void FStateTreeExecutionContext::StopTemporaryEvaluatorsAndGlobalTasks(const FSt
 	for (int32 Index = TempInstances.Num() - 1; Index >= 0; Index--)
 	{
 		FStateTreeTemporaryInstanceData& TempInstance = TempInstances[Index];
-		if (TempInstance.StateTree != CurrentFrame.StateTree || TempInstance.RootState != CurrentFrame.RootState)
+		if (TempInstance.FrameId != CurrentFrame.FrameId)
 		{
 			continue;
 		}
diff --git a/Engine/Plugins/Runtime/StateTree/Source/StateTreeModule/Private/StateTreeInstanceData.cpp b/Engine/Plugins/Runtime/StateTree/Source/StateTreeModule/Private/StateTreeInstanceData.cpp
--- a/Engine/Plugins/Runtime/StateTree/Source/StateTreeModule/Private/StateTreeInstanceData.cpp
+++ b/Engine/Plugins/Runtime/StateTree/Source/StateTreeModule/Private/StateTreeInstanceData.cpp
@@ -256,8 +256,7 @@ FStructView FStateTreeInstanceStorage::AddTemporaryInstance(UObject& InOwner, co
 {
 	FStateTreeTemporaryInstanceData* TempInstance = TemporaryInstances.FindByPredicate([&Frame, &OwnerNodeIndex, &DataHandle](const FStateTreeTemporaryInstanceData& TempInstance)
 	{
-		return TempInstance.StateTree == Frame.StateTree
-				&& TempInstance.RootState == Frame.RootState
+		return TempInstance.FrameId == Frame.FrameId
 				&& TempInstance.OwnerNodeIndex == OwnerNodeIndex
 				&& TempInstance.DataHandle == DataHandle;
 	});
@@ -273,8 +272,7 @@ FStructView FStateTreeInstanceStorage::AddTemporaryInstance(UObject& InOwner, co
 	{
 		TempInstance = &TemporaryInstances.AddDefaulted_GetRef();
 		check(TempInstance);
-		TempInstance->StateTree = Frame.StateTree;
-		TempInstance->RootState = Frame.RootState;
+		TempInstance->FrameId = Frame.FrameId;
 		TempInstance->OwnerNodeIndex = OwnerNodeIndex;
 		TempInstance->DataHandle = DataHandle;
 		TempInstance->Instance = NewInstanceData;
@@ -295,8 +293,7 @@ FStructView FStateTreeInstanceStorage::GetMutableTemporaryStruct(const FStateTre
 {
 	FStateTreeTemporaryInstanceData* ExistingInstance = TemporaryInstances.FindByPredicate([&Frame, &DataHandle](const FStateTreeTemporaryInstanceData& TempInstance)
 	{
-		return TempInstance.StateTree == Frame.StateTree
-				&& TempInstance.RootState == Frame.RootState
+		return TempInstance.FrameId == Frame.FrameId
 				&& TempInstance.DataHandle == DataHandle;
 	});
 	return ExistingInstance ? FStructView(ExistingInstance->Instance) : FStructView();
@@ -306,8 +303,7 @@ UObject* FStateTreeInstanceStorage::GetMutableTemporaryObject(const FStateTreeEx
 {
 	FStateTreeTemporaryInstanceData* ExistingInstance = TemporaryInstances.FindByPredicate([&Frame, &DataHandle](const FStateTreeTemporaryInstanceData& TempInstance)
 	{
-		return TempInstance.StateTree == Frame.StateTree
-				&& TempInstance.RootState == Frame.RootState
+		return TempInstance.FrameId == Frame.FrameId
 				&& TempInstance.DataHandle == DataHandle;
 	});
 	if (ExistingInstance)
diff --git a/Engine/Plugins/Runtime/StateTree/Source/StateTreeModule/Public/StateTreeInstanceData.h b/Engine/Plugins/Runtime/StateTree/Source/StateTreeModule/Public/StateTreeInstanceData.h
--- a/Engine/Plugins/Runtime/StateTree/Source/StateTreeModule/Public/StateTreeInstanceData.h
+++ b/Engine/Plugins/Runtime/StateTree/Source/StateTreeModule/Public/StateTreeInstanceData.h
@@ -45,11 +45,7 @@ struct STATETREEMODULE_API FStateTreeTemporaryInstanceData
 {
 	GENERATED_BODY()
 
-	UPROPERTY()
-	TObjectPtr<const UStateTree> StateTree = nullptr;
-
-	UPROPERTY()
-	FStateTreeStateHandle RootState = FStateTreeStateHandle::Invalid;
+	UE::StateTree::FActiveFrameId FrameId;
 
 	UPROPERTY()
 	FStateTreeDataHandle DataHandle = FStateTreeDataHandle::Invalid;
@@ -59,6 +55,16 @@ struct STATETREEMODULE_API FStateTreeTemporaryInstanceData
 	
 	UPROPERTY()
 	FInstancedStruct Instance;
+
+#if WITH_EDITORONLY_DATA
+	UE_DEPRECATED(5.6, "Use the frame id to identify the frame.")
+	UPROPERTY()
+	TObjectPtr<const UStateTree> StateTree = nullptr;
+
+	UE_DEPRECATED(5.6, "Use the frame id to identify the frame.")
+	UPROPERTY()
+	FStateTreeStateHandle RootState = FStateTreeStateHandle::Invalid;
+#endif
 };
 
 
