Fix build. Bad deprecation warning.
#rnx

diff --git a/Engine/Plugins/Runtime/StateTree/Source/StateTreeModule/Public/StateTreeInstanceData.h b/Engine/Plugins/Runtime/StateTree/Source/StateTreeModule/Public/StateTreeInstanceData.h
--- a/Engine/Plugins/Runtime/StateTree/Source/StateTreeModule/Public/StateTreeInstanceData.h
+++ b/Engine/Plugins/Runtime/StateTree/Source/StateTreeModule/Public/StateTreeInstanceData.h
@@ -45,6 +45,15 @@ struct STATETREEMODULE_API FStateTreeTemporaryInstanceData
 {
 	GENERATED_BODY()
 
+	FStateTreeTemporaryInstanceData() = default;
+
+PRAGMA_DISABLE_DEPRECATION_WARNINGS
+	FStateTreeTemporaryInstanceData(const FStateTreeTemporaryInstanceData& Other) = default;
+	FStateTreeTemporaryInstanceData(FStateTreeTemporaryInstanceData&& Other) = default;
+	FStateTreeTemporaryInstanceData& operator=(FStateTreeTemporaryInstanceData const& Other) = default;
+	FStateTreeTemporaryInstanceData& operator=(FStateTreeTemporaryInstanceData&& Other) = default;
+PRAGMA_ENABLE_DEPRECATION_WARNINGS
+
 	UE::StateTree::FActiveFrameId FrameId;
 
 	UPROPERTY()
@@ -63,7 +72,7 @@ struct STATETREEMODULE_API FStateTreeTemporaryInstanceData
 
 	UE_DEPRECATED(5.6, "Use the frame id to identify the frame.")
 	UPROPERTY()
-	FStateTreeStateHandle RootState = FStateTreeStateHandle::Invalid;
+	FStateTreeStateHandle RootState = FStateTreeStateHandle::Invalid; 
 #endif
 };
 
