StateTree: Fix a localization issue.
#rnx

diff --git a/Engine/Plugins/Runtime/StateTree/Source/StateTreeModule/Private/Conditions/StateTreeObjectConditions.cpp b/Engine/Plugins/Runtime/StateTree/Source/StateTreeModule/Private/Conditions/StateTreeObjectConditions.cpp
--- a/Engine/Plugins/Runtime/StateTree/Source/StateTreeModule/Private/Conditions/StateTreeObjectConditions.cpp
+++ b/Engine/Plugins/Runtime/StateTree/Source/StateTreeModule/Private/Conditions/StateTreeObjectConditions.cpp
@@ -6,7 +6,7 @@
 
 #include UE_INLINE_GENERATED_CPP_BY_NAME(StateTreeObjectConditions)
 
-#define LOCTEXT_NAMESPACE "StateTreeEditor"
+#define LOCTEXT_NAMESPACE "StateTreeObjectCondition"
 
 
 //----------------------------------------------------------------------//
