StateTree: Add state Description. Editor only, allows the user to add a lengthy description on the state.

diff --git a/Engine/Plugins/Runtime/StateTree/Source/StateTreeEditorModule/Private/Customizations/Widgets/SCompactStateTreeView.cpp b/Engine/Plugins/Runtime/StateTree/Source/StateTreeEditorModule/Private/Customizations/Widgets/SCompactStateTreeView.cpp
--- a/Engine/Plugins/Runtime/StateTree/Source/StateTreeEditorModule/Private/Customizations/Widgets/SCompactStateTreeView.cpp
+++ b/Engine/Plugins/Runtime/StateTree/Source/StateTreeEditorModule/Private/Customizations/Widgets/SCompactStateTreeView.cpp
@@ -239,6 +239,7 @@ void SCompactStateTreeView::CacheState(TSharedPtr<FStateTreeStateItem> ParentNod
 	{
 		TSharedRef<FStateTreeStateItem> StateItem = MakeShared<FStateTreeStateItem>();
 		StateItem->Desc = FText::FromName(State->Name);
+		StateItem->TooltipText = FText::FromString(State->Description);
 		StateItem->StateID = State->ID;
 		StateItem->bIsSubTree = State->Type == EStateTreeStateType::Subtree;
 		StateItem->bIsEnabled = State->bEnabled;
@@ -303,7 +304,6 @@ TSharedRef<ITableRow> SCompactStateTreeView::GenerateStateItemRow(TSharedPtr<FSt
 			.Text(Item->Desc)
 			.TextStyle(&FStateTreeEditorStyle::Get().GetWidgetStyle<FTextBlockStyle>("Normal.Normal"))
 			.OverflowPolicy(ETextOverflowPolicy::Ellipsis)
-			.ToolTipText(Item->TooltipText)
 			.IsEnabled(Item->bIsEnabled)
 			.HighlightText_Lambda([this]() { return SearchBox.IsValid() ? SearchBox->GetText() : FText::GetEmpty(); })
 			+SRichTextBlock::Decorator(FTextStyleDecorator::Create(TEXT(""), FStateTreeEditorStyle::Get().GetWidgetStyle<FTextBlockStyle>("Normal.Normal")))
@@ -348,6 +348,7 @@ TSharedRef<ITableRow> SCompactStateTreeView::GenerateStateItemRow(TSharedPtr<FSt
 		[
 			SNew(SBorder)
 			.BorderBackgroundColor(Item.Get(), &FStateTreeStateItem::GetBorderColor)
+			.ToolTipText(Item->TooltipText)
 			[
 				Container
 			]
diff --git a/Engine/Plugins/Runtime/StateTree/Source/StateTreeEditorModule/Private/SStateTreeViewRow.cpp b/Engine/Plugins/Runtime/StateTree/Source/StateTreeEditorModule/Private/SStateTreeViewRow.cpp
--- a/Engine/Plugins/Runtime/StateTree/Source/StateTreeEditorModule/Private/SStateTreeViewRow.cpp
+++ b/Engine/Plugins/Runtime/StateTree/Source/StateTreeEditorModule/Private/SStateTreeViewRow.cpp
@@ -301,6 +301,23 @@ void SStateTreeViewRow::Construct(const FArguments& InArgs, const TSharedRef<STa
 												.IsSelected(this, &SStateTreeViewRow::IsStateSelected)
 											]
 
+											// Description
+											+ SHorizontalBox::Slot()
+											.VAlign(VAlign_Center)
+											.AutoWidth()
+											[
+												SNew(SBox)
+												.Padding(FMargin(2.f, 0.f, 2.f, 1.f))
+												.Visibility(this, &SStateTreeViewRow::GetStateDescriptionVisibility)
+												[
+													SNew(SImage)
+													.Image(FAppStyle::Get().GetBrush(TEXT("Icons.Comment")))
+													.ColorAndOpacity(FStyleColors::Foreground)
+													.ColorAndOpacity(IconTint)
+													.ToolTipText(this, &SStateTreeViewRow::GetStateDescription)
+												]
+											]
+
 											// Flags icons
 											+ SHorizontalBox::Slot()
 											.VAlign(VAlign_Center)
@@ -2107,6 +2124,24 @@ EVisibility SStateTreeViewRow::GetTransitionsBreakpointVisibility(const EStateTr
 	return EVisibility::Collapsed;
 }
 
+EVisibility SStateTreeViewRow::GetStateDescriptionVisibility() const
+{
+	if (const UStateTreeState* State = WeakState.Get())
+	{
+		return State->Description.IsEmpty() ? EVisibility::Collapsed : EVisibility::Visible;
+	}
+	return EVisibility::Collapsed;
+}
+
+FText SStateTreeViewRow::GetStateDescription() const
+{
+	if (const UStateTreeState* State = WeakState.Get())
+	{
+		return FText::FromString(State->Description);
+	}
+	return FText::GetEmpty();
+}
+
 EVisibility SStateTreeViewRow::GetTransitionDashVisibility() const
 {
 	if (const UStateTreeState* State = WeakState.Get())
diff --git a/Engine/Plugins/Runtime/StateTree/Source/StateTreeEditorModule/Private/SStateTreeViewRow.h b/Engine/Plugins/Runtime/StateTree/Source/StateTreeEditorModule/Private/SStateTreeViewRow.h
--- a/Engine/Plugins/Runtime/StateTree/Source/StateTreeEditorModule/Private/SStateTreeViewRow.h
+++ b/Engine/Plugins/Runtime/StateTree/Source/StateTreeEditorModule/Private/SStateTreeViewRow.h
@@ -77,6 +77,9 @@ private:
 	EVisibility GetWarningsVisibility() const;
 	FText GetWarningsTooltipText() const;
 
+	EVisibility GetStateDescriptionVisibility() const;
+	FText GetStateDescription() const;
+
 	EVisibility GetTransitionDashVisibility() const;
 	
 	enum class ETransitionDescRequirement : uint8
diff --git a/Engine/Plugins/Runtime/StateTree/Source/StateTreeEditorModule/Public/StateTreeState.h b/Engine/Plugins/Runtime/StateTree/Source/StateTreeEditorModule/Public/StateTreeState.h
--- a/Engine/Plugins/Runtime/StateTree/Source/StateTreeEditorModule/Public/StateTreeState.h
+++ b/Engine/Plugins/Runtime/StateTree/Source/StateTreeEditorModule/Public/StateTreeState.h
@@ -341,6 +341,10 @@ public:
 	UPROPERTY(EditDefaultsOnly, Category = "State")
 	FName Name;
 
+	/** Description of the State */
+	UPROPERTY(EditDefaultsOnly, Category = "State", meta=(MultiLine))
+	FString Description;
+
 	/** GameplayTag describing the State */
 	UPROPERTY(EditDefaultsOnly, Category = "State")
 	FGameplayTag Tag;
