StateTree: Add the node indexes to the description. It helps when reading big assets. Enable the log with cvar StateTree.Compiler.LogResultOnCompilationCompleted.

#rnx

diff --git a/Engine/Plugins/Runtime/StateTree/Source/StateTreeModule/Private/StateTree.cpp b/Engine/Plugins/Runtime/StateTree/Source/StateTreeModule/Private/StateTree.cpp
--- a/Engine/Plugins/Runtime/StateTree/Source/StateTreeModule/Private/StateTree.cpp
+++ b/Engine/Plugins/Runtime/StateTree/Source/StateTreeModule/Private/StateTree.cpp
@@ -1277,9 +1277,11 @@ FString UStateTree::DebugInternalLayoutAsString() const
 	GetFullName(DebugString);
 	DebugString << TEXT("')\n");
 
-	auto PrintObjectNameSafe = [&DebugString](const UObject* Obj)
+	auto PrintObjectNameSafe = [&DebugString](int32 Index, const UObject* Obj)
 		{
-			DebugString << TEXT("  ");
+			DebugString << TEXT("  (");
+			DebugString << Index;
+			DebugString << TEXT(")");
 			if (Obj)
 			{
 				DebugString << Obj->GetFName();
@@ -1290,9 +1292,11 @@ FString UStateTree::DebugInternalLayoutAsString() const
 			}
 			DebugString << TEXT('\n');
 		};
-	auto PrintViewNameSafe = [&DebugString](const FConstStructView& View)
+	auto PrintViewNameSafe = [&DebugString](int32 Index, const FConstStructView& View)
 		{
-			DebugString << TEXT("  ");
+			DebugString << TEXT("  (");
+			DebugString << Index;
+			DebugString << TEXT(")");
 			if (View.IsValid())
 			{
 				DebugString << View.GetScriptStruct()->GetFName();
@@ -1309,7 +1313,7 @@ FString UStateTree::DebugInternalLayoutAsString() const
 	for (int32 Index = 0; Index < Nodes.Num(); Index++)
 	{
 		const FConstStructView Node = Nodes[Index];
-		PrintViewNameSafe(Node);
+		PrintViewNameSafe(Index, Node);
 	}
 
 	// Instance InstanceData data (e.g. tasks)
@@ -1319,12 +1323,12 @@ FString UStateTree::DebugInternalLayoutAsString() const
 		if (DefaultInstanceData.IsObject(Index))
 		{
 			const UObject* Data = DefaultInstanceData.GetObject(Index);
-			PrintObjectNameSafe(Data);
+			PrintObjectNameSafe(Index, Data);
 		}
 		else
 		{
 			const FConstStructView Data = DefaultInstanceData.GetStruct(Index);
-			PrintViewNameSafe(Data);
+			PrintViewNameSafe(Index, Data);
 		}
 	}
 
@@ -1332,10 +1336,15 @@ FString UStateTree::DebugInternalLayoutAsString() const
 	DebugString.Appendf(TEXT("\nExternal Data(%d)\n"), ExternalDataDescs.Num());
 	if (ExternalDataDescs.Num())
 	{
-		DebugString.Appendf(TEXT("  [% -40s | % -8s | % 15s]\n"), TEXT("Name"), TEXT("Optional"), TEXT("Handle"));
-		for (const FStateTreeExternalDataDesc& Desc : ExternalDataDescs)
+		DebugString.Appendf(TEXT("  [ (Idx) | %-40s | %-8s | %15s ]\n"), TEXT("Name"), TEXT("Optional"), TEXT("Handle"));
+		for (int32 DataDescIndex = 0; DataDescIndex < ExternalDataDescs.Num(); ++DataDescIndex)
 		{
-			DebugString.Appendf(TEXT("  | %-40s | %8s | %15s |\n"), Desc.Struct ? *Desc.Struct->GetName() : TEXT("null"), *UEnum::GetDisplayValueAsText(Desc.Requirement).ToString(), *Desc.Handle.DataHandle.Describe());
+			const FStateTreeExternalDataDesc& Desc = ExternalDataDescs[DataDescIndex];
+			DebugString.Appendf(TEXT("  | (%3d) | %-40s | %8s | %15s |\n"),
+				DataDescIndex,
+				Desc.Struct ? *Desc.Struct->GetName() : TEXT("null"),
+				*UEnum::GetDisplayValueAsText(Desc.Requirement).ToString(),
+				*Desc.Handle.DataHandle.Describe());
 		}
 	}
 
@@ -1348,16 +1357,19 @@ FString UStateTree::DebugInternalLayoutAsString() const
 	DebugString.Appendf(TEXT("\nStates(%d)\n"), States.Num());
 	if (States.Num())
 	{
-		DebugString.Appendf(TEXT("  [ %-30s | %15s | %5s [%3s:%-3s[ | Begin Idx : %4s %4s %4s %4s | Num : %4s %4s %4s %4s ]\n"),
+		DebugString.Appendf(TEXT("  [ (Idx) | %-30s | %15s | %5s [%3s:%-3s[ | Begin Idx : %4s %4s %4s %4s | Num : %4s %4s %4s %4s ]\n"),
 			TEXT("Name"), TEXT("Parent"), TEXT("Child"), TEXT("Beg"), TEXT("End"),
 			TEXT("Cond"), TEXT("Tr"), TEXT("Tsk"), TEXT("Uti"), TEXT("Cond"), TEXT("Tr"), TEXT("Tsk"), TEXT("Uti"));
-		for (const FCompactStateTreeState& State : States)
+		for (int32 StateIndex = 0; StateIndex < States.Num(); ++StateIndex)
 		{
-			DebugString.Appendf(TEXT("  | %-30s | %15s | %5s [%3d:%-3d[ | %9s   %4d %4d %4d %4d | %3s   %4d %4d %4d %4d |\n"),
-				*State.Name.ToString(), *State.Parent.Describe(),
-				TEXT(""), State.ChildrenBegin, State.ChildrenEnd,
-				TEXT(""), State.EnterConditionsBegin, State.TransitionsBegin, State.TasksBegin, State.UtilityConsiderationsBegin,
-				TEXT(""), State.EnterConditionsNum, State.TransitionsNum, State.TasksNum, State.UtilityConsiderationsNum);
+			const FCompactStateTreeState& State = States[StateIndex];
+			DebugString.Appendf(TEXT("  | (%3d) | %-30s | %15s | %5s [%3d:%-3d[ | %9s   %4d %4d %4d %4d | %3s   %4d %4d %4d %4d |\n"),
+				StateIndex,
+				*State.Name.ToString(),
+				*State.Parent.Describe(),
+				TEXT(" "), State.ChildrenBegin, State.ChildrenEnd,
+				TEXT(" "), State.EnterConditionsBegin, State.TransitionsBegin, State.TasksBegin, State.UtilityConsiderationsBegin,
+				TEXT(" "), State.EnterConditionsNum, State.TransitionsNum, State.TasksNum, State.UtilityConsiderationsNum);
 		}
 	}
 
@@ -1365,16 +1377,20 @@ FString UStateTree::DebugInternalLayoutAsString() const
 	DebugString.Appendf(TEXT("\nTransitions(%d)\n"), Transitions.Num());
 	if (Transitions.Num())
 	{
-		DebugString.Appendf(TEXT("  [ %-3s | %15s | %-20s | %-40s | %-40s | %-8s ]\n")
+		DebugString.Appendf(TEXT("  [ (Idx) | %-3s | %15s | %-20s | %-40s | %-40s | %-8s ]\n")
 			, TEXT("Idx"), TEXT("State"), TEXT("Transition Trigger"), TEXT("Transition Event Tag"), TEXT("Transition Event Payload"), TEXT("Cond:Num"));
-		for (const FCompactStateTransition& Transition : Transitions)
+		for (int32 TransitionIndex = 0; TransitionIndex < Transitions.Num(); ++TransitionIndex)
 		{
-			DebugString.Appendf(TEXT("  | %3d | %15s | %-20s | %-40s | %-40s | %4d:%3d |\n"),
-				Transition.ConditionsBegin, *Transition.State.Describe(),
+			const FCompactStateTransition& Transition = Transitions[TransitionIndex];
+			DebugString.Appendf(TEXT("  | (%3d) | %3d | %15s | %-20s | %-40s | %-40s | %4d:%3d |\n"),
+				TransitionIndex,
+				Transition.ConditionsBegin,
+				*Transition.State.Describe(),
 				*UEnum::GetDisplayValueAsText(Transition.Trigger).ToString(),
 				*Transition.RequiredEvent.Tag.ToString(),
 				Transition.RequiredEvent.PayloadStruct ? *Transition.RequiredEvent.PayloadStruct->GetName() : TEXT("None"),
-				Transition.ConditionsBegin, Transition.ConditionsNum);
+				Transition.ConditionsBegin,
+				Transition.ConditionsNum);
 		}
 	}
 
@@ -1382,18 +1398,21 @@ FString UStateTree::DebugInternalLayoutAsString() const
 	DebugString.Appendf(TEXT("\nEvaluators(%d)\n"), EvaluatorsNum);
 	if (EvaluatorsNum)
 	{
-		DebugString.Appendf(TEXT("  [ %-30s | %8s | %10s ]\n"),
+		DebugString.Appendf(TEXT("  [ (Idx) | %-30s | %8s | %10s ]\n"),
 			TEXT("Name"), TEXT("Bindings"), TEXT("Struct Idx"));
 		for (int32 EvalIndex = EvaluatorsBegin; EvalIndex < (EvaluatorsBegin + EvaluatorsNum); EvalIndex++)
 		{
 			const FStateTreeEvaluatorBase& Eval = Nodes[EvalIndex].Get<const FStateTreeEvaluatorBase>();
-			DebugString.Appendf(TEXT("  | %-30s | %8d | %10s |\n"),
-				*Eval.Name.ToString(), Eval.BindingsBatch.Get(), *Eval.InstanceDataHandle.Describe());
+			DebugString.Appendf(TEXT("  | (%3d) | %-30s | %8d | %10s |\n"),
+				EvalIndex,
+				*Eval.Name.ToString(),
+				Eval.BindingsBatch.Get(),
+				*Eval.InstanceDataHandle.Describe());
 		}
 	}
 
 	// Tasks
-	DebugString.Appendf(TEXT("\nTasks\n  [ %-30s | %-30s | %8s | %10s ]\n"),
+	DebugString.Appendf(TEXT("\nTasks\n  [ (Idx) | %-30s | %-30s | %8s | %10s ]\n"),
 		TEXT("State"), TEXT("Name"), TEXT("Bindings"), TEXT("Struct Idx"));
 	for (const FCompactStateTreeState& State : States)
 	{
@@ -1402,30 +1421,40 @@ FString UStateTree::DebugInternalLayoutAsString() const
 			for (int32 TaskIndex = State.TasksBegin; TaskIndex < (State.TasksBegin + State.TasksNum); TaskIndex++)
 			{
 				const FStateTreeTaskBase& Task = Nodes[TaskIndex].Get<const FStateTreeTaskBase>();
-				DebugString.Appendf(TEXT("  | %-30s | %-30s | %8d | %10s |\n"),
-					*State.Name.ToString(), *Task.Name.ToString(), Task.BindingsBatch.Get(), *Task.InstanceDataHandle.Describe());
+				DebugString.Appendf(TEXT("  | (%3d) | %-30s | %-30s | %8d | %10s |\n"),
+					TaskIndex,
+					*State.Name.ToString(),
+					*Task.Name.ToString(),
+					Task.BindingsBatch.Get(),
+					*Task.InstanceDataHandle.Describe());
 			}
 		}
 	}
 	for (int32 TaskIndex = GlobalTasksBegin; TaskIndex < (GlobalTasksBegin + GlobalTasksNum); TaskIndex++)
 	{
 		const FStateTreeTaskBase& Task = Nodes[TaskIndex].Get<const FStateTreeTaskBase>();
-		DebugString.Appendf(TEXT("  | %-30s | %-30s | %8d | %10s |\n"),
-			TEXT("Global"), *Task.Name.ToString(), Task.BindingsBatch.Get(), *Task.InstanceDataHandle.Describe());
+		DebugString.Appendf(TEXT("  | (%3d) | %-30s | %-30s | %8d | %10s |\n"),
+			TaskIndex,
+			TEXT("Global"),
+			*Task.Name.ToString(),
+			Task.BindingsBatch.Get(),
+			*Task.InstanceDataHandle.Describe());
 	}
 
 	// Conditions
-	DebugString.Appendf(TEXT("\nCondition\n  [ %-30s | %-30s | %8s | %12s | %10s ]\n"),
-		TEXT("State"), TEXT("Name"), TEXT("Operand"), TEXT("Evaluation"), TEXT("Struct Idx"));
-	for (const FCompactStateTreeState& State : States)
+	DebugString.Appendf(TEXT("\nConditions\n  [ (Idx) | %-30s | %8s | %12s | %10s ]\n"),
+		TEXT("Name"), TEXT("Operand"), TEXT("Evaluation"), TEXT("Struct Idx"));
 	{
-		if (State.EnterConditionsNum)
+		for (int32 Index = 0; Index < Nodes.Num(); ++Index)
 		{
-			for (int32 CondIndex = State.EnterConditionsBegin; CondIndex < (State.EnterConditionsBegin + State.EnterConditionsNum); CondIndex++)
+			if(const FStateTreeConditionBase* Cond = Nodes[Index].GetPtr<const FStateTreeConditionBase>())
 			{
-				const FStateTreeConditionBase& Cond = Nodes[CondIndex].Get<const FStateTreeConditionBase>();
-				DebugString.Appendf(TEXT("  | %-30s | %-30s | %8s | %12s | %10s |\n"),
-					*State.Name.ToString(), *Cond.Name.ToString(), *UEnum::GetDisplayValueAsText(Cond.Operand).ToString(), *UEnum::GetDisplayValueAsText(Cond.EvaluationMode).ToString(), *Cond.InstanceDataHandle.Describe());
+				DebugString.Appendf(TEXT("  | (%3d) | %-30s | %8s | %12s | %10s |\n"),
+					Index,
+					*Cond->Name.ToString(),
+					*UEnum::GetDisplayValueAsText(Cond->Operand).ToString(),
+					*UEnum::GetDisplayValueAsText(Cond->EvaluationMode).ToString(),
+					*Cond->InstanceDataHandle.Describe());
 			}
 		}
 	}
