StateTree: Disable binding for delegate dispatcher. A delegate listener can only bind to a dispatcher. Can't promote delegate and listener to a parameter. Expands the nobinding meta data to struct and class definition.

diff --git a/Engine/Plugins/Runtime/PropertyBindingUtils/Source/PropertyBindingUtilsEditor/Private/PropertyBindingExtension.cpp b/Engine/Plugins/Runtime/PropertyBindingUtils/Source/PropertyBindingUtilsEditor/Private/PropertyBindingExtension.cpp
--- a/Engine/Plugins/Runtime/PropertyBindingUtils/Source/PropertyBindingUtilsEditor/Private/PropertyBindingExtension.cpp
+++ b/Engine/Plugins/Runtime/PropertyBindingUtils/Source/PropertyBindingUtilsEditor/Private/PropertyBindingExtension.cpp
@@ -1095,15 +1095,30 @@ bool FPropertyBindingExtension::IsPropertyExtendable(const UClass* InObjectClass
 		return false;
 	}
 
-	bool bCanBind = true;
-
+	// Does the container or property support bindings
 	TSharedPtr<const IPropertyHandle> CurrentPropertyHandle = InPropertyHandle.AsShared();
 	while (CurrentPropertyHandle.IsValid())
 	{
+		if (const FProperty* MetaDataProperty = CurrentPropertyHandle->GetMetaDataProperty())
+		{
+			if (const FStructProperty* StructProperty = CastField<const FStructProperty>(MetaDataProperty))
+			{
+				if (StructProperty->Struct->HasMetaData(UE::PropertyBinding::MetaDataNoBindingName))
+				{
+					return false;
+				}
+			}
+			else if (const FObjectPropertyBase* ObjectProperty = CastField<const FObjectPropertyBase>(MetaDataProperty))
+			{
+				if (ObjectProperty->PropertyClass->HasMetaData(UE::PropertyBinding::MetaDataNoBindingName))
+				{
+					return false;
+				}
+			}
+		}
 		if (CurrentPropertyHandle->HasMetaData(UE::PropertyBinding::MetaDataNoBindingName))
 		{
-			bCanBind = false;
-			break;
+			return false;
 		}
 		CurrentPropertyHandle = CurrentPropertyHandle->GetParentHandle();
 	}
@@ -1124,15 +1139,18 @@ bool FPropertyBindingExtension::IsPropertyExtendable(const UClass* InObjectClass
 
 			if (!TargetPath.GetStructID().IsValid())
 			{
-				bCanBind = false;
+				return false;
 			}
 			else if (const IPropertyHandle* BindablePropertyHandle = SharedHandle.Get())
 			{
-				bCanBind = CanBindToProperty(TargetPath, *BindablePropertyHandle);
+				if (!CanBindToProperty(TargetPath, *BindablePropertyHandle))
+				{
+					return false;
+				}
 			}
 		}
 	}
-	return bCanBind;
+	return true;
 }
 
 void FPropertyBindingExtension::ExtendWidgetRow(FDetailWidgetRow& InWidgetRow, const IDetailLayoutBuilder& InDetailBuilder, const UClass* InObjectClass, TSharedPtr<IPropertyHandle> InPropertyHandle)
diff --git a/Engine/Plugins/Runtime/PropertyBindingUtils/Source/PropertyBindingUtilsEditor/Public/PropertyBindingExtension.h b/Engine/Plugins/Runtime/PropertyBindingUtils/Source/PropertyBindingUtilsEditor/Public/PropertyBindingExtension.h
--- a/Engine/Plugins/Runtime/PropertyBindingUtils/Source/PropertyBindingUtilsEditor/Public/PropertyBindingExtension.h
+++ b/Engine/Plugins/Runtime/PropertyBindingUtils/Source/PropertyBindingUtilsEditor/Public/PropertyBindingExtension.h
@@ -32,6 +32,16 @@ struct FPropertyInfoOverride
 	FEdGraphPinType PinType;
 };
 
+namespace Meta::Private
+{
+// valid keywords for the UCLASS UPROPERTY and USTRUCT macro
+enum
+{
+	// [ClassMetadata] [PropertyMetadata] [StructMetadata] The property cannot be bound to (target of a binding).
+	NoBinding,
+};
+} // namespace Meta::Private
+
 const FLazyName MetaDataStructIDName("StructIDForPropertyBinding");
 const FLazyName MetaDataNoBindingName("NoBinding");
 
diff --git a/Engine/Plugins/Runtime/StateTree/Source/StateTreeEditorModule/Private/Customizations/StateTreeBindingExtension.cpp b/Engine/Plugins/Runtime/StateTree/Source/StateTreeEditorModule/Private/Customizations/StateTreeBindingExtension.cpp
--- a/Engine/Plugins/Runtime/StateTree/Source/StateTreeEditorModule/Private/Customizations/StateTreeBindingExtension.cpp
+++ b/Engine/Plugins/Runtime/StateTree/Source/StateTreeEditorModule/Private/Customizations/StateTreeBindingExtension.cpp
@@ -58,9 +58,19 @@ UObject* FindEditorBindingsOwner(UObject* InObject)
 }
 
 
-bool IsDelegateListenerProperty(const FProperty& Property)
+static bool IsDelegateDispatcherProperty(const FProperty* Property)
 {
-	if (const FStructProperty* StructProperty = CastField<FStructProperty>(&Property))
+	if (const FStructProperty* StructProperty = CastField<FStructProperty>(Property))
+	{
+		return StructProperty->Struct == FStateTreeDelegateDispatcher::StaticStruct();
+	}
+
+	return false;
+}
+
+static bool IsDelegateListenerProperty(const FProperty* Property)
+{
+	if (const FStructProperty* StructProperty = CastField<FStructProperty>(Property))
 	{
 		return StructProperty->Struct == FStateTreeDelegateListener::StaticStruct() || StructProperty->Struct == FStateTreeTransitionDelegateListener::StaticStruct();
 	}
@@ -179,7 +189,7 @@ struct FStateTreeCachedBindingData : UE::PropertyBinding::FCachedBindingData, TS
 			return true;
 		}
 
-		if (IsDelegateListenerProperty(InProperty))
+		if (IsDelegateListenerProperty(&InProperty))
 		{
 			OutPinType.PinCategory = UEdGraphSchema_K2::PC_Delegate;
 			OutIconName = TEXT("Icons.Event");
@@ -334,6 +344,12 @@ struct FStateTreeCachedBindingData : UE::PropertyBinding::FCachedBindingData, TS
 			}
 		}
 
+		// Listener can only bind to dispatcher (prevents listener to listener)
+		if (IsDelegateListenerProperty(GetPropertyHandle()->GetProperty()))
+		{
+			return IsDelegateDispatcherProperty(&SourceProperty);
+		}
+
 		return true;
 	}
 
@@ -402,7 +418,7 @@ struct FStateTreeCachedBindingData : UE::PropertyBinding::FCachedBindingData, TS
 			bOutAreCompatible = PropertyRefHelpers::IsPropertyRefCompatibleWithProperty(*TargetStructProperty, *InSourceProperty, InTargetPropertyValue, InSourcePropertyValue);
 			return true;
 		}
-		else if (TargetStructProperty && IsDelegateListenerProperty(*TargetStructProperty))
+		else if (TargetStructProperty && IsDelegateListenerProperty(TargetStructProperty))
 		{
 			if (const FStructProperty* SourceStructProperty = CastField<FStructProperty>(InSourceProperty))
 			{
@@ -602,6 +618,7 @@ bool FStateTreeBindingExtension::GetPromotionToParameterOverrideInternal(const F
 		if (StructProperty->Struct && StructProperty->Struct->IsChildOf(FStateTreePropertyRef::StaticStruct()))
 		{
 			bOutOverride = false;
+			return true;
 		}
 	}
 	return false;
diff --git a/Engine/Plugins/Runtime/StateTree/Source/StateTreeModule/Public/StateTreeDelegate.h b/Engine/Plugins/Runtime/StateTree/Source/StateTreeModule/Public/StateTreeDelegate.h
--- a/Engine/Plugins/Runtime/StateTree/Source/StateTreeModule/Public/StateTreeDelegate.h
+++ b/Engine/Plugins/Runtime/StateTree/Source/StateTreeModule/Public/StateTreeDelegate.h
@@ -24,7 +24,7 @@
  /**
   * StateTree's delegate dispatcher.
   */
-USTRUCT(BlueprintType)
+USTRUCT(BlueprintType, meta=(NoBinding))
 struct STATETREEMODULE_API FStateTreeDelegateDispatcher
 {
 	GENERATED_BODY()
