StateTree: Add the DisplayItemIds to the developer menu. Update WeakContext description.

#rnx

diff --git a/Engine/Plugins/Runtime/StateTree/Source/StateTreeEditorModule/Private/StateTreeEditor.cpp b/Engine/Plugins/Runtime/StateTree/Source/StateTreeEditorModule/Private/StateTreeEditor.cpp
--- a/Engine/Plugins/Runtime/StateTree/Source/StateTreeEditorModule/Private/StateTreeEditor.cpp
+++ b/Engine/Plugins/Runtime/StateTree/Source/StateTreeEditorModule/Private/StateTreeEditor.cpp
@@ -360,6 +360,27 @@ void FillDeveloperMenu(UToolMenu* InMenu)
 		FToolMenuSection& Section = InMenu->AddSection("FileDeveloperCompilerSettings", LOCTEXT("CompileOptionsHeading", "Compiler Settings"));
 		Section.AddMenuEntry(Commands.LogCompilationResult);
 	}
+	{
+		FToolMenuSection& Section = InMenu->AddSection("FileDeveloperSettings", LOCTEXT("DeveloperOptionsHeading", "Settings"));
+		Section.AddEntry(FToolMenuEntry::InitMenuEntry(
+			"DisplayItemIds",
+			LOCTEXT("DisplayItemIds", "Display Nodes IDs"),
+			UE::StateTree::Editor::CVarDisplayItemIds->GetDetailedHelp(),
+			TAttribute<FSlateIcon>(),
+			FUIAction(
+				FExecuteAction::CreateLambda([]()
+					{
+						UE::StateTree::Editor::CVarDisplayItemIds->Set(!UE::StateTree::Editor::CVarDisplayItemIds->GetBool(), ECVF_SetByConsole);
+					}),
+				FCanExecuteAction(),
+				FIsActionChecked::CreateLambda([]()
+					{
+						return UE::StateTree::Editor::CVarDisplayItemIds->GetBool();
+					})
+				),
+			EUserInterfaceActionType::ToggleButton
+			));
+	}
 }
 void FillDynamicDeveloperMenu(FToolMenuSection& Section)
 {
diff --git a/Engine/Plugins/Runtime/StateTree/Source/StateTreeModule/Public/StateTreeExecutionContext.h b/Engine/Plugins/Runtime/StateTree/Source/StateTreeModule/Public/StateTreeExecutionContext.h
--- a/Engine/Plugins/Runtime/StateTree/Source/StateTreeModule/Public/StateTreeExecutionContext.h
+++ b/Engine/Plugins/Runtime/StateTree/Source/StateTreeModule/Public/StateTreeExecutionContext.h
@@ -1150,9 +1150,9 @@ private:
  * You are responsible to make SendEvent threadsafe if needed.
  *
  *		ThreadSafeAsyncCallback.AddLambda(
- *			[MyTag, MiminalExecutionContext = Context.MakeWeakExecutionContext]()
+ *			[MyTag, WeakContext = Context.MakeWeakExecutionContext]()
  *			{
- *				if (MiminalExecutionContext.SendEvent(MyTag))
+ *				if (WeakContext.SendEvent(MyTag))
  *				{
  *					...
  *				}
