Fix build. Bad deprecation warning.
#rnx

diff --git a/Engine/Plugins/Runtime/StateTree/Source/StateTreeModule/Private/Conditions/StateTreeObjectConditions.cpp b/Engine/Plugins/Runtime/StateTree/Source/StateTreeModule/Private/Conditions/StateTreeObjectConditions.cpp
--- a/Engine/Plugins/Runtime/StateTree/Source/StateTreeModule/Private/Conditions/StateTreeObjectConditions.cpp
+++ b/Engine/Plugins/Runtime/StateTree/Source/StateTreeModule/Private/Conditions/StateTreeObjectConditions.cpp
@@ -72,8 +72,8 @@ FText FStateTreeObjectIsChildOfClassCondition::GetDescription(const FGuid& ID, F
 {
 	const FText InvertText = UE::StateTree::DescHelpers::GetInvertText(bInvert, Formatting);
 	const FText Format = (Formatting == EStateTreeNodeFormatting::RichText)
-		? LOCTEXT("ObjectIsValidConditionsRich", "{EmptyOrNot}<s>Is Child Of Class</>")
-		: LOCTEXT("ObjectIsValidCondition", "{EmptyOrNot}Is Child Of Class");
+		? LOCTEXT("ObjectIsChildOfConditionsRich", "{EmptyOrNot}<s>Is Child Of Class</>")
+		: LOCTEXT("ObjectIsChildOfCondition", "{EmptyOrNot}Is Child Of Class");
 	return FText::FormatNamed(Format,
 		TEXT("EmptyOrNot"), InvertText);
 }
