StateTree: Fix bad selection code introduced with CL 39000111.

diff --git a/Engine/Plugins/Runtime/StateTree/Source/StateTreeEditorModule/Private/StateTreeEditorMode.cpp b/Engine/Plugins/Runtime/StateTree/Source/StateTreeEditorModule/Private/StateTreeEditorMode.cpp
--- a/Engine/Plugins/Runtime/StateTree/Source/StateTreeEditorModule/Private/StateTreeEditorMode.cpp
+++ b/Engine/Plugins/Runtime/StateTree/Source/StateTreeEditorModule/Private/StateTreeEditorMode.cpp
@@ -464,32 +464,18 @@ void UStateTreeEditorMode::HandleModelAssetChanged()
 
 void UStateTreeEditorMode::HandleModelSelectionChanged(const TArray<TWeakObjectPtr<UStateTreeState>>& InSelectedStates) const
 {
-UE_LOG(LogTemp, Warning, TEXT("UStateTreeEditorMode::HandleModelSelectionChanged"));
 	if (TSharedPtr<IDetailsView> DetailsView = GetDetailsView())
 	{
-		//// Delay the selection by one frame to handle focus lost.
-		//constexpr bool bLoop = false;
-		//constexpr float DelayedTime = UE_SMALL_NUMBER;
-		//GEditor->GetTimerManager()->SetTimer(
-		//	SetObjectTimerHandle,
-		//	FTimerDelegate::CreateLambda([WeakDetailsView = DetailsView.ToWeakPtr(), InSelectedStates]()
-		//		{
-		//			if (TSharedPtr<IDetailsView> DetailsView = WeakDetailsView.Pin())
-		//			{
-						TArray<UObject*> Selected;
-						Selected.Reserve(InSelectedStates.Num());
-						for (const TWeakObjectPtr<UStateTreeState>& WeakState : InSelectedStates)
-						{
-							if (UStateTreeState* State = WeakState.Get())
-							{
-								Selected.Add(State);
-							}
-						}
-						DetailsView->SetObjects(Selected);
-				//	}
-				//}),
-				//DelayedTime,
-				//bLoop);
+		TArray<UObject*> Selected;
+		Selected.Reserve(InSelectedStates.Num());
+		for (const TWeakObjectPtr<UStateTreeState>& WeakState : InSelectedStates)
+		{
+			if (UStateTreeState* State = WeakState.Get())
+			{
+				Selected.Add(State);
+			}
+		}
+		DetailsView->SetObjects(Selected);
 	}
 }
 
