StateTree: Serialize the compiled listener ID.
#rnx

diff --git a/Engine/Plugins/Runtime/StateTree/Source/StateTreeModule/Public/StateTreeDelegate.h b/Engine/Plugins/Runtime/StateTree/Source/StateTreeModule/Public/StateTreeDelegate.h
--- a/Engine/Plugins/Runtime/StateTree/Source/StateTreeModule/Public/StateTreeDelegate.h
+++ b/Engine/Plugins/Runtime/StateTree/Source/StateTreeModule/Public/StateTreeDelegate.h
@@ -78,6 +78,7 @@ private:
 	FStateTreeDelegateDispatcher Dispatcher;
 
 	/** The generated ID of the listener. */
+	UPROPERTY()
 	int32 ID = 0;
 
 	friend struct FStateTreePropertyBindingCompiler;
