StateTree: FrameId to FrameID and StateId to StateID

#rnx

diff --git a/Engine/Plugins/Runtime/StateTree/Source/StateTreeModule/Private/StateTreeExecutionContext.cpp b/Engine/Plugins/Runtime/StateTree/Source/StateTreeModule/Private/StateTreeExecutionContext.cpp
--- a/Engine/Plugins/Runtime/StateTree/Source/StateTreeModule/Private/StateTreeExecutionContext.cpp
+++ b/Engine/Plugins/Runtime/StateTree/Source/StateTreeModule/Private/StateTreeExecutionContext.cpp
@@ -668,7 +668,7 @@ EStateTreeRunStatus FStateTreeExecutionContext::Start(FStartParameters Parameter
 	// Initialize for the init frame.
 	ensure(Exec.ActiveFrames.Num() == 0);
 	FStateTreeExecutionFrame& InitFrame = Exec.ActiveFrames.AddDefaulted_GetRef();
-	InitFrame.FrameId = UE::StateTree::FActiveFrameId(InstanceDataStorage->GenerateUniqueId());
+	InitFrame.FrameID = UE::StateTree::FActiveFrameID(InstanceDataStorage->GenerateUniqueId());
 	InitFrame.StateTree = &RootStateTree;
 	InitFrame.RootState = FStateTreeStateHandle::Root;
 	InitFrame.ActiveStates = {};
@@ -1448,7 +1448,7 @@ void FStateTreeExecutionContext::UpdateInstanceData(TConstArrayView<FStateTreeEx
 	{
 		FStateTreeTemporaryInstanceData* TempData = TempInstances.FindByPredicate([&Frame, &DataHandle](const FStateTreeTemporaryInstanceData& Data)
 		{
-			return Data.FrameId == Frame.FrameId && Data.DataHandle == DataHandle;
+			return Data.FrameID == Frame.FrameID && Data.DataHandle == DataHandle;
 		});
 		return TempData ? &TempData->Instance : nullptr;
 	};
@@ -2180,13 +2180,13 @@ void FStateTreeExecutionContext::CaptureNewStateEvents(TConstArrayView<FStateTre
 		{
 			const FStateTreeExecutionFrame& PrevFrame = PrevFrames[FrameIndex];
 
-			if (PrevFrame.FrameId == NewFrame.FrameId)
+			if (PrevFrame.FrameID == NewFrame.FrameID)
 			{
 				checkf(PrevFrame.RootState == NewFrame.RootState && PrevFrame.StateTree == NewFrame.StateTree, TEXT("If the Id matches, then the root and the tree must also match."));
 
 				for (int32 StateIndex = 0; StateIndex < NewFrame.ActiveStates.Num(); ++StateIndex)
 				{
-					if (!PrevFrame.ActiveStates.IsValidIndex(StateIndex) || PrevFrame.ActiveStates.StateIds[StateIndex] != NewFrame.ActiveStates.StateIds[StateIndex])
+					if (!PrevFrame.ActiveStates.IsValidIndex(StateIndex) || PrevFrame.ActiveStates.StateIDs[StateIndex] != NewFrame.ActiveStates.StateIDs[StateIndex])
 					{
 						UniqueStates = TConstArrayView<FStateTreeStateHandle>(&NewFrame.ActiveStates[StateIndex], NewFrame.ActiveStates.Num() - StateIndex);
 						break;
@@ -2297,7 +2297,7 @@ EStateTreeRunStatus FStateTreeExecutionContext::EnterState(FStateTreeTransitionR
 			const FCompactStateTreeState& State = CurrentStateTree->States[CurrentHandle.Index];
 
 			// Add only enabled States to the list of active States
-			if (State.bEnabled && !CurrentFrame.ActiveStates.Push(CurrentHandle, NextFrame.ActiveStates.StateIds[StateIndex]))
+			if (State.bEnabled && !CurrentFrame.ActiveStates.Push(CurrentHandle, NextFrame.ActiveStates.StateIDs[StateIndex]))
 			{
 				STATETREE_LOG(Error, TEXT("%hs: Reached max execution depth when trying to enter state '%s'.  '%s' using StateTree '%s'."),
 					__FUNCTION__, *GetStateStatusString(Exec), *GetNameSafe(&Owner), *GetFullNameSafe(&RootStateTree));
@@ -3332,7 +3332,7 @@ void FStateTreeExecutionContext::StopTemporaryEvaluatorsAndGlobalTasks(const FSt
 	for (int32 Index = TempInstances.Num() - 1; Index >= 0; Index--)
 	{
 		FStateTreeTemporaryInstanceData& TempInstance = TempInstances[Index];
-		if (TempInstance.FrameId != CurrentFrame.FrameId)
+		if (TempInstance.FrameID != CurrentFrame.FrameID)
 		{
 			continue;
 		}
@@ -3880,7 +3880,7 @@ void FStateTreeExecutionContext::SetupNextTransition(const FStateTreeExecutionFr
 	}
 	else
 	{
-		NewFrame.ActiveStates = FStateTreeActiveStates(NextState, UE::StateTree::FActiveStateId::Invalid);
+		NewFrame.ActiveStates = FStateTreeActiveStates(NextState, UE::StateTree::FActiveStateID::Invalid);
 	}
 }
 
@@ -4238,7 +4238,7 @@ bool FStateTreeExecutionContext::TriggerTransitions()
 									// Delay state added, skip requesting the transition.
 									continue;
 								}
-								// Fallthrough to request transition if duration was zero. 
+								// Fallthrough to request transition if duration was zero.
 							}
 							else
 							{
@@ -4768,7 +4768,7 @@ bool FStateTreeExecutionContext::SelectStateInternal(
 	}
 
 	// Activate/Push a new state
-	if (!CurrentFrame.ActiveStates.Push(NextStateHandle, UE::StateTree::FActiveStateId(InstanceDataStorage->GenerateUniqueId())))
+	if (!CurrentFrame.ActiveStates.Push(NextStateHandle, UE::StateTree::FActiveStateID(InstanceDataStorage->GenerateUniqueId())))
 	{
 		STATETREE_LOG(Error, TEXT("%hs: Reached max execution depth when trying to select state %s from '%s'.  '%s' using StateTree '%s'."),
 			__FUNCTION__, *GetSafeStateName(CurrentFrame, NextStateHandle), *GetStateStatusString(Exec), *GetNameSafe(&Owner), *GetFullNameSafe(CurrentFrame.StateTree));
@@ -4834,7 +4834,7 @@ bool FStateTreeExecutionContext::SelectStateInternal(
 				FStateTreeExecutionFrame NewFrame;
 				NewFrame.StateTree = CurrentFrame.StateTree;
 				NewFrame.RootState = NextState.LinkedState;
-				NewFrame.FrameId = UE::StateTree::FActiveFrameId(InstanceDataStorage->GenerateUniqueId());
+				NewFrame.FrameID = UE::StateTree::FActiveFrameID(InstanceDataStorage->GenerateUniqueId());
 				NewFrame.ExternalDataBaseIndex = CurrentFrame.ExternalDataBaseIndex;
 
 				// Check and prevent recursion.
@@ -4857,7 +4857,7 @@ bool FStateTreeExecutionContext::SelectStateInternal(
 					});
 				if (ExistingFrame)
 				{
-					// Note. Do not copy the FrameId. This is a new frame and the enter/exit state should be executed.
+					// Note. Do not copy the FrameID. This is a new frame and the enter/exit state should be executed.
 					NewFrame.ActiveInstanceIndexBase = ExistingFrame->ActiveInstanceIndexBase;
 					NewFrame.GlobalInstanceIndexBase = ExistingFrame->GlobalInstanceIndexBase;
 					NewFrame.StateParameterDataHandle = ExistingFrame->StateParameterDataHandle;
@@ -4911,7 +4911,7 @@ bool FStateTreeExecutionContext::SelectStateInternal(
 			}
 				
 			FStateTreeExecutionFrame NewFrame;
-			NewFrame.FrameId = UE::StateTree::FActiveFrameId(InstanceDataStorage->GenerateUniqueId());
+			NewFrame.FrameID = UE::StateTree::FActiveFrameID(InstanceDataStorage->GenerateUniqueId());
 			NewFrame.StateTree = NextLinkedStateAsset;
 			NewFrame.RootState = FStateTreeStateHandle::Root;
 			NewFrame.bIsGlobalFrame = true;
@@ -4938,7 +4938,7 @@ bool FStateTreeExecutionContext::SelectStateInternal(
 			bool bStartedTemporaryEvaluatorsAndGlobalTasks = false;
 			if (ExistingFrame)
 			{
-				// Note. Do not copy the FrameId. This is a new frame and the enter/exit state should be executed.
+				// Note. Do not copy the FrameID. This is a new frame and the enter/exit state should be executed.
 				NewFrame.ActiveInstanceIndexBase = ExistingFrame->ActiveInstanceIndexBase;
 				NewFrame.GlobalInstanceIndexBase = ExistingFrame->GlobalInstanceIndexBase;
 				NewFrame.StateParameterDataHandle = ExistingFrame->StateParameterDataHandle;
diff --git a/Engine/Plugins/Runtime/StateTree/Source/StateTreeModule/Private/StateTreeExecutionTypes.cpp b/Engine/Plugins/Runtime/StateTree/Source/StateTreeModule/Private/StateTreeExecutionTypes.cpp
--- a/Engine/Plugins/Runtime/StateTree/Source/StateTreeModule/Private/StateTreeExecutionTypes.cpp
+++ b/Engine/Plugins/Runtime/StateTree/Source/StateTreeModule/Private/StateTreeExecutionTypes.cpp
@@ -115,7 +115,7 @@ UE::StateTree::FActiveStatePath FStateTreeExecutionState::GetActiveStatePath() c
 	{
 		for (int32 StateIndex = 0; StateIndex < Frame.ActiveStates.Num(); ++StateIndex)
 		{
-			Elements.Emplace(Frame.ActiveStates.States[StateIndex], Frame.ActiveStates.StateIds[StateIndex]);
+			Elements.Emplace(Frame.ActiveStates.States[StateIndex], Frame.ActiveStates.StateIDs[StateIndex]);
 		}
 	}
 
diff --git a/Engine/Plugins/Runtime/StateTree/Source/StateTreeModule/Private/StateTreeInstanceData.cpp b/Engine/Plugins/Runtime/StateTree/Source/StateTreeModule/Private/StateTreeInstanceData.cpp
--- a/Engine/Plugins/Runtime/StateTree/Source/StateTreeModule/Private/StateTreeInstanceData.cpp
+++ b/Engine/Plugins/Runtime/StateTree/Source/StateTreeModule/Private/StateTreeInstanceData.cpp
@@ -256,7 +256,7 @@ FStructView FStateTreeInstanceStorage::AddTemporaryInstance(UObject& InOwner, co
 {
 	FStateTreeTemporaryInstanceData* TempInstance = TemporaryInstances.FindByPredicate([&Frame, &OwnerNodeIndex, &DataHandle](const FStateTreeTemporaryInstanceData& TempInstance)
 	{
-		return TempInstance.FrameId == Frame.FrameId
+		return TempInstance.FrameID == Frame.FrameID
 				&& TempInstance.OwnerNodeIndex == OwnerNodeIndex
 				&& TempInstance.DataHandle == DataHandle;
 	});
@@ -272,7 +272,7 @@ FStructView FStateTreeInstanceStorage::AddTemporaryInstance(UObject& InOwner, co
 	{
 		TempInstance = &TemporaryInstances.AddDefaulted_GetRef();
 		check(TempInstance);
-		TempInstance->FrameId = Frame.FrameId;
+		TempInstance->FrameID = Frame.FrameID;
 		TempInstance->OwnerNodeIndex = OwnerNodeIndex;
 		TempInstance->DataHandle = DataHandle;
 		TempInstance->Instance = NewInstanceData;
@@ -293,7 +293,7 @@ FStructView FStateTreeInstanceStorage::GetMutableTemporaryStruct(const FStateTre
 {
 	FStateTreeTemporaryInstanceData* ExistingInstance = TemporaryInstances.FindByPredicate([&Frame, &DataHandle](const FStateTreeTemporaryInstanceData& TempInstance)
 	{
-		return TempInstance.FrameId == Frame.FrameId
+		return TempInstance.FrameID == Frame.FrameID
 				&& TempInstance.DataHandle == DataHandle;
 	});
 	return ExistingInstance ? FStructView(ExistingInstance->Instance) : FStructView();
@@ -303,7 +303,7 @@ UObject* FStateTreeInstanceStorage::GetMutableTemporaryObject(const FStateTreeEx
 {
 	FStateTreeTemporaryInstanceData* ExistingInstance = TemporaryInstances.FindByPredicate([&Frame, &DataHandle](const FStateTreeTemporaryInstanceData& TempInstance)
 	{
-		return TempInstance.FrameId == Frame.FrameId
+		return TempInstance.FrameID == Frame.FrameID
 				&& TempInstance.DataHandle == DataHandle;
 	});
 	if (ExistingInstance)
diff --git a/Engine/Plugins/Runtime/StateTree/Source/StateTreeModule/Private/StateTreeStatePath.cpp b/Engine/Plugins/Runtime/StateTree/Source/StateTreeModule/Private/StateTreeStatePath.cpp
--- a/Engine/Plugins/Runtime/StateTree/Source/StateTreeModule/Private/StateTreeStatePath.cpp
+++ b/Engine/Plugins/Runtime/StateTree/Source/StateTreeModule/Private/StateTreeStatePath.cpp
@@ -7,8 +7,8 @@
 
 namespace UE::StateTree
 {
-	const FActiveFrameId FActiveFrameId::Invalid = FActiveFrameId();
-	const FActiveStateId FActiveStateId::Invalid = FActiveStateId();
+	const FActiveFrameID FActiveFrameID::Invalid = FActiveFrameID();
+	const FActiveStateID FActiveStateID::Invalid = FActiveStateID();
 
 	FActiveStatePath::FActiveStatePath(TNotNull<const UStateTree*> InStateTree, const TArrayView<const FElement> InElements)
 		: StateTree(InStateTree)
@@ -56,14 +56,14 @@ namespace UE::StateTree
 		return States.Num() > 0 && States.Last() == Other;
 	}
 	
-	bool FActiveStatePath::Matches(FActiveStateId Other) const
+	bool FActiveStatePath::Matches(FActiveStateID Other) const
 	{
-		return States.Num() > 0 && States.Last().GetActiveStateId() == Other;
+		return States.Num() > 0 && States.Last().GetActiveStateID() == Other;
 	}
 
 	const TArrayView<const FActiveStatePath::FElement> FActiveStatePath::Intersect(const FActiveStatePath& Other) const
 	{
-		// If the id is the same, then are equal.
+		// If the ID is the same, then are equal.
 		if (States.Num() == Other.States.Num()
 			&& States.Num() > 0
 			&& States.Last() == Other.States.Last())
@@ -101,13 +101,13 @@ namespace UE::StateTree
 		return false;
 	}
 
-	bool FActiveStatePath::Contains(FActiveStateId Other) const
+	bool FActiveStatePath::Contains(FActiveStateID Other) const
 	{
 		if (Other.IsValid())
 		{
 			for (int32 Index = States.Num() - 1; Index >= 0; --Index)
 			{
-				if (States[Index].GetActiveStateId() == Other)
+				if (States[Index].GetActiveStateID() == Other)
 				{
 					return true;
 				}
diff --git a/Engine/Plugins/Runtime/StateTree/Source/StateTreeModule/Public/StateTreeExecutionTypes.h b/Engine/Plugins/Runtime/StateTree/Source/StateTreeModule/Public/StateTreeExecutionTypes.h
--- a/Engine/Plugins/Runtime/StateTree/Source/StateTreeModule/Public/StateTreeExecutionTypes.h
+++ b/Engine/Plugins/Runtime/StateTree/Source/StateTreeModule/Public/StateTreeExecutionTypes.h
@@ -301,11 +301,11 @@ struct STATETREEMODULE_API FStateTreeActiveStates
 	UE_DEPRECATED(5.6, "Use the constructor with FActiveStateId.")
 	explicit FStateTreeActiveStates(const FStateTreeStateHandle StateHandle)
 	{
-		Push(StateHandle, UE::StateTree::FActiveStateId::Invalid);
+		Push(StateHandle, UE::StateTree::FActiveStateID::Invalid);
 	}
-	explicit FStateTreeActiveStates(const FStateTreeStateHandle StateHandle, UE::StateTree::FActiveStateId StateId)
+	explicit FStateTreeActiveStates(const FStateTreeStateHandle StateHandle, UE::StateTree::FActiveStateID StateID)
 	{
-		Push(StateHandle, StateId);
+		Push(StateHandle, StateID);
 	}
 	
 	/** Resets the active state array to empty. */
@@ -318,11 +318,11 @@ struct STATETREEMODULE_API FStateTreeActiveStates
 	/** Pushes new state at the back of the array and returns true if there was enough space. */
 	bool Push(const FStateTreeStateHandle StateHandle)
 	{
-		return Push(StateHandle, UE::StateTree::FActiveStateId::Invalid);
+		return Push(StateHandle, UE::StateTree::FActiveStateID::Invalid);
 	}
 
 	/** Pushes new state at the back of the array and returns true if there was enough space. */
-	bool Push(const FStateTreeStateHandle StateHandle, UE::StateTree::FActiveStateId StateId)
+	bool Push(const FStateTreeStateHandle StateHandle, UE::StateTree::FActiveStateID StateID)
 	{
 		if ((NumStates + 1) > MaxStates)
 		{
@@ -330,7 +330,7 @@ struct STATETREEMODULE_API FStateTreeActiveStates
 		}
 		
 		States[NumStates] = StateHandle;
-		StateIds[NumStates] = StateId;
+		StateIDs[NumStates] = StateID;
 		++NumStates;
 		
 		return true;
@@ -340,11 +340,11 @@ struct STATETREEMODULE_API FStateTreeActiveStates
 	/** Pushes new state at the front of the array and returns true if there was enough space. */
 	bool PushFront(const FStateTreeStateHandle StateHandle)
 	{
-		return PushFront(StateHandle, UE::StateTree::FActiveStateId::Invalid);
+		return PushFront(StateHandle, UE::StateTree::FActiveStateID::Invalid);
 	}
 
 	/** Pushes new state at the front of the array and returns true if there was enough space. */
-	bool PushFront(const FStateTreeStateHandle StateHandle, UE::StateTree::FActiveStateId StateId)
+	bool PushFront(const FStateTreeStateHandle StateHandle, UE::StateTree::FActiveStateID StateID)
 	{
 		if ((NumStates + 1) > MaxStates)
 		{
@@ -355,10 +355,10 @@ struct STATETREEMODULE_API FStateTreeActiveStates
 		for (int32 Index = (int32)NumStates - 1; Index > 0; Index--)
 		{
 			States[Index] = States[Index - 1];
-			StateIds[Index] = StateIds[Index - 1];
+			StateIDs[Index] = StateIDs[Index - 1];
 		}
 		States[0] = StateHandle;
-		StateIds[0] = StateId;
+		StateIDs[0] = StateID;
 		
 		return true;
 	}
@@ -385,7 +385,7 @@ struct STATETREEMODULE_API FStateTreeActiveStates
 			for (int32 Index = NumStates; Index < NewNum; Index++)
 			{
 				States[Index] = FStateTreeStateHandle::Invalid;
-				StateIds[Index] = UE::StateTree::FActiveStateId::Invalid;
+				StateIDs[Index] = UE::StateTree::FActiveStateID::Invalid;
 			}
 		}
 		NumStates = static_cast<uint8>(NewNum);
@@ -470,7 +470,7 @@ struct STATETREEMODULE_API FStateTreeActiveStates
 	FORCEINLINE const FStateTreeStateHandle* begin() const { return &States[0]; }
 	FORCEINLINE const FStateTreeStateHandle* end  () const { return &States[0] + Num(); }
 
-	UE::StateTree::FActiveStateId StateIds[MaxStates];
+	UE::StateTree::FActiveStateID StateIDs[MaxStates];
 
 	UPROPERTY(VisibleDefaultsOnly, Category = Default)
 	FStateTreeStateHandle States[MaxStates];
@@ -783,7 +783,7 @@ struct STATETREEMODULE_API FStateTreeExecutionFrame
 	FStateTreeActiveStates ActiveStates;
 
 	/** Unique frame id for this frame. Can be used to identify the frame. */
-	UE::StateTree::FActiveFrameId FrameId;
+	UE::StateTree::FActiveFrameID FrameID;
 
 	/** First index of the external data for this frame. */
 	UPROPERTY()
diff --git a/Engine/Plugins/Runtime/StateTree/Source/StateTreeModule/Public/StateTreeInstanceData.h b/Engine/Plugins/Runtime/StateTree/Source/StateTreeModule/Public/StateTreeInstanceData.h
--- a/Engine/Plugins/Runtime/StateTree/Source/StateTreeModule/Public/StateTreeInstanceData.h
+++ b/Engine/Plugins/Runtime/StateTree/Source/StateTreeModule/Public/StateTreeInstanceData.h
@@ -54,7 +54,7 @@ PRAGMA_DISABLE_DEPRECATION_WARNINGS
 	FStateTreeTemporaryInstanceData& operator=(FStateTreeTemporaryInstanceData&& Other) = default;
 PRAGMA_ENABLE_DEPRECATION_WARNINGS
 
-	UE::StateTree::FActiveFrameId FrameId;
+	UE::StateTree::FActiveFrameID FrameID;
 
 	UPROPERTY()
 	FStateTreeDataHandle DataHandle = FStateTreeDataHandle::Invalid;
diff --git a/Engine/Plugins/Runtime/StateTree/Source/StateTreeModule/Public/StateTreeStatePath.h b/Engine/Plugins/Runtime/StateTree/Source/StateTreeModule/Public/StateTreeStatePath.h
--- a/Engine/Plugins/Runtime/StateTree/Source/StateTreeModule/Public/StateTreeStatePath.h
+++ b/Engine/Plugins/Runtime/StateTree/Source/StateTreeModule/Public/StateTreeStatePath.h
@@ -12,24 +12,24 @@ struct FStateTreeExecutionFrame;
 
 namespace UE::StateTree
 {
-	/** A unique id for FStateTreeExecutionFrame. */
-	struct FActiveFrameId
+	/** A unique ID for FStateTreeExecutionFrame. */
+	struct FActiveFrameID
 	{
 	private:
 		static constexpr uint32 InvalidValue = 0;
 
 	public:
-		static const FActiveFrameId Invalid;
+		static const FActiveFrameID Invalid;
 
-		explicit FActiveFrameId() = default;		
-		explicit FActiveFrameId(uint32 NewId)
-			: Value(NewId)
+		explicit FActiveFrameID() = default;
+		explicit FActiveFrameID(uint32 NewID)
+			: Value(NewID)
 		{
-			check(NewId != InvalidValue);
+			check(NewID != InvalidValue);
 		}
 
-		[[nodiscard]] bool operator==(const FActiveFrameId&) const = default;
-		[[nodiscard]] bool operator!=(const FActiveFrameId&) const = default;
+		[[nodiscard]] bool operator==(const FActiveFrameID&) const = default;
+		[[nodiscard]] bool operator!=(const FActiveFrameID&) const = default;
 
 		[[nodiscard]] bool IsValid() const
 		{
@@ -40,24 +40,24 @@ namespace UE::StateTree
 		uint32 Value = InvalidValue;
 	};
 	
-	/** A unique id for FStateTreeActiveStates state. */
-	struct FActiveStateId
+	/** A unique ID for FStateTreeActiveStates state. */
+	struct FActiveStateID
 	{
 	private:
 		static constexpr uint32 InvalidValue = 0;
 
 	public:
-		static const FActiveStateId Invalid;
+		static const FActiveStateID Invalid;
 
-		explicit FActiveStateId() = default;
-		explicit FActiveStateId(uint32 NewId)
-			: Value(NewId)
+		explicit FActiveStateID() = default;
+		explicit FActiveStateID(uint32 NewID)
+			: Value(NewID)
 		{
-			check(NewId != InvalidValue);
+			check(NewID != InvalidValue);
 		}
 		
-		[[nodiscard]] bool operator==(const FActiveStateId&) const = default;
-		[[nodiscard]] bool operator!=(const FActiveStateId&) const = default;
+		[[nodiscard]] bool operator==(const FActiveStateID&) const = default;
+		[[nodiscard]] bool operator!=(const FActiveStateID&) const = default;
 
 		[[nodiscard]] bool IsValid() const
 		{
@@ -89,8 +89,8 @@ namespace UE::StateTree
 		/** An state entry in the state path. */
 		struct FElement
 		{
-			explicit FElement(FStateTreeStateHandle InHandle, FActiveStateId InStateId)
-				: ActiveStateId(InStateId)
+			explicit FElement(FStateTreeStateHandle InHandle, FActiveStateID InStateID)
+				: ActiveStateID(InStateID)
 				, TreeStateIndex(InHandle)
 			{}
 
@@ -102,9 +102,9 @@ namespace UE::StateTree
 				return TreeStateIndex.IsValid();
 			}
 
-			[[nodiscard]] FActiveStateId GetActiveStateId() const
+			[[nodiscard]] FActiveStateID GetActiveStateID() const
 			{
-				return ActiveStateId;
+				return ActiveStateID;
 			}
 
 			[[nodiscard]] FStateTreeStateHandle GetHandle() const
@@ -113,8 +113,8 @@ namespace UE::StateTree
 			}
 
 		private:
-			/** The unique id of the state for this instance. */
-			FActiveStateId ActiveStateId;
+			/** The unique ID of the state for this instance. */
+			FActiveStateID ActiveStateID;
 			/** The index of the state inside the state tree asset. */
 			FStateTreeStateHandle TreeStateIndex;
 		};
@@ -143,7 +143,7 @@ namespace UE::StateTree
 		[[nodiscard]] bool Matches(FElement Other) const;
 
 		/** @return true if the last element in the path matches. */
-		[[nodiscard]] bool Matches(FActiveStateId Other) const;
+		[[nodiscard]] bool Matches(FActiveStateID Other) const;
 
 		/** @return the path where the 2 paths intersect. */
 		[[nodiscard]] const TArrayView<const FElement> Intersect(const FActiveStatePath& Other) const;
@@ -158,7 +158,7 @@ namespace UE::StateTree
 		[[nodiscard]] bool Contains(FElement Other) const;
 
 		/** @return true if this path contains anywhere the other element. */
-		[[nodiscard]] bool Contains(FActiveStateId Other) const;
+		[[nodiscard]] bool Contains(FActiveStateID Other) const;
 
 		/** @return debugging information about the path. */
 		[[nodiscard]] TValueOrError<FString, void> Describe() const;
diff --git a/Engine/Plugins/Runtime/StateTree/Source/StateTreeTestSuite/Private/StateTreeStatePathTest.cpp b/Engine/Plugins/Runtime/StateTree/Source/StateTreeTestSuite/Private/StateTreeStatePathTest.cpp
--- a/Engine/Plugins/Runtime/StateTree/Source/StateTreeTestSuite/Private/StateTreeStatePathTest.cpp
+++ b/Engine/Plugins/Runtime/StateTree/Source/StateTreeTestSuite/Private/StateTreeStatePathTest.cpp
@@ -254,28 +254,28 @@ struct FStateTreeTest_StatePath_LinkStates : FStateTreeTestBase
 
 		// Test variable and helper functions
 		int32 ActiveCounter = 0;
-		TArray<UE::StateTree::FActiveFrameId> PreviousFrameIds;
-		auto PreTestFrameId = [&PreviousFrameIds, &InstanceData]()
+		TArray<UE::StateTree::FActiveFrameID> PreviousFrameIDs;
+		auto PreTestFrameId = [&PreviousFrameIDs, &InstanceData]()
 			{
-				PreviousFrameIds.Reset();
+				PreviousFrameIDs.Reset();
 				for (int32 Index = 0; Index < InstanceData.GetExecutionState()->ActiveFrames.Num(); ++Index)
 				{
-					PreviousFrameIds.Add(InstanceData.GetExecutionState()->ActiveFrames[Index].FrameId);
+					PreviousFrameIDs.Add(InstanceData.GetExecutionState()->ActiveFrames[Index].FrameID);
 				}
 			};
-		auto TestFrameId = [&PreviousFrameIds, &InstanceData](int32 CorrectAmount)
+		auto TestFrameId = [&PreviousFrameIDs, &InstanceData](int32 CorrectAmount)
 			{
 				int32 Index = 0;
-				for (; Index < PreviousFrameIds.Num() && Index < CorrectAmount; ++Index)
+				for (; Index < PreviousFrameIDs.Num() && Index < CorrectAmount; ++Index)
 				{
-					if (PreviousFrameIds[Index] != InstanceData.GetExecutionState()->ActiveFrames[Index].FrameId)
+					if (PreviousFrameIDs[Index] != InstanceData.GetExecutionState()->ActiveFrames[Index].FrameID)
 					{
 						return false;
 					}
 				}
-				for (; Index < InstanceData.GetExecutionState()->ActiveFrames.Num() && Index < PreviousFrameIds.Num(); ++Index)
+				for (; Index < InstanceData.GetExecutionState()->ActiveFrames.Num() && Index < PreviousFrameIDs.Num(); ++Index)
 				{
-					if (PreviousFrameIds[Index] == InstanceData.GetExecutionState()->ActiveFrames[Index].FrameId)
+					if (PreviousFrameIDs[Index] == InstanceData.GetExecutionState()->ActiveFrames[Index].FrameID)
 					{
 						return false;
 					}
@@ -293,31 +293,31 @@ struct FStateTreeTest_StatePath_LinkStates : FStateTreeTestBase
 			AITEST_TRUE("Has the correct number of path elements", ExecPath.Num() == 9);
 			AITEST_TRUE("Has the correct number of active states", InstanceData.GetExecutionState()->ActiveFrames.Num() == 4);
 			{
-				AITEST_TRUE("Frame for Tree1Root is active", InstanceData.GetExecutionState()->ActiveFrames[0].FrameId == UE::StateTree::FActiveFrameId(++ActiveCounter));
-				AITEST_TRUE("State Tree1Root is active", ExecPath.Contains(UE::StateTree::FActiveStateId(++ActiveCounter)));
-				AITEST_TRUE("State Tree1Root is active", ExecPath.Contains(UE::StateTree::FActiveStatePath::FElement(FStateTreeStateHandle(0), UE::StateTree::FActiveStateId(ActiveCounter))));
-				AITEST_TRUE("State Tree1RootState1 is active", ExecPath.Contains(UE::StateTree::FActiveStateId(++ActiveCounter)));
-				AITEST_TRUE("State Tree1RootState1 is active", ExecPath.Contains(UE::StateTree::FActiveStatePath::FElement(FStateTreeStateHandle(1), UE::StateTree::FActiveStateId(ActiveCounter))));
-				AITEST_TRUE("State Tree1State1StateSub1A is active", ExecPath.Contains(UE::StateTree::FActiveStateId(++ActiveCounter)));
-				AITEST_TRUE("State Tree1State1StateSub1A is active", ExecPath.Contains(UE::StateTree::FActiveStatePath::FElement(FStateTreeStateHandle(2), UE::StateTree::FActiveStateId(4))));
+				AITEST_TRUE("Frame for Tree1Root is active", InstanceData.GetExecutionState()->ActiveFrames[0].FrameID == UE::StateTree::FActiveFrameID(++ActiveCounter));
+				AITEST_TRUE("State Tree1Root is active", ExecPath.Contains(UE::StateTree::FActiveStateID(++ActiveCounter)));
+				AITEST_TRUE("State Tree1Root is active", ExecPath.Contains(UE::StateTree::FActiveStatePath::FElement(FStateTreeStateHandle(0), UE::StateTree::FActiveStateID(ActiveCounter))));
+				AITEST_TRUE("State Tree1RootState1 is active", ExecPath.Contains(UE::StateTree::FActiveStateID(++ActiveCounter)));
+				AITEST_TRUE("State Tree1RootState1 is active", ExecPath.Contains(UE::StateTree::FActiveStatePath::FElement(FStateTreeStateHandle(1), UE::StateTree::FActiveStateID(ActiveCounter))));
+				AITEST_TRUE("State Tree1State1StateSub1A is active", ExecPath.Contains(UE::StateTree::FActiveStateID(++ActiveCounter)));
+				AITEST_TRUE("State Tree1State1StateSub1A is active", ExecPath.Contains(UE::StateTree::FActiveStatePath::FElement(FStateTreeStateHandle(2), UE::StateTree::FActiveStateID(4))));
 			}
 			{
-				AITEST_TRUE("Frame for Tree1Sub1 is active", InstanceData.GetExecutionState()->ActiveFrames[1].FrameId == UE::StateTree::FActiveFrameId(++ActiveCounter));
-				AITEST_TRUE("State Tree1Sub1 is active", ExecPath.Contains(UE::StateTree::FActiveStateId(++ActiveCounter)));
-				AITEST_TRUE("State Tree1Sub1StateLinkTree2A is active", ExecPath.Contains(UE::StateTree::FActiveStateId(++ActiveCounter)));
+				AITEST_TRUE("Frame for Tree1Sub1 is active", InstanceData.GetExecutionState()->ActiveFrames[1].FrameID == UE::StateTree::FActiveFrameID(++ActiveCounter));
+				AITEST_TRUE("State Tree1Sub1 is active", ExecPath.Contains(UE::StateTree::FActiveStateID(++ActiveCounter)));
+				AITEST_TRUE("State Tree1Sub1StateLinkTree2A is active", ExecPath.Contains(UE::StateTree::FActiveStateID(++ActiveCounter)));
 			}
 			{
-				AITEST_TRUE("Frame for Tree2StateRoot is active", InstanceData.GetExecutionState()->ActiveFrames[2].FrameId == UE::StateTree::FActiveFrameId(++ActiveCounter));
-				AITEST_TRUE("State Tree2StateRoot is active", ExecPath.Contains(UE::StateTree::FActiveStateId(++ActiveCounter)));
-				AITEST_TRUE("State Tree2RootStateSub1A is active", ExecPath.Contains(UE::StateTree::FActiveStateId(++ActiveCounter)));
+				AITEST_TRUE("Frame for Tree2StateRoot is active", InstanceData.GetExecutionState()->ActiveFrames[2].FrameID == UE::StateTree::FActiveFrameID(++ActiveCounter));
+				AITEST_TRUE("State Tree2StateRoot is active", ExecPath.Contains(UE::StateTree::FActiveStateID(++ActiveCounter)));
+				AITEST_TRUE("State Tree2RootStateSub1A is active", ExecPath.Contains(UE::StateTree::FActiveStateID(++ActiveCounter)));
 			}
 			{
-				AITEST_TRUE("Frame for Tree2Sub1 is active", InstanceData.GetExecutionState()->ActiveFrames[3].FrameId == UE::StateTree::FActiveFrameId(++ActiveCounter));
-				AITEST_TRUE("State Tree2Sub1 is active", ExecPath.Contains(UE::StateTree::FActiveStateId(++ActiveCounter)));
-				AITEST_TRUE("State Tree2Sub1State1 is active", ExecPath.Contains(UE::StateTree::FActiveStateId(++ActiveCounter)));
+				AITEST_TRUE("Frame for Tree2Sub1 is active", InstanceData.GetExecutionState()->ActiveFrames[3].FrameID == UE::StateTree::FActiveFrameID(++ActiveCounter));
+				AITEST_TRUE("State Tree2Sub1 is active", ExecPath.Contains(UE::StateTree::FActiveStateID(++ActiveCounter)));
+				AITEST_TRUE("State Tree2Sub1State1 is active", ExecPath.Contains(UE::StateTree::FActiveStateID(++ActiveCounter)));
 			}
 			{
-				AITEST_FALSE("No accidental increment", ExecPath.Contains(UE::StateTree::FActiveStateId(ActiveCounter+1)));
+				AITEST_FALSE("No accidental increment", ExecPath.Contains(UE::StateTree::FActiveStateID(ActiveCounter+1)));
 			}
 			Exec.LogClear();
 		}
@@ -330,8 +330,8 @@ struct FStateTreeTest_StatePath_LinkStates : FStateTreeTestBase
 			UE::StateTree::FActiveStatePath ExecPath = InstanceData.GetExecutionState()->GetActiveStatePath();
 			AITEST_TRUE("Has the correct number of path elements", ExecPath.Num() == 9);
 			AITEST_TRUE("Has the correct number of active states", InstanceData.GetExecutionState()->ActiveFrames.Num() == 4);
-			AITEST_TRUE("State Tree2Sub1State2 is active", ExecPath.Contains(UE::StateTree::FActiveStateId(++ActiveCounter)));
-			AITEST_FALSE("No accidental increment", ExecPath.Contains(UE::StateTree::FActiveStateId(ActiveCounter + 1)));
+			AITEST_TRUE("State Tree2Sub1State2 is active", ExecPath.Contains(UE::StateTree::FActiveStateID(++ActiveCounter)));
+			AITEST_FALSE("No accidental increment", ExecPath.Contains(UE::StateTree::FActiveStateID(ActiveCounter + 1)));
 			AITEST_TRUE("All frames are the same", TestFrameId(4));
 			Exec.LogClear();
 		}
