Remove include to GSLNotNull.h
#rnx

diff --git a/Engine/Plugins/Runtime/StateTree/Source/StateTreeModule/Public/Debugger/StateTreeDebug.h b/Engine/Plugins/Runtime/StateTree/Source/StateTreeModule/Public/Debugger/StateTreeDebug.h
--- a/Engine/Plugins/Runtime/StateTree/Source/StateTreeModule/Public/Debugger/StateTreeDebug.h
+++ b/Engine/Plugins/Runtime/StateTree/Source/StateTreeModule/Public/Debugger/StateTreeDebug.h
@@ -3,7 +3,6 @@
 #pragma once
 
 #include "StateTreeTypes.h"
-#include "Templates/GSLNotNull.h"
 
 class UStateTree;
 struct FStateTreeExecutionContext;
diff --git a/Engine/Source/Runtime/Core/Public/Templates/NotNull.h b/Engine/Source/Runtime/Core/Public/Templates/NotNull.h
--- a/Engine/Source/Runtime/Core/Public/Templates/NotNull.h
+++ b/Engine/Source/Runtime/Core/Public/Templates/NotNull.h
@@ -8,4 +8,6 @@
 
 UE_DEPRECATED_HEADER(5.6, "Templates/NotNull.h has been renamed to Templates/GSLNotNull.h to avoid confusion with Misc/NotNull.h which contains the TNotNull type")
 
+#if UE_ENABLE_INCLUDE_ORDER_DEPRECATED_IN_5_6
 #include "Templates/GSLNotNull.h"
+#endif
