[State Tree Editor] fixed StateTreeViewModel's selected states not synced to StateTree View's ones

diff --git a/Engine/Plugins/Runtime/StateTree/Source/StateTreeEditorModule/Private/StateTreeViewModel.cpp b/Engine/Plugins/Runtime/StateTree/Source/StateTreeEditorModule/Private/StateTreeViewModel.cpp
--- a/Engine/Plugins/Runtime/StateTree/Source/StateTreeEditorModule/Private/StateTreeViewModel.cpp
+++ b/Engine/Plugins/Runtime/StateTree/Source/StateTreeEditorModule/Private/StateTreeViewModel.cpp
@@ -655,7 +655,7 @@ void FStateTreeViewModel::SetSelection(UStateTreeState* Selected)
 	}
 
 	SelectedStates.Reset();
-	if (Selected == nullptr)
+	if (Selected != nullptr)
 	{
 		SelectedStates.Add(Selected);
 	}
