[StateTree] added bindable text to DebugTextTask

diff --git a/Engine/Plugins/Runtime/StateTree/Source/StateTreeModule/Private/Tasks/StateTreeDebugTextTask.cpp b/Engine/Plugins/Runtime/StateTree/Source/StateTreeModule/Private/Tasks/StateTreeDebugTextTask.cpp
--- a/Engine/Plugins/Runtime/StateTree/Source/StateTreeModule/Private/Tasks/StateTreeDebugTextTask.cpp
+++ b/Engine/Plugins/Runtime/StateTree/Source/StateTreeModule/Private/Tasks/StateTreeDebugTextTask.cpp
@@ -39,8 +39,16 @@ EStateTreeRunStatus FStateTreeDebugTextTask::EnterState(FStateTreeExecutionConte
 		return EStateTreeRunStatus::Failed;
 	}
 
-	DrawDebugString(World, Offset, Text, InstanceData.ReferenceActor, TextColor,	/*Duration*/-1,	/*DrawShadows*/true, FontScale);
-	
+	if (!Text.IsEmpty())
+	{
+		DrawDebugString(World, Offset, Text, InstanceData.ReferenceActor, TextColor, /*Duration*/-1, /*DrawShadows*/true, FontScale);
+	}
+
+	if (!InstanceData.BindableText.IsEmpty())
+	{
+		DrawDebugString(World, Offset, InstanceData.BindableText, InstanceData.ReferenceActor, TextColor, /*Duration*/-1, /*DrawShadows*/true, FontScale);
+	}
+
 	return EStateTreeRunStatus::Running;
 }
 
@@ -73,6 +81,8 @@ void FStateTreeDebugTextTask::ExitState(FStateTreeExecutionContext& Context, con
 #if WITH_EDITOR
 FText FStateTreeDebugTextTask::GetDescription(const FGuid& ID, FStateTreeDataView InstanceDataView, const IStateTreeBindingLookup& BindingLookup, EStateTreeNodeFormatting Formatting) const
 {
+	// Note that FStateTreeDebugTextTaskInstanceData::BindableText is not added to the formatted string
+	//  since the bindings are not copied at this point so there is nothing to display when not at runtime.
 	const FText Format = (Formatting == EStateTreeNodeFormatting::RichText)
 		? LOCTEXT("DebugTextRich", "<b>Debug Text</> \"{Text}\"")
 		: LOCTEXT("DebugText", "Debug Text \"{Text}\"");
diff --git a/Engine/Plugins/Runtime/StateTree/Source/StateTreeModule/Private/Tasks/StateTreeDebugTextTask.h b/Engine/Plugins/Runtime/StateTree/Source/StateTreeModule/Private/Tasks/StateTreeDebugTextTask.h
--- a/Engine/Plugins/Runtime/StateTree/Source/StateTreeModule/Private/Tasks/StateTreeDebugTextTask.h
+++ b/Engine/Plugins/Runtime/StateTree/Source/StateTreeModule/Private/Tasks/StateTreeDebugTextTask.h
@@ -16,6 +16,9 @@ struct STATETREEMODULE_API FStateTreeDebugTextTaskInstanceData
 	/** Optional actor where to draw the text at. */
 	UPROPERTY(EditAnywhere, Category = "Input", meta=(Optional))
 	TObjectPtr<AActor> ReferenceActor = nullptr;
+
+	UPROPERTY(EditAnywhere, Category = "Parameter")
+	FString BindableText;
 };
 
 /**
