StateTree: Prevent copy of the execution context. The destructor reset the external context field.

diff --git a/Engine/Plugins/Runtime/StateTree/Source/StateTreeModule/Public/StateTreeExecutionContext.h b/Engine/Plugins/Runtime/StateTree/Source/StateTreeModule/Public/StateTreeExecutionContext.h
--- a/Engine/Plugins/Runtime/StateTree/Source/StateTreeModule/Public/StateTreeExecutionContext.h
+++ b/Engine/Plugins/Runtime/StateTree/Source/StateTreeModule/Public/StateTreeExecutionContext.h
@@ -51,6 +51,10 @@ public:
 	explicit FStateTreeMinimalExecutionContext(UObject& Owner, const UStateTree& StateTree, FStateTreeInstanceStorage& Storage);
 	virtual ~FStateTreeMinimalExecutionContext();
 
+private:
+	FStateTreeMinimalExecutionContext(const FStateTreeMinimalExecutionContext&) = delete;
+	FStateTreeMinimalExecutionContext& operator=(const FStateTreeMinimalExecutionContext&) = delete;
+
 public:
 	/**
 	 * Indicates if the instance is valid and would be able to run the instance of the associated StateTree asset with a regular execution context.
