StateTree: Use state ID to identify the tasks that finish. This prevents a task from triggering a finished task after it is finished.
Clean the completed state array in exit state to prevent users from finishing a task after it is finished.
Add find functions to find the frame by ID.
Update the weak context to use frame and state ID.
The current execution is not needed anymore. It could introduced bugs if the functionnality was triggered from a delegate listener.
Update delegates to use frame and state ID.

diff --git a/Engine/Plugins/Runtime/StateTree/Source/StateTreeModule/Private/StateTreeExecutionContext.cpp b/Engine/Plugins/Runtime/StateTree/Source/StateTreeModule/Private/StateTreeExecutionContext.cpp
--- a/Engine/Plugins/Runtime/StateTree/Source/StateTreeModule/Private/StateTreeExecutionContext.cpp
+++ b/Engine/Plugins/Runtime/StateTree/Source/StateTreeModule/Private/StateTreeExecutionContext.cpp
@@ -75,54 +75,6 @@ namespace UE::StateTree::ExecutionContext::Private
 		static_assert((int32)EStateTreeRunStatus::Failed == (int32)EStateTreeFinishTaskType::Failed, "Failed do not match in both types.");
 		return (EStateTreeRunStatus)FinishTask;
 	}
-
-	struct FActiveStateIndex
-	{
-		int32 FrameIndex = INDEX_NONE;
-		int32 ActiveStateIndex = INDEX_NONE;
-	};
-	TOptional<FActiveStateIndex> FindActiveStateIndex(const FStateTreeExecutionState& ExecutionState, const FStateTreeDelayedFinishedTask& Task)
-	{
-		TOptional<FActiveStateIndex> Result;
-
-		const int32 FrameIndex = ExecutionState.ActiveFrames.IndexOfByPredicate([&Task](const FStateTreeExecutionFrame& Frame)
-			{
-				return FObjectKey(Frame.StateTree) == Task.FrameStateTree && Frame.RootState == Task.FrameRootState;
-			});
-		if (FrameIndex == INDEX_NONE)
-		{
-			return Result;
-		}
-
-		const UStateTree* StateTree = ExecutionState.ActiveFrames[FrameIndex].StateTree;
-		if (!StateTree->GetStates().IsValidIndex(Task.State.Index))
-		{
-			return Result;
-		}
-
-		const int32 ActiveStateIndex = ExecutionState.ActiveFrames[FrameIndex].ActiveStates.IndexOfReverse(Task.State);
-		if (ActiveStateIndex == INDEX_NONE)
-		{
-			return Result;
-		}
-
-		if (Task.TaskIndex.IsValid())
-		{
-			const FCompactStateTreeState& State = StateTree->GetStates()[Task.State.Index];
-			const int32 BeginTaskIndex = State.TasksBegin;
-			const int32 EndTaskIndex = State.TasksBegin + State.TasksNum;
-			if (ensure(Task.TaskIndex.AsInt32() >= BeginTaskIndex && Task.TaskIndex.AsInt32() < EndTaskIndex))
-			{
-				Result = FActiveStateIndex{ .FrameIndex = FrameIndex, .ActiveStateIndex = ActiveStateIndex };
-			}
-		}
-		else
-		{
-			//The TaskIndex can be invalid if no task is responsible for the state completion.
-			Result = FActiveStateIndex{ .FrameIndex = FrameIndex, .ActiveStateIndex = ActiveStateIndex };
-		}
-		return Result;
-	}
 }
 
 FStateTreeMinimalExecutionContext::FStateTreeMinimalExecutionContext(UObject& InOwner, const UStateTree& InStateTree, FStateTreeInstanceData& InInstanceData)
@@ -387,25 +339,6 @@ FStateTreeInstanceDebugId FStateTreeMinimalExecutionContext::GetInstanceDebugId(
 }
 #endif // WITH_STATETREE_TRACE
 
-FStateTreeExecutionContext::FCurrentlyExecutionContextScope::FCurrentlyExecutionContextScope(FStateTreeExecutionContext& InContext)
-	: CurrentContext(InContext)
-{
-	FStateTreeExecutionState& Exec = InContext.GetExecState();
-	ensureMsgf(Exec.CurrentExecutionContext == nullptr || Exec.CurrentExecutionContext == &InContext, TEXT("Only one execution context can operate on the State Tree Instance Data."));
-	PreviousExecutionContext = Exec.CurrentExecutionContext;
-	Exec.CurrentExecutionContext = &InContext;
-}
-
-FStateTreeExecutionContext::FCurrentlyExecutionContextScope::~FCurrentlyExecutionContextScope()
-{
-	FStateTreeExecutionState& Exec = CurrentContext.GetExecState();
-	// The tree can stop (it will reset the execution instance data) while we are ticking the state tree.
-	if (Exec.CurrentExecutionContext == &CurrentContext)
-	{
-		Exec.CurrentExecutionContext = PreviousExecutionContext;
-	}
-}
-
 FStateTreeExecutionContext::FCurrentlyProcessedFrameScope::FCurrentlyProcessedFrameScope(FStateTreeExecutionContext& InContext, const FStateTreeExecutionFrame* CurrentParentFrame, const FStateTreeExecutionFrame& CurrentFrame): Context(InContext)
 {
 	check(CurrentFrame.StateTree);
@@ -675,7 +608,6 @@ EStateTreeRunStatus FStateTreeExecutionContext::Start(FStartParameters Parameter
 		SetGlobalParameters(RootStateTree.GetDefaultParameters());
 	}
 
-	FCurrentlyExecutionContextScope ContextScope(*this);
 	TGuardValue<bool> ScheduledNextTickScope(bAllowedToScheduleNextTick, false);
 
 	// Initialize for the init frame.
@@ -815,7 +747,6 @@ EStateTreeRunStatus FStateTreeExecutionContext::Stop(EStateTreeRunStatus Complet
 		return EStateTreeRunStatus::Failed;
 	}
 
-	FCurrentlyExecutionContextScope ContextScope(*this);
 	TGuardValue<bool> ScheduledNextTickScope(bAllowedToScheduleNextTick, false);
 
 	// Set scoped phase only for properly initialized context with valid Instance data
@@ -940,7 +871,6 @@ EStateTreeRunStatus FStateTreeExecutionContext::Tick(const float DeltaTime)
 {
 	CSV_SCOPED_TIMING_STAT_EXCLUSIVE(StateTree_Tick);
 
-	FCurrentlyExecutionContextScope ContextScope(*this);
 	TGuardValue<bool> ScheduledNextTickScope(bAllowedToScheduleNextTick, false);
 
 	const EStateTreeRunStatus PreludeResult = TickPrelude();
@@ -959,7 +889,6 @@ EStateTreeRunStatus FStateTreeExecutionContext::TickUpdateTasks(const float Delt
 {
 	CSV_SCOPED_TIMING_STAT_EXCLUSIVE(StateTree_Tick);
 
-	FCurrentlyExecutionContextScope ContextScope(*this);
 	TGuardValue<bool> ScheduledNextTickScope(bAllowedToScheduleNextTick, false);
 
 	const EStateTreeRunStatus PreludeResult = TickPrelude();
@@ -977,7 +906,6 @@ EStateTreeRunStatus FStateTreeExecutionContext::TickTriggerTransitions()
 {
 	CSV_SCOPED_TIMING_STAT_EXCLUSIVE(StateTree_Tick);
 
-	FCurrentlyExecutionContextScope ContextScope(*this);
 	TGuardValue<bool> ScheduledNextTickScope(bAllowedToScheduleNextTick, false);
 
 	const EStateTreeRunStatus PreludeResult = TickPrelude();
@@ -1152,9 +1080,14 @@ void FStateTreeExecutionContext::BindDelegate(const FStateTreeDelegateListener&
 	}
 
 	const FStateTreeExecutionFrame* CurrentFrame = GetCurrentlyProcessedFrame();
-	check(CurrentFrame);
+	if (CurrentFrame == nullptr)
+	{
+		return;
+	}
 
-	GetExecState().DelegateActiveListeners.Add(Listener, MoveTemp(Delegate), CurrentlyProcessedFrame->StateTree, CurrentlyProcessedState.Root, CurrentlyProcessedState, FStateTreeIndex16(CurrentNodeDataHandle.GetIndex()));
+	const int32 ActiveStateIndex = CurrentFrame->ActiveStates.IndexOfReverse(CurrentlyProcessedState);
+	const UE::StateTree::FActiveStateID CurrentlyStateID = ActiveStateIndex != INDEX_NONE ? CurrentFrame->ActiveStates.StateIDs[ActiveStateIndex] : UE::StateTree::FActiveStateID::Invalid;
+	GetExecState().DelegateActiveListeners.Add(Listener, MoveTemp(Delegate), CurrentFrame->FrameID, CurrentlyStateID, FStateTreeIndex16(CurrentNodeDataHandle.GetIndex()));
 }
 
 void FStateTreeExecutionContext::RemoveDelegateListener(const FStateTreeDelegateListener& Listener)
@@ -1176,9 +1109,6 @@ void FStateTreeExecutionContext::UnbindDelegate(const FStateTreeDelegateListener
 		return;
 	}
 
-	const FStateTreeExecutionFrame* CurrentFrame = GetCurrentlyProcessedFrame();
-	check(CurrentFrame);
-
 	GetExecState().DelegateActiveListeners.Remove(Listener);
 }
 
@@ -1253,23 +1183,34 @@ void FStateTreeExecutionContext::FinishTask(const FStateTreeTaskBase& Task, ESta
 	{
 		return;
 	}
+	check(CurrentlyProcessedFrame);
+	check(CurrentlyProcessedState.IsValid());
 
+	using namespace UE::StateTree;
 	using namespace UE::StateTree::ExecutionContext::Private;
 
 	FStateTreeExecutionState& Exec = GetExecState();
 
 	// Then we have the current frame and the current state.
-	const int32 ActiveFrameIndex = Exec.ActiveFrames.IndexOfByPredicate([FramePtr = CurrentlyProcessedFrame](const FStateTreeExecutionFrame& Other) { return &Other == FramePtr; });
-	check(ActiveFrameIndex != INDEX_NONE);
-	const int32 ActiveStateIndex = Exec.ActiveFrames[ActiveFrameIndex].ActiveStates.IndexOfReverse(CurrentlyProcessedState);
-	check(ActiveStateIndex != INDEX_NONE);
-	FinishTaskInternal(ActiveFrameIndex, ActiveStateIndex, CurrentNodeIndex, FinishTaskTypeToRunStatus(FinishType));
+	const FActiveFrameID ActiveFrameID = CurrentlyProcessedFrame->FrameID;
+	check(ActiveFrameID.IsValid());
+	if (CurrentlyProcessedState.IsValid())
+	{
+		const int32 ActiveStateIndex = CurrentlyProcessedFrame->ActiveStates.IndexOfReverse(CurrentlyProcessedState);
+		check(ActiveStateIndex != INDEX_NONE);
+		const FActiveStateID ActiveStateID = CurrentlyProcessedFrame->ActiveStates.StateIDs[ActiveStateIndex];
+		FinishTaskInternal(FDelayedFinishedTask::EReasonType::StateTask, ActiveFrameID, ActiveStateID, CurrentNodeIndex, FinishTaskTypeToRunStatus(FinishType));
+	}
+	else
+	{
+		FinishTaskInternal(FDelayedFinishedTask::EReasonType::GlobalTask, ActiveFrameID, FActiveStateID::Invalid, CurrentNodeIndex, FinishTaskTypeToRunStatus(FinishType));
+	}
 
 	// The FStateTreeTaskBase::Tick or OnEnter returns the value. Set it here and test it later.
 	CompletedStateRunStatus = FinishTaskTypeToRunStatus(FinishType);
 }
 
-void FStateTreeExecutionContext::FinishTask(const FStateTreeDelayedFinishedTask& Task, EStateTreeFinishTaskType FinishType)
+void FStateTreeExecutionContext::FinishTask(const UE::StateTree::FDelayedFinishedTask& Task, EStateTreeFinishTaskType FinishType)
 {
 	if (!IsValid())
 	{
@@ -1278,135 +1219,179 @@ void FStateTreeExecutionContext::FinishTask(const FStateTreeDelayedFinishedTask&
 		return;
 	}
 
-	using namespace UE::StateTree::ExecutionContext::Private;
-
-	FStateTreeExecutionState& Exec = GetExecState();
-	//Is it still valid to use
-	TOptional<FActiveStateIndex> FrameAndState = FindActiveStateIndex(Exec, Task);
-	if (FrameAndState.IsSet())
+	if (CurrentNodeIndex == Task.TaskIndex.AsInt32() && CurrentlyProcessedFrame && Task.FrameID == CurrentlyProcessedFrame->FrameID)
 	{
-		if (CurrentNodeIndex == Task.TaskIndex.AsInt32())
-		{
-			// We are currently ticking that task. We are either in Tick, EnterState, ExitState, Transition.
-			//NB. Completed tasks are not considered during ExitState or Transition.
-			//The node will return the value and it will be processed once returned. Set it here and test it later.
-			CompletedStateRunStatus = FinishTaskTypeToRunStatus(FinishType);
-		}
-		else
-		{
-			FinishTaskInternal(FrameAndState.GetValue().FrameIndex, FrameAndState.GetValue().ActiveStateIndex, Task.TaskIndex.AsInt32(), FinishTaskTypeToRunStatus(FinishType));
-			ScheduleNextTick();
-		}
+		// We are currently ticking that task. We are either in Tick, EnterState, ExitState, Transition.
+		//Note: Completed tasks are not considered during ExitState or Transition.
+		//The node will return the value and it will be processed once returned. Set it here and test it later.
+		CompletedStateRunStatus = UE::StateTree::ExecutionContext::Private:: FinishTaskTypeToRunStatus(FinishType);
+	}
+	else if (IsDelayedFinishedTaskValid(Task))
+	{
+		FStateTreeExecutionState& Exec = GetExecState();
+		const UE::StateTree::FDelayedFinishedTask& DelayedCompletedTask = Exec.DelayedFinishedTasks.Add_GetRef(Task);
+		AddUniqueCompletedStateList(DelayedCompletedTask);
+		ScheduleNextTick();
 	}
 }
 
-void FStateTreeExecutionContext::FinishTaskInternal(int32 ActiveFrameIndex, int32 ActiveStateIndex, int32 TaskIndex, EStateTreeRunStatus RunStatus)
+void FStateTreeExecutionContext::FinishTaskInternal(UE::StateTree::FDelayedFinishedTask::EReasonType Reason, UE::StateTree::FActiveFrameID ActiveFrameID, UE::StateTree::FActiveStateID ActiveStateID, int32 TaskIndex, EStateTreeRunStatus RunStatus)
 {
+	using namespace UE::StateTree;
+
 	FStateTreeExecutionState& Exec = GetExecState();
-	FStateTreeExecutionFrame& ActiveFrame = Exec.ActiveFrames[ActiveFrameIndex];
 
 	// Add it for later. In case the user called TickTasks and will call TickTransitions later.
 	const FStateTreeIndex16 TaskIndex16 = TaskIndex >= 0 ? FStateTreeIndex16(TaskIndex) : FStateTreeIndex16::Invalid;
 	const EStateTreeFinishTaskType FinishType = RunStatus == EStateTreeRunStatus::Succeeded ? EStateTreeFinishTaskType::Succeeded : EStateTreeFinishTaskType::Failed;
-	const FStateTreeDelayedFinishedTask& DelayedCompletedTask = Exec.DelayedFinishedTasks.Add_GetRef(FStateTreeDelayedFinishedTask{.FrameStateTree = FObjectKey(ActiveFrame.StateTree), .FrameRootState = ActiveFrame.RootState, .State = ActiveFrame.ActiveStates[ActiveStateIndex], .TaskIndex = TaskIndex16, .Type = FinishType});
 
-	AddUniqueCompletedStateList(ActiveFrameIndex, ActiveStateIndex, DelayedCompletedTask);
+	const FDelayedFinishedTask Task {.FrameID = ActiveFrameID, .StateID = ActiveStateID, .TaskIndex = TaskIndex16, .Type = FinishType, .Reason = Reason };
+	if (IsDelayedFinishedTaskValid(Task))
+	{
+		Exec.DelayedFinishedTasks.Add(Task);
+		AddUniqueCompletedStateList(Task);
+	}
 }
 
-void FStateTreeExecutionContext::UpdateCompletedStateList()
+bool FStateTreeExecutionContext::IsDelayedFinishedTaskValid(const UE::StateTree::FDelayedFinishedTask& Task) const
 {
-	FStateTreeExecutionState& Exec = GetExecState();
+	using namespace UE::StateTree;
+
+	const FStateTreeExecutionState& Exec = GetExecState();
+
+	// Is frame valid
+	const FStateTreeExecutionFrame* ExecutionFrame = Exec.FindActiveFrame(Task.FrameID);
+	if (ExecutionFrame == nullptr)
+	{
+		return false;
+	}
+
+	const UStateTree* StateTree = ExecutionFrame->StateTree;
+	check(StateTree);
 
+	if (Task.Reason == FDelayedFinishedTask::EReasonType::StateTask)
+	{
+		if (!ensure(Task.StateID.IsValid() && Task.TaskIndex.IsValid()))
+		{
+			return false;
+		}
+
+		const int32 ActiveStateIndex = ExecutionFrame->ActiveStates.IndexOfReverse(Task.StateID);
+		if (ActiveStateIndex == INDEX_NONE)
+		{
+			return false;
+		}
+
+		const int32 TreeStateIndex = ExecutionFrame->ActiveStates[ActiveStateIndex].Index;
+		check(StateTree->GetStates().IsValidIndex(TreeStateIndex));
+		const FCompactStateTreeState& State = StateTree->GetStates()[TreeStateIndex];
+		const int32 BeginTaskIndex = State.TasksBegin;
+		const int32 EndTaskIndex = State.TasksBegin + State.TasksNum;
+		const bool bIsValidTask = Task.TaskIndex.AsInt32() >= BeginTaskIndex && Task.TaskIndex.AsInt32() < EndTaskIndex;
+		ensureMsgf(bIsValidTask, TEXT("The delayed finished task has a task index that doesn't match the state. Have you recompiled the asset runtime?"));
+		return bIsValidTask;
+	}
+	else if (Task.Reason == FDelayedFinishedTask::EReasonType::GlobalTask)
+	{
+		if (!ensure(ExecutionFrame->bIsGlobalFrame && Task.TaskIndex.IsValid()))
+		{
+			return false;
+		}
+
+		const int32 BeginTaskIndex = StateTree->GlobalTasksBegin;
+		const int32 EndTaskIndex = StateTree->GlobalTasksBegin + StateTree->GlobalTasksNum;
+		const bool bIsValidTask = Task.TaskIndex.AsInt32() >= BeginTaskIndex && Task.TaskIndex.AsInt32() < EndTaskIndex;
+		ensureMsgf(bIsValidTask, TEXT("The delayed finished task has a task index that is not global. Have you recompiled the asset runtime?"));
+		return bIsValidTask;
+	}
+	else if (Task.Reason == FDelayedFinishedTask::EReasonType::InternalTransition)
+	{
+		if (!ensure(Task.StateID.IsValid()))
+		{
+			return false;
+		}
+
+		return ExecutionFrame->ActiveStates.Contains(Task.StateID);
+	}
+	return false;
+}
+
+void FStateTreeExecutionContext::UpdateCompletedStateList()
+{
 	// Do not modify the PendingCompletedState until transitions completion.
 	//User can call TickTask and wait a couple of frame before calling TickTransitions.
 
 	// Recreate the completed states list from the pending list
 	CompletedStates.Empty();
 
-	for (const FStateTreeDelayedFinishedTask& PendingCompletedState : Exec.DelayedFinishedTasks)
+	FStateTreeExecutionState& Exec = GetExecState();
+	for (int32 Index = Exec.DelayedFinishedTasks.Num() - 1; Index >= 0; --Index)
 	{
-		if (const UObject* CompletedTree = PendingCompletedState.FrameStateTree.ResolveObjectPtr())
+		const UE::StateTree::FDelayedFinishedTask& Task = Exec.DelayedFinishedTasks[Index];
+		if (IsDelayedFinishedTaskValid(Task))
 		{
-			for (int32 FrameIndex = 0; FrameIndex < Exec.ActiveFrames.Num(); ++FrameIndex)
-			{
-				const FStateTreeExecutionFrame& ActiveFrame = Exec.ActiveFrames[FrameIndex];
-				if (ActiveFrame.StateTree == CompletedTree && ActiveFrame.RootState == PendingCompletedState.FrameRootState)
-				{
-					bool bFound = false;
-					int32 StateIndex = INDEX_NONE;
-					if (PendingCompletedState.State.IsValid())
-					{
-						StateIndex = ActiveFrame.ActiveStates.IndexOfReverse(PendingCompletedState.State);
-						bFound = ActiveFrame.ActiveStates.IsValidIndex(StateIndex);
-					}
-					else
-					{
-						bFound = true; // It's a global task.
-					}
-
-					// It's valid. Add Unique entry.
-					if (bFound)
-					{
-						AddUniqueCompletedStateList(FrameIndex, StateIndex, PendingCompletedState);
-					}
-					break; // Next pending
-				}
-			}
+			AddUniqueCompletedStateList(Task);
+		}
+		else
+		{
+			Exec.DelayedFinishedTasks.RemoveAtSwap(Index);
 		}
 	}
 }
 
-void FStateTreeExecutionContext::AddUniqueCompletedStateList(int32 ActiveFrameIndex, int32 ActiveStateIndex, const FStateTreeDelayedFinishedTask& DelayedFinishedTask)
+void FStateTreeExecutionContext::AddUniqueCompletedStateList(const UE::StateTree::FDelayedFinishedTask& NewFinishedTask)
 {
 	bool bAdd = true;
-	// The priority is global tasks and failure. Only keep 1 per Frame per State.
+	// The priority is global tasks and failure. Only keep 1 FinishedTask per Frame per State.
+	//If the global task exist and a state (of the same tree) completes, only keep the global
 	for (int32 CompletedStateIndex = CompletedStates.Num() - 1; CompletedStateIndex >= 0; --CompletedStateIndex)
 	{
 		FCompletedState& ExistingCompletedState = CompletedStates[CompletedStateIndex];
-		if (ExistingCompletedState.ActiveFrameIndex == ActiveFrameIndex)
+		if (ExistingCompletedState.FrameID == NewFinishedTask.FrameID)
 		{
-			// If it's a global task.
-			if (!DelayedFinishedTask.State.IsValid())
+			// If the existing one is global and the new one is global, only update the finish type.
+			// If the existing one is global and the new one is state, discard new.
+			// If the existing one is state and the new one is global, discard the existing and add the new.
+			// If the existing one is state and match the new one, only update the finish type.
+			
+			if (ExistingCompletedState.Reason == UE::StateTree::FDelayedFinishedTask::EReasonType::GlobalTask)
 			{
-				// Remove or update the existing entry.
-				if (ExistingCompletedState.ActiveStateIndex == INDEX_NONE)
+				if (NewFinishedTask.Reason == UE::StateTree::FDelayedFinishedTask::EReasonType::GlobalTask
+					&& NewFinishedTask.Type == EStateTreeFinishTaskType::Failed)
 				{
 					// Set failure, otherwise keep as it is
-					if (DelayedFinishedTask.Type == EStateTreeFinishTaskType::Failed)
-					{
-						ExistingCompletedState.RunStatus = EStateTreeRunStatus::Failed;
-					}
-					bAdd = false;
-				}
-				else
-				{
-					CompletedStates.RemoveAtSwap(CompletedStateIndex);
+					ExistingCompletedState.RunStatus = EStateTreeRunStatus::Failed;
 				}
-				break;
-			}
-			else if (ExistingCompletedState.ActiveStateIndex == INDEX_NONE)
-			{
-				// Keep the existing entry.
 				bAdd = false;
 				break;
 			}
-			else if (ExistingCompletedState.ActiveStateIndex == ActiveStateIndex)
+			else
 			{
-				// Set failure, otherwise keep as it is
-				if (DelayedFinishedTask.Type == EStateTreeFinishTaskType::Failed)
+				check(ExistingCompletedState.Reason == UE::StateTree::FDelayedFinishedTask::EReasonType::StateTask
+					|| ExistingCompletedState.Reason == UE::StateTree::FDelayedFinishedTask::EReasonType::InternalTransition);
+				if (NewFinishedTask.Reason == UE::StateTree::FDelayedFinishedTask::EReasonType::GlobalTask)
 				{
-					ExistingCompletedState.RunStatus = EStateTreeRunStatus::Failed;
+					CompletedStates.RemoveAtSwap(CompletedStateIndex);
+					break;
+				}
+				else if (ExistingCompletedState.StateID == NewFinishedTask.StateID)
+				{
+					if (NewFinishedTask.Type == EStateTreeFinishTaskType::Failed)
+					{
+						// Set failure, otherwise keep as it is
+						ExistingCompletedState.RunStatus = EStateTreeRunStatus::Failed;
+					}
+					bAdd = false;
+					break;
 				}
-				bAdd = false;
-				break;
 			}
 		}
 	}
 
 	if (bAdd)
 	{
-		const EStateTreeRunStatus RunStatus = UE::StateTree::ExecutionContext::Private::FinishTaskTypeToRunStatus(DelayedFinishedTask.Type);
-		CompletedStates.Add(FCompletedState{.ActiveFrameIndex = ActiveFrameIndex, .ActiveStateIndex = ActiveStateIndex, .RunStatus = RunStatus, .DelayedTask = DelayedFinishedTask});
+		FCompletedState& CompletedState = CompletedStates.Add_GetRef(FCompletedState(NewFinishedTask));
+		CompletedState.RunStatus = UE::StateTree::ExecutionContext::Private::FinishTaskTypeToRunStatus(NewFinishedTask.Type);
 	}
 }
 
@@ -1415,7 +1400,7 @@ EStateTreeRunStatus FStateTreeExecutionContext::GetGlobalTasksCompletedStatesSta
 	EStateTreeRunStatus Result = EStateTreeRunStatus::Running;
 	for (const FCompletedState& CompletedState : CompletedStates)
 	{
-		if (CompletedState.ActiveStateIndex == INDEX_NONE)
+		if (CompletedState.Reason == UE::StateTree::FDelayedFinishedTask::EReasonType::GlobalTask)
 		{
 			// return failure in priority
 			if (CompletedState.RunStatus == EStateTreeRunStatus::Failed)
@@ -2277,6 +2262,7 @@ EStateTreeRunStatus FStateTreeExecutionContext::EnterState(FStateTreeTransitionR
 
 	// Reset the current active frames, new ones are added one by one.
 	Exec.ActiveFrames.Reset();
+	CompletedStateRunStatus.Reset();
 
 	// Track any changed state (i.e., not sustained) to prevent reused subtrees from being considered as sustained states since the whole
 	// tree is reused and should be considered as a new instance.
@@ -2446,7 +2432,7 @@ EStateTreeRunStatus FStateTreeExecutionContext::EnterState(FStateTreeTransitionR
 					if (Status != EStateTreeRunStatus::Running)
 					{
 						Result = Status;
-						FinishTaskInternal(FrameIndex, StateIndex, TaskIndex, Result);
+						FinishTaskInternal(UE::StateTree::FDelayedFinishedTask::EReasonType::StateTask, CurrentFrame.FrameID, CurrentFrame.ActiveStates.StateIDs[StateIndex], TaskIndex, Result);
 						if (Status == EStateTreeRunStatus::Failed)
 						{
 							Exec.EnterStateFailedFrameIndex = FStateTreeIndex16(FrameIndex);
@@ -2734,6 +2720,15 @@ void FStateTreeExecutionContext::ExitState(const FStateTreeTransitionResult& Tra
 						}
 					}
 				}
+
+				// Clean up completed state to prevent FinishTask called in ExitState.
+				CompletedStates.RemoveAllSwap([StateID = CurrentFrame.ActiveStates.StateIDs[StateIndex]](const FCompletedState& Completed)
+					{
+						return Completed.StateID == StateID;
+					});
+
+				// Delegate Listeners Cleanup
+				GetExecState().DelegateActiveListeners.RemoveAll(CurrentFrame.ActiveStates.StateIDs[StateIndex]);
 			}
 
 			UE_STATETREE_DEBUG_STATE_EVENT(this, CurrentHandle, EStateTreeTraceEventType::OnExited);
@@ -2746,56 +2741,20 @@ void FStateTreeExecutionContext::ExitState(const FStateTreeTransitionResult& Tra
 			{
 				CurrentTransition.ChangeType = ExitCall.ChangeType;
 				CallStopOnEvaluatorsAndGlobalTasks(CurrentParentFrame, CurrentFrame, CurrentTransition);
-			}
-		}
-	}
 
-	// Delegate Listeners Cleanup
-	{
-		if (Transition.NextActiveFrames.IsEmpty())
-		{
-			RemoveAllDelegateListeners();
-		}
-		else if (ensure(FirstModifiedFrame != INDEX_NONE && FirstModifiedState != INDEX_NONE))
-		{
-			const FStateTreeExecutionFrame* ChangedFrame = &Exec.ActiveFrames[FirstModifiedFrame];
-			FStateTreeStateHandle ChangedState = ChangedFrame->ActiveStates[FirstModifiedState];
-			TConstArrayView<FStateTreeExecutionFrame> DeactivatedFrames;
+				// Clean up completed state to prevent FinishTask called in ExitState.
+				CompletedStates.RemoveAllSwap([FrameID = CurrentFrame.FrameID](const FCompletedState& Completed) 
+					{
+						return Completed.FrameID == FrameID;
+					});
 
-			if (Exec.ActiveFrames.Num() > FirstModifiedFrame + 1)
-			{
-				DeactivatedFrames = TConstArrayView<FStateTreeExecutionFrame>(&Exec.ActiveFrames[FirstModifiedFrame + 1], Exec.ActiveFrames.Num() - FirstModifiedFrame - 1);
+				// Delegate Listeners Cleanup
+				GetExecState().DelegateActiveListeners.RemoveAll(CurrentFrame.FrameID);
 			}
-
-			RemoveDeactivatedDelegateListeners(ChangedFrame, ChangedState, DeactivatedFrames);
 		}
 	}
 }
 
-void FStateTreeExecutionContext::RemoveDeactivatedDelegateListeners(const FStateTreeExecutionFrame* ChangedFrame, FStateTreeStateHandle ChangedState, TConstArrayView<FStateTreeExecutionFrame> DeactivatedFrames)
-{
-	GetExecState().DelegateActiveListeners.RemoveByPredicate([&](FStateTreeStateHandle OwningState, FStateTreeIndex16 OwningNodeIndex, const UStateTree* FrameStateTree, FStateTreeStateHandle FrameRootState)
-		{
-			if (FrameStateTree == nullptr)
-			{
-				return true;
-			}
-
-			if (ChangedFrame && ChangedFrame->RootState == FrameRootState)
-			{
-				check(ChangedState.IsValid());
-				const bool bIsGlobalListener = !OwningState.IsValid();
-				return !bIsGlobalListener && OwningState.Index >= ChangedState.Index;
-			}
-
-			return DeactivatedFrames.ContainsByPredicate([&](const FStateTreeExecutionFrame& Frame)
-				{
-					return Frame.RootState == FrameRootState && Frame.StateTree == FrameStateTree;
-				});
-
-		});
-}
-
 void FStateTreeExecutionContext::RemoveAllDelegateListeners()
 {
 	GetExecState().DelegateActiveListeners = FStateTreeDelegateActiveListeners();
@@ -2904,6 +2863,7 @@ EStateTreeRunStatus FStateTreeExecutionContext::TickEvaluatorsAndGlobalTasks(con
 	FStateTreeExecutionState& Exec = GetExecState();
 
 	EStateTreeRunStatus Result = EStateTreeRunStatus::Running;
+	CompletedStateRunStatus.Reset();
 
 	for (int32 FrameIndex = 0; FrameIndex < Exec.ActiveFrames.Num(); FrameIndex++)
 	{
@@ -2940,7 +2900,10 @@ EStateTreeRunStatus FStateTreeExecutionContext::TickEvaluatorsAndGlobalTasks(con
 			// Used to stop ticking tasks after one fails, but we still want to keep updating the data views so that property binding works properly.
 			bool bShouldTickTasks = true;
 
-			const FCompletedState* FrameCompletedState = CompletedStates.FindByPredicate([FrameIndex](FCompletedState& Other) { return Other.ActiveFrameIndex == FrameIndex && Other.ActiveStateIndex == INDEX_NONE; });
+			const FCompletedState* FrameCompletedState = CompletedStates.FindByPredicate([FrameID = CurrentFrame.FrameID](FCompletedState& Other)
+				{
+					return Other.FrameID == FrameID && Other.Reason == UE::StateTree::FDelayedFinishedTask::EReasonType::GlobalTask;
+				});
 			if (FrameCompletedState)
 			{
 				bShouldTickTasks = false;
@@ -3413,6 +3376,7 @@ EStateTreeRunStatus FStateTreeExecutionContext::TickTasks(const float DeltaTime)
 	
 	// Used to stop ticking tasks after one fails, but we still want to keep updating the data views so that property binding works properly.
 	bool bShouldTickTasks = CompletedStates.Num() == 0;
+	CompletedStateRunStatus.Reset();
 
 	STATETREE_CLOG(Exec.ActiveFrames.Num() > 0, VeryVerbose, TEXT("Ticking Tasks"));
 
@@ -3504,7 +3468,7 @@ EStateTreeRunStatus FStateTreeExecutionContext::TickTasks(const float DeltaTime)
 					if (TaskResult != EStateTreeRunStatus::Running)
 					{
 						Result = TaskResult;
-						FinishTaskInternal(FrameIndex, StateIndex, TaskIndex, TaskResult);
+						FinishTaskInternal(UE::StateTree::FDelayedFinishedTask::EReasonType::StateTask, CurrentFrame.FrameID, CurrentFrame.ActiveStates.StateIDs[StateIndex], TaskIndex, TaskResult);
 						if (TaskResult == EStateTreeRunStatus::Failed)
 						{
 							bShouldTickTasks = false;
@@ -3528,9 +3492,17 @@ EStateTreeRunStatus FStateTreeExecutionContext::TickTasks(const float DeltaTime)
 	else if (NumTotalEnabledTasks == 0 && CompletedStates.Num() == 0)
 	{
 		// No enabled tasks, done ticking. Check from the bottom state in the bottom frame to trigger completion transitions
-		const int32 CompletedFrameIndex = Exec.ActiveFrames.Num() > 0 ? Exec.ActiveFrames.Num() - 1 : INDEX_NONE;
-		const int32 CompletedStateIndex = Exec.ActiveFrames.Num() > 0 ? (int32)Exec.ActiveFrames.Last().ActiveStates.Num() - 1 : INDEX_NONE;
-		FinishTaskInternal(CompletedFrameIndex, CompletedStateIndex, INDEX_NONE, EStateTreeRunStatus::Succeeded);
+		if (Exec.ActiveFrames.Num() > 0)
+		{
+			const FStateTreeExecutionFrame& LastFrame = Exec.ActiveFrames.Last();
+			const UE::StateTree::FActiveFrameID CompletedFrameID = LastFrame.FrameID;
+			const UE::StateTree::FActiveStateID CompletedStateID = LastFrame.ActiveStates.StateIDs[LastFrame.ActiveStates.Num() - 1];
+			FinishTaskInternal(UE::StateTree::FDelayedFinishedTask::EReasonType::InternalTransition, CompletedFrameID, CompletedStateID, INDEX_NONE, EStateTreeRunStatus::Succeeded);
+		}
+		else
+		{
+			Exec.RequestedStop = EStateTreeRunStatus::Succeeded;
+		}
 		Result = EStateTreeRunStatus::Succeeded;
 	}
 
@@ -4295,33 +4267,33 @@ bool FStateTreeExecutionContext::TriggerTransitions()
 		&& Exec.LastTickStatus != EStateTreeRunStatus::Running)
 	{
 		// Sort CompletedStates, only keep the smaller frame and the active state of the smaller frame.
-		TOptional<const FCompletedState> FirstCompletedState;
-		{
-			for (const FCompletedState& CompletedState : CompletedStates)
-			{
-				if (CompletedState.ActiveStateIndex == INDEX_NONE)
-				{
-					// Global task are not triggered in the transition loop.
-					FirstCompletedState.Reset();
-					break;
-				}
-				if (!FirstCompletedState.IsSet()
-					|| CompletedState.ActiveFrameIndex < FirstCompletedState.GetValue().ActiveFrameIndex
-					|| (CompletedState.ActiveFrameIndex == FirstCompletedState.GetValue().ActiveFrameIndex && CompletedState.ActiveStateIndex < FirstCompletedState.GetValue().ActiveStateIndex)
-					)
+		int32 CompletedFrameStartIndex = INT_MAX;
+		int32 CompletedStateStartIndex = INT_MAX;
+		EStateTreeRunStatus CompletedExpectedRunStatus = Exec.LastTickStatus;
+		for (const FCompletedState& CompletedState : CompletedStates)
+		{
+			// Global task are not triggered in the transition loop.
+			if (CompletedState.Reason != UE::StateTree::FDelayedFinishedTask::EReasonType::GlobalTask)
+			{
+				const int32 NewFrameIndex = Exec.IndexOfActiveFrame(CompletedState.FrameID);
+				check(NewFrameIndex != INDEX_NONE); // the array is supposed to be clear every loop
+				const int32 NewStateIndex = Exec.ActiveFrames[NewFrameIndex].ActiveStates.IndexOfReverse(CompletedState.StateID);
+				const int32 NewLoopStateIndex = NewStateIndex == INDEX_NONE ? UINT_MAX : NewStateIndex;
+				if (NewFrameIndex < CompletedFrameStartIndex
+					|| (NewFrameIndex == CompletedFrameStartIndex && NewLoopStateIndex < CompletedStateStartIndex))
 				{
-					FirstCompletedState = CompletedState;
+					CompletedFrameStartIndex = NewFrameIndex;
+					CompletedStateStartIndex = NewLoopStateIndex;
+					CompletedExpectedRunStatus = CompletedState.RunStatus;
 				}
 			}
-
-			Exec.DelayedFinishedTasks.Empty();
-			CompletedStates.Empty();
 		}
 
+		Exec.DelayedFinishedTasks.Empty();
+		CompletedStates.Empty();
+
 		// Start from the last completed state if specified, else only check the last state from the last frame.
-		const int32 FrameStartIndex = FirstCompletedState.IsSet() ? FirstCompletedState.GetValue().ActiveFrameIndex : (Exec.ActiveFrames.Num() - 1);
-		check(FrameStartIndex >= 0 && FrameStartIndex < Exec.ActiveFrames.Num());
-		
+		const int32 FrameStartIndex = FMath::Min(CompletedFrameStartIndex, Exec.ActiveFrames.Num() - 1);
 		for (int32 FrameIndex = FrameStartIndex; FrameIndex >= 0; FrameIndex--)
 		{
 			FStateTreeExecutionFrame* CurrentParentFrame = FrameIndex > 0 ? &Exec.ActiveFrames[FrameIndex - 1] : nullptr;
@@ -4331,14 +4303,14 @@ bool FStateTreeExecutionContext::TriggerTransitions()
 			FCurrentlyProcessedFrameScope FrameScope(*this, CurrentParentFrame, CurrentFrame);
 
 			int32 StateStartIndex = CurrentFrame.ActiveStates.Num() - 1; // This is ok, even if the ActiveStates is 0, -1 will skip the whole state loop below.
-			if (FrameIndex == FrameStartIndex && FirstCompletedState.IsSet())
+			if (FrameIndex == FrameStartIndex && CompletedStateStartIndex < CurrentFrame.ActiveStates.Num())
 			{
-				StateStartIndex = FirstCompletedState.GetValue().ActiveStateIndex;
+				StateStartIndex = CompletedStateStartIndex;
 			}
 
 			// Use worst case of LastTickStatus and FoundCompletedState.
 			const bool bLastTickWasSuccess = Exec.LastTickStatus == EStateTreeRunStatus::Succeeded;
-			const EStateTreeRunStatus ExpectedRunStatus = bLastTickWasSuccess && FirstCompletedState.IsSet() ? FirstCompletedState.GetValue().RunStatus : Exec.LastTickStatus;
+			const EStateTreeRunStatus ExpectedRunStatus = bLastTickWasSuccess && FrameIndex == FrameStartIndex ? CompletedExpectedRunStatus : Exec.LastTickStatus;
 			const EStateTreeTransitionTrigger CompletionTrigger = ExpectedRunStatus == EStateTreeRunStatus::Succeeded ? EStateTreeTransitionTrigger::OnStateSucceeded : EStateTreeTransitionTrigger::OnStateFailed;
 
 			// Check completion transitions
@@ -4458,7 +4430,8 @@ bool FStateTreeExecutionContext::TriggerTransitions()
 
 				// Set the parent linked state as last completed state, and update tick status to the status from the transition.
 				const int32 TaskIndex = INDEX_NONE;
-				FinishTaskInternal(ParentFrameIndex, ParentFrame.ActiveStates.Num() - 1, TaskIndex, RunStatus);
+				const UE::StateTree::FActiveStateID ParentLastActiveStateID = ParentFrame.ActiveStates.StateIDs[ParentFrame.ActiveStates.Num() - 1];
+				FinishTaskInternal(UE::StateTree::FDelayedFinishedTask::EReasonType::InternalTransition, ParentFrame.FrameID, ParentLastActiveStateID, TaskIndex, RunStatus);
 				Exec.LastTickStatus = RunStatus;
 
 				// Clear the transition and return that no transition took place.
@@ -4983,7 +4956,7 @@ bool FStateTreeExecutionContext::SelectStateInternal(
 						__FUNCTION__, *GetSafeStateName(CurrentFrame, NextStateHandle), *GetFullNameSafe(NewFrame.StateTree), *GetNameSafe(&Owner), *GetFullNameSafe(CurrentFrame.StateTree));
 					
 					StopTemporaryEvaluatorsAndGlobalTasks(nullptr, NewFrame);
-					RemoveDeactivatedDelegateListeners(nullptr, FStateTreeStateHandle::Invalid, {NewFrame});
+					GetExecState().DelegateActiveListeners.RemoveAll(NewFrame.FrameID);
 
 					break;
 				}
@@ -5003,7 +4976,7 @@ bool FStateTreeExecutionContext::SelectStateInternal(
 			if (bStartedTemporaryEvaluatorsAndGlobalTasks)
 			{
 				StopTemporaryEvaluatorsAndGlobalTasks(nullptr, NewFrame);
-				RemoveDeactivatedDelegateListeners(nullptr, FStateTreeStateHandle::Invalid, {NewFrame});
+				GetExecState().DelegateActiveListeners.RemoveAll(NewFrame.FrameID);
 			}
 				
 			OutSelectionResult.PopFrame();
@@ -5657,9 +5630,12 @@ FStateTreeWeakExecutionContext::FStateTreeWeakExecutionContext(const FStateTreeE
 {
 	if (const FStateTreeExecutionFrame* Frame = Context.GetCurrentlyProcessedFrame())
 	{
-		FrameStateTree = Frame->StateTree;
-		FrameRootState = Frame->RootState;
-		CurrentState = Context.GetCurrentlyProcessedState();
+		FrameID = Frame->FrameID;
+		const int32 ActiveStateIndex = Frame->ActiveStates.IndexOfReverse(Context.GetCurrentlyProcessedState());
+		if (ActiveStateIndex != INDEX_NONE)
+		{
+			StateID = Frame->ActiveStates.StateIDs[ActiveStateIndex];
+		}
 	}
 }
 
@@ -5670,8 +5646,17 @@ bool FStateTreeWeakExecutionContext::SendEvent(const FGameplayTag Tag, const FCo
 	TSharedPtr<FStateTreeInstanceStorage> StoragePinned = Storage.Pin();
 	if (OwnerPinned && StateTreePinned && StoragePinned)
 	{
-		FStateTreeMinimalExecutionContext Context(*OwnerPinned.Get(), *StateTreePinned.Get(), *StoragePinned.Get());
-		Context.SendEvent(Tag, Payload, Origin);
+		// Manually acquire and release the write access.
+		//The context may already exist and we want to confirm that we are on the same thread.
+		StoragePinned->AcquireWriteAccess();
+
+		if (IsFrameOrStateActive(StoragePinned->GetExecutionState()))
+		{
+			FStateTreeMinimalExecutionContext Context(*OwnerPinned.Get(), *StateTreePinned.Get(), *StoragePinned.Get());
+			Context.SendEvent(Tag, Payload, Origin);
+		}
+
+		StoragePinned->ReleaseWriteAccess();
 		return true;
 	}
 	return false;
@@ -5679,6 +5664,8 @@ bool FStateTreeWeakExecutionContext::SendEvent(const FGameplayTag Tag, const FCo
 
 bool FStateTreeWeakExecutionContext::RequestTransition(FStateTreeStateHandle TargetState, EStateTreeTransitionPriority Priority, const EStateTreeSelectionFallback Fallback) const
 {
+	bool bRequested = false;
+
 	TStrongObjectPtr<UObject> OwnerPinned = Owner.Pin();
 	TStrongObjectPtr<const UStateTree> StateTreePinned = StateTree.Pin();
 	TSharedPtr<FStateTreeInstanceStorage> StoragePinned = Storage.Pin();
@@ -5689,19 +5676,24 @@ bool FStateTreeWeakExecutionContext::RequestTransition(FStateTreeStateHandle Tar
 		StoragePinned->AcquireWriteAccess();
 
 		FStateTreeExecutionState& Exec = StoragePinned->GetMutableExecutionState();
-		if (Exec.CurrentExecutionContext)
-		{
-			Exec.CurrentExecutionContext->RequestTransition(TargetState, Priority, Fallback);
-		}
-		else if (Exec.ActiveFrames.Num() > 0)
+		if (IsFrameOrStateActive(Exec) && Exec.ActiveFrames.Num() > 0)
 		{
-			const FStateTreeExecutionFrame* RootFrame = &Exec.ActiveFrames[0];
-
 			FStateTreeTransitionRequest Request;
-			//todo: use the StatePath for the correct source
-			Request.SourceStateTree = RootFrame->StateTree;
-			Request.SourceRootState = RootFrame->ActiveStates[0];
-			Request.SourceState = FStateTreeStateHandle();
+
+			if (const FStateTreeExecutionFrame* CurrentFrame = Exec.FindActiveFrame(FrameID))
+			{
+				Request.SourceStateTree = CurrentFrame->StateTree;
+				Request.SourceRootState = CurrentFrame->RootState;
+				const int32 IndexOfActiveState = CurrentFrame->ActiveStates.IndexOfReverse(StateID);
+				Request.SourceState = IndexOfActiveState != INDEX_NONE ? CurrentFrame->ActiveStates[IndexOfActiveState] : FStateTreeStateHandle();
+			}
+			else
+			{
+				const FStateTreeExecutionFrame* RootFrame = &Exec.ActiveFrames[0];
+				Request.SourceStateTree = RootFrame->StateTree;
+				Request.SourceRootState = RootFrame->RootState;
+				Request.SourceState = FStateTreeStateHandle();
+			}
 
 			Request.TargetState = TargetState;
 			Request.Priority = Priority;
@@ -5710,12 +5702,13 @@ bool FStateTreeWeakExecutionContext::RequestTransition(FStateTreeStateHandle Tar
 			StoragePinned->AddTransitionRequest(OwnerPinned.Get(), Request);
 
 			ScheduleNextTick(OwnerPinned.Get(), StateTreePinned.Get(), *StoragePinned);
+			bRequested = true;
 		}
 
 		StoragePinned->ReleaseWriteAccess();
-		return true;
 	}
-	return false;
+
+	return bRequested;
 }
 
 bool FStateTreeWeakExecutionContext::BroadcastDelegate(const FStateTreeDelegateDispatcher& Dispatcher) const
@@ -5738,25 +5731,14 @@ bool FStateTreeWeakExecutionContext::BroadcastDelegate(const FStateTreeDelegateD
 		StoragePinned->AcquireWriteAccess();
 
 		FStateTreeExecutionState& Exec = StoragePinned->GetMutableExecutionState();
-		if (Exec.CurrentExecutionContext)
-		{
-			Exec.CurrentExecutionContext->BroadcastDelegate(Dispatcher);
-			bBroadcasted = true;
-		}
-		else if (Exec.ActiveFrames.Num() > 0)
+		if (IsFrameOrStateActive(Exec))
 		{
-			TStrongObjectPtr<const UStateTree> FrameStateTreePinned = FrameStateTree.Pin();
-			if (FrameStateTreePinned)
+			if (const FStateTreeExecutionFrame* CurrentFrame = Exec.FindActiveFrame(FrameID))
 			{
-				const FStateTreeExecutionFrame* ParentFrame = nullptr;
-				const FStateTreeExecutionFrame* CurrentFrame = FStateTreeExecutionContext::FindFrame(FrameStateTreePinned.Get(), FrameRootState, Exec.ActiveFrames, ParentFrame);
-				if (CurrentFrame)
+				Exec.DelegateActiveListeners.BroadcastDelegate(Dispatcher, Exec);
+				if (UE::StateTree::ExecutionContext::Private::MarkDelegateAsBroadcasted(Dispatcher, *CurrentFrame, *StoragePinned))
 				{
-					Exec.DelegateActiveListeners.BroadcastDelegate(Dispatcher, Exec);
-					if (UE::StateTree::ExecutionContext::Private::MarkDelegateAsBroadcasted(Dispatcher, *CurrentFrame, *StoragePinned))
-					{
-						ScheduleNextTick(OwnerPinned.Get(), RootStateTreePinned.Get(), *StoragePinned);
-					}
+					ScheduleNextTick(OwnerPinned.Get(), RootStateTreePinned.Get(), *StoragePinned);
 				}
 				bBroadcasted = true;
 			}
@@ -5788,20 +5770,17 @@ bool FStateTreeWeakExecutionContext::BindDelegate(const FStateTreeWeakTaskRef& T
 		StoragePinned->AcquireWriteAccess();
 
 		FStateTreeExecutionState& Exec = StoragePinned->GetMutableExecutionState();
-		if (Exec.CurrentExecutionContext)
-		{
-			Exec.CurrentExecutionContext->BindDelegate(Listener, MoveTemp(Delegate));
-			bResult = true;
-		}
-		else if (Exec.ActiveFrames.Num() > 0)
+		if (IsFrameOrStateActive(Exec))
 		{
-			FStateTreeStrongTaskRef TaskPin = Task.Pin();
-			const FStateTreeTaskBase* TaskPtr = TaskPin.GetTask();
-			TStrongObjectPtr<const UStateTree> FrameStateTreePinned = FrameStateTree.Pin();
-			if (TaskPtr != nullptr && FrameStateTreePinned.Get() != nullptr && TaskPin.GetStateTree() == FrameStateTreePinned.Get())
+			if (const FStateTreeExecutionFrame* CurrentFrame = Exec.FindActiveFrame(FrameID))
 			{
-				Exec.DelegateActiveListeners.Add(Listener, MoveTemp(Delegate), FrameStateTreePinned.Get(), FrameRootState, CurrentState, TaskPin.GetTaskIndex());
-				bResult = true;
+				const FStateTreeStrongTaskRef TaskPin = Task.Pin();
+				// The user uses a task that is not part of this context (ex: use task from state A inside context constructed from state B)
+				if (ensure(TaskPin.IsValid() && TaskPin.GetStateTree() == CurrentFrame->StateTree))
+				{
+					Exec.DelegateActiveListeners.Add(Listener, MoveTemp(Delegate), FrameID, StateID, TaskPin.GetTaskIndex());
+					bResult = true;
+				}
 			}
 		}
 
@@ -5836,12 +5815,7 @@ bool FStateTreeWeakExecutionContext::UnbindDelegate(const FStateTreeDelegateList
 		StoragePinned->AcquireWriteAccess();
 
 		FStateTreeExecutionState& Exec = StoragePinned->GetMutableExecutionState();
-		if (Exec.CurrentExecutionContext)
-		{
-			Exec.CurrentExecutionContext->UnbindDelegate(Listener);
-			bRemoved = true;
-		}
-		else if (Exec.ActiveFrames.Num() > 0)
+		if (IsFrameOrStateActive(Exec))
 		{
 			Exec.DelegateActiveListeners.Remove(Listener);
 			bRemoved = true;
@@ -5866,31 +5840,33 @@ bool FStateTreeWeakExecutionContext::FinishTask(const FStateTreeWeakTaskRef& Tas
 		//The context may already exist and we want to confirm that we are on the same thread.
 		StoragePinned->AcquireWriteAccess();
 
-		FStateTreeStrongTaskRef TaskPin = Task.Pin();
-		const FStateTreeTaskBase* TaskPtr = TaskPin.GetTask();
-		TStrongObjectPtr<const UStateTree> FrameStateTreePinned = FrameStateTree.Pin();
-		if (TaskPtr != nullptr && FrameStateTreePinned.Get() != nullptr && TaskPin.GetStateTree() == FrameStateTreePinned.Get())
+		FStateTreeExecutionState& Exec = StoragePinned->GetMutableExecutionState();
+		if (IsFrameOrStateActive(Exec))
 		{
-			FStateTreeDelayedFinishedTask CompletedTask{
-				.FrameStateTree = FObjectKey(FrameStateTreePinned.Get()), 
-				.FrameRootState = FrameRootState,
-				.State = CurrentState,
-				.TaskIndex = TaskPin.GetTaskIndex(),
-				.Type = FinishType
-				};
-
-			FStateTreeExecutionState& Exec = StoragePinned->GetMutableExecutionState();
-			if (Exec.CurrentExecutionContext)
-			{
-				Exec.CurrentExecutionContext->FinishTask(CompletedTask, FinishType);
-				bCompleted = true;
-			}
-			else
+			FStateTreeStrongTaskRef TaskPin = Task.Pin();
+			const FStateTreeExecutionFrame* ExecutionFrame = Exec.FindActiveFrame(FrameID);
+			if (TaskPin.IsValid() && ExecutionFrame)
 			{
-				// Confirm that the task if from the correct state.
-				Exec.DelayedFinishedTasks.Add(CompletedTask);
-				ScheduleNextTick(OwnerPinned.Get(), StateTreePinned.Get(), *StoragePinned);
-				bCompleted = true;
+				const bool bIsSameTree = ExecutionFrame->StateTree == TaskPin.GetStateTree();
+				if (ensureMsgf(bIsSameTree, TEXT("User error. The task is not from the same execution tree")))
+				{
+					const int32 GlobalTaskBeginIndex = ExecutionFrame->StateTree->GlobalTasksBegin;
+					const int32 GlobalTaskEndIndex = ExecutionFrame->StateTree->GlobalTasksBegin + ExecutionFrame->StateTree->GlobalTasksNum;
+					const bool bIsGlobalTask = TaskPin.GetTaskIndex().AsInt32() >= GlobalTaskBeginIndex && TaskPin.GetTaskIndex().AsInt32() < GlobalTaskEndIndex;
+					const UE::StateTree::FDelayedFinishedTask::EReasonType Reason = bIsGlobalTask ? UE::StateTree::FDelayedFinishedTask::EReasonType::GlobalTask : UE::StateTree::FDelayedFinishedTask::EReasonType::StateTask;
+					UE::StateTree::FDelayedFinishedTask CompletedTask{
+						.FrameID = FrameID,
+						.StateID = StateID,
+						.TaskIndex = TaskPin.GetTaskIndex(),
+						.Type = FinishType,
+						.Reason = Reason
+					};
+
+					// Confirm that the task if from the correct state.
+					Exec.DelayedFinishedTasks.Add(CompletedTask);
+					ScheduleNextTick(OwnerPinned.Get(), StateTreePinned.Get(), *StoragePinned);
+					bCompleted = true;
+				}
 			}
 		}
 
@@ -5917,11 +5893,7 @@ bool FStateTreeWeakExecutionContext::UpdateScheduledTickRequest(UE::StateTree::F
 		bCompleted = true;
 
 		FStateTreeExecutionState& Exec = StoragePinned->GetMutableExecutionState();
-		if (Exec.CurrentExecutionContext)
-		{
-			Exec.CurrentExecutionContext->UpdateScheduledTickRequest(Handle, ScheduledTick);
-		}
-		else
+		if (IsFrameOrStateActive(Exec))
 		{
 			if (Exec.UpdateScheduledTickRequest(Handle, ScheduledTick))
 			{
@@ -5944,6 +5916,22 @@ void FStateTreeWeakExecutionContext::ScheduleNextTick(TNotNull<UObject*> Owner,
 	}
 }
 
+bool FStateTreeWeakExecutionContext::IsFrameOrStateActive(const FStateTreeExecutionState& Exec) const
+{
+	if (!FrameID.IsValid())
+	{
+		return true;
+	}
+
+	const FStateTreeExecutionFrame* CurrentFrame = Exec.FindActiveFrame(FrameID);
+	if (CurrentFrame == nullptr)
+	{
+		return false;
+	}
+
+	return StateID.IsValid() ? CurrentFrame->ActiveStates.Contains(StateID) : true;
+}
+
 #undef STATETREE_LOG
 #undef STATETREE_CLOG
 #undef UE_STATETREE_ENSURE_COMPLETED_STATE_RUN_STATUS
diff --git a/Engine/Plugins/Runtime/StateTree/Source/StateTreeModule/Private/StateTreeExecutionTypes.cpp b/Engine/Plugins/Runtime/StateTree/Source/StateTreeModule/Private/StateTreeExecutionTypes.cpp
--- a/Engine/Plugins/Runtime/StateTree/Source/StateTreeModule/Private/StateTreeExecutionTypes.cpp
+++ b/Engine/Plugins/Runtime/StateTree/Source/StateTreeModule/Private/StateTreeExecutionTypes.cpp
@@ -122,6 +122,30 @@ UE::StateTree::FActiveStatePath FStateTreeExecutionState::GetActiveStatePath() c
 	return UE::StateTree::FActiveStatePath(ActiveFrames[0].StateTree, MoveTemp(Elements));
 }
 
+const FStateTreeExecutionFrame* FStateTreeExecutionState::FindActiveFrame(UE::StateTree::FActiveFrameID FrameID) const
+{
+	return ActiveFrames.FindByPredicate([FrameID](const FStateTreeExecutionFrame& Other)
+		{
+			return Other.FrameID == FrameID;
+		});
+}
+
+FStateTreeExecutionFrame* FStateTreeExecutionState::FindActiveFrame(UE::StateTree::FActiveFrameID FrameID)
+{
+	return ActiveFrames.FindByPredicate([FrameID](const FStateTreeExecutionFrame& Other)
+		{
+			return Other.FrameID == FrameID;
+		});
+}
+
+int32 FStateTreeExecutionState::IndexOfActiveFrame(UE::StateTree::FActiveFrameID FrameID) const
+{
+	return ActiveFrames.IndexOfByPredicate([FrameID](const FStateTreeExecutionFrame& Other)
+		{
+			return Other.FrameID == FrameID;
+		});
+}
+
 UE::StateTree::FScheduledTickHandle FStateTreeExecutionState::AddScheduledTickRequest(FStateTreeScheduledTick ScheduledTick)
 {
 	UE::StateTree::FScheduledTickHandle Result = UE::StateTree::FScheduledTickHandle::GenerateNewHandle();
@@ -296,13 +320,12 @@ UE::StateTree::FScheduledTickHandle UE::StateTree::FScheduledTickHandle::Generat
 //----------------------------------------------------------------------//
 // FStateTreeDelegateActiveListeners
 //----------------------------------------------------------------------//
-FStateTreeDelegateActiveListeners::FActiveListener::FActiveListener(const FStateTreeDelegateListener& InListener, FSimpleDelegate InDelegate, TNotNull<const UStateTree*> InFrameStateTree, FStateTreeStateHandle InFrameRootState, FStateTreeStateHandle InOwningState, FStateTreeIndex16 InOwningNodeIndex)
+FStateTreeDelegateActiveListeners::FActiveListener::FActiveListener(const FStateTreeDelegateListener& InListener, FSimpleDelegate InDelegate, UE::StateTree::FActiveFrameID InFrameID, UE::StateTree::FActiveStateID InStateID, FStateTreeIndex16 InOwningNodeIndex)
 	: Listener(InListener)
 	, Delegate(MoveTemp(InDelegate))
-	, OwningState(InOwningState)
+	, FrameID(InFrameID)
+	, StateID(InStateID)
 	, OwningNodeIndex(InOwningNodeIndex)
-	, FrameStateTree(InFrameStateTree)
-	, FrameRootState(InFrameRootState)
 {}
 
 FStateTreeDelegateActiveListeners::~FStateTreeDelegateActiveListeners()
@@ -310,11 +333,11 @@ FStateTreeDelegateActiveListeners::~FStateTreeDelegateActiveListeners()
 	check(BroadcastingLockCount == 0);
 }
 
-void FStateTreeDelegateActiveListeners::Add(const FStateTreeDelegateListener& Listener, FSimpleDelegate Delegate, TNotNull<const UStateTree*> FrameStateTree, FStateTreeStateHandle FrameRootState, FStateTreeStateHandle OwningState, FStateTreeIndex16 OwningNodeIndex)
+void FStateTreeDelegateActiveListeners::Add(const FStateTreeDelegateListener& Listener, FSimpleDelegate Delegate, UE::StateTree::FActiveFrameID InFrameID, UE::StateTree::FActiveStateID InStateID, FStateTreeIndex16 OwningNodeIndex)
 {
 	check(Listener.IsValid());
 	Remove(Listener);
-	Listeners.Emplace(Listener, MoveTemp(Delegate), FrameStateTree, FrameRootState, OwningState, OwningNodeIndex);
+	Listeners.Emplace(Listener, MoveTemp(Delegate), InFrameID, InStateID, OwningNodeIndex);
 }
 
 void FStateTreeDelegateActiveListeners::Remove(const FStateTreeDelegateListener& Listener)
@@ -342,12 +365,21 @@ void FStateTreeDelegateActiveListeners::Remove(const FStateTreeDelegateListener&
 	}
 }
 
-void FStateTreeDelegateActiveListeners::RemoveByPredicate(TFunctionRef<bool(FStateTreeStateHandle OwningState, FStateTreeIndex16 OwningNodeIndex, const UStateTree* FrameStateTree, FStateTreeStateHandle FrameRootState)> Predicate)
+void FStateTreeDelegateActiveListeners::RemoveAll(UE::StateTree::FActiveFrameID FrameID)
+{
+	check(BroadcastingLockCount == 0);
+	Listeners.RemoveAllSwap([FrameID](const FActiveListener& Listener)
+		{
+			return Listener.FrameID == FrameID;
+		});
+}
+
+void FStateTreeDelegateActiveListeners::RemoveAll(UE::StateTree::FActiveStateID StateID)
 {
 	check(BroadcastingLockCount == 0);
-	Listeners.RemoveAllSwap([&](const FActiveListener& Listener)
+	Listeners.RemoveAllSwap([StateID](const FActiveListener& Listener)
 		{
-			return Predicate(Listener.OwningState, Listener.OwningNodeIndex, CastChecked<const UStateTree>(Listener.FrameStateTree.ResolveObjectPtr()), Listener.FrameRootState);
+			return Listener.StateID == StateID;
 		});
 }
 
@@ -360,18 +392,17 @@ void FStateTreeDelegateActiveListeners::BroadcastDelegate(const FStateTreeDelega
 	const int32 NumListeners = Listeners.Num();
 	for (int32 Index = 0; Index < NumListeners; ++Index)
 	{
-		FActiveListener& Listener = Listeners[Index];
-		if (Listener.Listener.GetDispatcher() == Dispatcher
-			&& Listener.OwningNodeIndex.Get() <= Exec.EnterStateFailedTaskIndex.Get()
-			&& Listener.OwningNodeIndex.Get() < Exec.LastExitedNodeIndex.Get())
+		FActiveListener& ActiveListener = Listeners[Index];
+		if (ActiveListener.Listener.GetDispatcher() == Dispatcher
+			&& ActiveListener.OwningNodeIndex.Get() <= Exec.EnterStateFailedTaskIndex.Get()
+			&& ActiveListener.OwningNodeIndex.Get() < Exec.LastExitedNodeIndex.Get())
 		{
-			if (Exec.ActiveFrames.ContainsByPredicate(
-				[&Listener](const FStateTreeExecutionFrame& Other)
-				{
-					return Other.RootState == Listener.FrameRootState && FObjectKey(Other.StateTree) == Listener.FrameStateTree;
-				}))
+			if (const FStateTreeExecutionFrame* ExectionFrame = Exec.FindActiveFrame(ActiveListener.FrameID))
 			{
-				Listener.Delegate.ExecuteIfBound();
+				if (!ActiveListener.StateID.IsValid() || ExectionFrame->ActiveStates.Contains(ActiveListener.StateID))
+				{
+					ActiveListener.Delegate.ExecuteIfBound();
+				}
 			}
 		}
 	}
diff --git a/Engine/Plugins/Runtime/StateTree/Source/StateTreeModule/Public/StateTree.h b/Engine/Plugins/Runtime/StateTree/Source/StateTreeModule/Public/StateTree.h
--- a/Engine/Plugins/Runtime/StateTree/Source/StateTreeModule/Public/StateTree.h
+++ b/Engine/Plugins/Runtime/StateTree/Source/StateTreeModule/Public/StateTree.h
@@ -389,6 +389,7 @@ private:
 	friend struct FStateTreeInstance;
 	friend struct FStateTreeExecutionContext;
 	friend struct FStateTreeMinimalExecutionContext;
+	friend struct FStateTreeWeakExecutionContext;
 #if WITH_EDITORONLY_DATA
 	friend struct FStateTreeCompiler;
 #endif
diff --git a/Engine/Plugins/Runtime/StateTree/Source/StateTreeModule/Public/StateTreeExecutionContext.h b/Engine/Plugins/Runtime/StateTree/Source/StateTreeModule/Public/StateTreeExecutionContext.h
--- a/Engine/Plugins/Runtime/StateTree/Source/StateTreeModule/Public/StateTreeExecutionContext.h
+++ b/Engine/Plugins/Runtime/StateTree/Source/StateTreeModule/Public/StateTreeExecutionContext.h
@@ -584,7 +584,7 @@ public:
 	 * If called during tick processing, then the state completes immediately.
 	 * If called outside of the tick processing, then the request is buffered and handled on the next tick.
 	 */
-	void FinishTask(const FStateTreeDelayedFinishedTask& Task, EStateTreeFinishTaskType FinishType);
+	void FinishTask(const UE::StateTree::FDelayedFinishedTask& Task, EStateTreeFinishTaskType FinishType);
 
 	/** @return data view of the specified handle relative to given frame. */
 	static FStateTreeDataView GetDataViewFromInstanceStorage(FStateTreeInstanceStorage& InstanceDataStorage, FStateTreeInstanceStorage* CurrentlyProcessedSharedInstanceStorage, const FStateTreeExecutionFrame* ParentFrame, const FStateTreeExecutionFrame& CurrentFrame, const FStateTreeDataHandle Handle);
@@ -696,14 +696,6 @@ protected:
 	 */
 	void ExitState(const FStateTreeTransitionResult& Transition);
 
-	/**
-	 * Removes all delegate listeners of deactivated frames and states.
-	 * @param ChangedFrame The frame where the change of the state took place.
-	 * @param ChangedState First state in ChangedFrame that became inactive.
-	 * @param DeactivatedFrames Frames that were entirely deactivated.
-	 */
-	void RemoveDeactivatedDelegateListeners(const FStateTreeExecutionFrame* ChangedFrame, FStateTreeStateHandle ChangedState, TConstArrayView<FStateTreeExecutionFrame> DeactivatedFrames);
-
 	/**
 	 * Removes all delegate listeners.
 	 */
@@ -801,7 +793,10 @@ protected:
 	bool TriggerTransitions();
 
 	/** Adds a task to the finished and delayed list. */
-	void FinishTaskInternal(int32 ActiveFrameIndex, int32 ActiveStateIndex, int32 TaskIndex, EStateTreeRunStatus RunStatus);
+	void FinishTaskInternal(UE::StateTree::FDelayedFinishedTask::EReasonType Reason, UE::StateTree::FActiveFrameID ActiveFrameID, UE::StateTree::FActiveStateID ActiveStateID, int32 TaskIndex, EStateTreeRunStatus RunStatus);
+
+	/** Confirms that the frame and state ID are valid and the task index is correct. */
+	bool IsDelayedFinishedTaskValid(const UE::StateTree::FDelayedFinishedTask& Task) const;
 
 	/**
 	 * Removes stale entries and fills the completed states list.
@@ -813,7 +808,7 @@ protected:
 	 * Adds uniquely the task to the finished.
 	 * Adds the task to the delayed list. The delayed list may contains duplicate. They will be filtered with UpdateCompletedStateList when ticked.
 	 */
-	void AddUniqueCompletedStateList(int32 ActiveFrameIndex, int32 ActiveStateIndex, const FStateTreeDelayedFinishedTask& DelayedFinisedTask);
+	void AddUniqueCompletedStateList(const UE::StateTree::FDelayedFinishedTask& DelayedFinisedTask);
 
 	/** In the list of completed states, returns the run status when it's completed by a global task. */
 	EStateTreeRunStatus GetGlobalTasksCompletedStatesStatus() const;
@@ -951,12 +946,9 @@ protected:
 	TOptional<EStateTreeRunStatus> CompletedStateRunStatus;
 
 	/** Helper struct to store the resolved FStateTreeExecutionState::FCompletedState */
-	struct FCompletedState
+	struct FCompletedState : UE::StateTree::FDelayedFinishedTask
 	{
-		int32 ActiveFrameIndex = INDEX_NONE;
-		int32 ActiveStateIndex = INDEX_NONE;
 		EStateTreeRunStatus RunStatus = EStateTreeRunStatus::Unset;
-		FStateTreeDelayedFinishedTask DelayedTask;
 	};
 
 	/**
@@ -978,17 +970,6 @@ protected:
 	/** Pointer to the shared instance data of the current frame we're processing.  */
 	FStateTreeInstanceStorage* CurrentlyProcessedSharedInstanceStorage = nullptr;
 
-	/** Helper struct to track currently execution context. */
-	struct FCurrentlyExecutionContextScope
-	{
-		FCurrentlyExecutionContextScope(FStateTreeExecutionContext& InContext);
-		~FCurrentlyExecutionContextScope();
-
-	private:
-		FStateTreeExecutionContext& CurrentContext;
-		FStateTreeExecutionContext* PreviousExecutionContext;
-	};
-
 	/** Helper struct to track currently processed frame. */
 	struct FCurrentlyProcessedFrameScope
 	{
@@ -1026,7 +1007,7 @@ protected:
 
 	private:
 		FStateTreeExecutionContext& Context;
-		FStateTreeStateHandle SavedState = FStateTreeStateHandle::Invalid; 
+		FStateTreeStateHandle SavedState = FStateTreeStateHandle::Invalid;
 	};
 
 	/** Current event we're processing in transition, or invalid if not applicable. */
@@ -1242,13 +1223,12 @@ public:
 
 private:
 	static void ScheduleNextTick(TNotNull<UObject*> Owner, TNotNull<const UStateTree*> RootStateTree, FStateTreeInstanceStorage& Storage);
+	bool IsFrameOrStateActive(const FStateTreeExecutionState& Exec) const;
 
 	TWeakObjectPtr<UObject> Owner;
 	TWeakObjectPtr<const UStateTree> StateTree;
 	TWeakPtr<FStateTreeInstanceStorage> Storage;
 
-	// This will be replace by StatePath
-	TWeakObjectPtr<const UStateTree> FrameStateTree;
-	FStateTreeStateHandle FrameRootState = FStateTreeStateHandle::Invalid;
-	FStateTreeStateHandle CurrentState = FStateTreeStateHandle::Invalid;
+	UE::StateTree::FActiveFrameID FrameID;
+	UE::StateTree::FActiveStateID StateID;
 };
diff --git a/Engine/Plugins/Runtime/StateTree/Source/StateTreeModule/Public/StateTreeExecutionTypes.h b/Engine/Plugins/Runtime/StateTree/Source/StateTreeModule/Public/StateTreeExecutionTypes.h
--- a/Engine/Plugins/Runtime/StateTree/Source/StateTreeModule/Public/StateTreeExecutionTypes.h
+++ b/Engine/Plugins/Runtime/StateTree/Source/StateTreeModule/Public/StateTreeExecutionTypes.h
@@ -424,11 +424,39 @@ struct STATETREEMODULE_API FStateTreeActiveStates
 		for (int32 Index = (int32)NumStates - 1; Index >= 0; Index--)
 		{
 			if (States[Index] == StateHandle)
+			{
 				return Index;
+			}
 		}
 		return INDEX_NONE;
 	}
 	
+	/** Returns index of a state, searching in reverse order. */
+	int32 IndexOfReverse(const UE::StateTree::FActiveStateID StateId) const
+	{
+		for (int32 Index = (int32)NumStates - 1; Index >= 0; Index--)
+		{
+			if (StateIDs[Index] == StateId)
+			{
+				return Index;
+			}
+		}
+		return INDEX_NONE;
+	}
+	
+	/** Returns true if the state id is inside the container. */
+	bool Contains(const UE::StateTree::FActiveStateID StateId) const
+	{
+		for (int32 Index = (int32)NumStates - 1; Index >= 0; Index--)
+		{
+			if (StateIDs[Index] == StateId)
+			{
+				return true;
+			}
+		}
+		return false;
+	}
+	
 	/** Returns last state in the array, or invalid state if the array is empty. */
 	FStateTreeStateHandle Last() const { return NumStates > 0 ? States[NumStates - 1] : FStateTreeStateHandle::Invalid; }
 	
@@ -617,20 +645,32 @@ struct STATETREEMODULE_API FStateTreeTransitionDelayedState
 };
 
 
+namespace UE::StateTree
+{
 /** Describes a finished task waiting to be processed by an execution context. */
-struct FStateTreeDelayedFinishedTask
+struct FDelayedFinishedTask
 {
-	/** Tree that identifies the active frame. */
-	FObjectKey FrameStateTree;
-	/** State that identifies the active frame. */
-	FStateTreeStateHandle FrameRootState = FStateTreeStateHandle::Invalid;
-	/** State that contains the finished task. Invalid if the task is global. */
-	FStateTreeStateHandle State = FStateTreeStateHandle::Invalid;
-	/** Task that finished and need to be processed. */
+	enum class EReasonType : uint8
+	{
+		/** A global task finished. The FrameID and TaskIndex are valid. */
+		GlobalTask,
+		/** A task inside a state finished. The FrameID, StateID and TaskIndex are valid. */
+		StateTask,
+		/** An internal transition finish the state. The FrameID and StateID are valid. */
+		InternalTransition,
+	};
+	/** Frame ID that identifies 1the active frame. */
+	FActiveFrameID FrameID;
+	/** State ID that contains the finished task. */
+	FActiveStateID StateID;
+	/** Task that is finished and needs to be processed. */
 	FStateTreeIndex16 TaskIndex = FStateTreeIndex16::Invalid;
 	/** The result of the finished task. */
 	EStateTreeFinishTaskType Type = EStateTreeFinishTaskType::Failed;
+	/** The reason of the finished task. */
+	EReasonType Reason = EReasonType::GlobalTask;
 };
+} //namespace UE::StateTree
 
 
 /*
@@ -716,13 +756,16 @@ struct FStateTreeDelegateActiveListeners
 	STATETREEMODULE_API ~FStateTreeDelegateActiveListeners();
 
 	/** Adds a delegate bound in the editor to the list. Safe to be called during broadcasting.*/
-	STATETREEMODULE_API void Add(const FStateTreeDelegateListener& Listener, FSimpleDelegate Delegate, TNotNull<const UStateTree*> FrameStateTree, FStateTreeStateHandle FrameRootState, FStateTreeStateHandle OwningState, FStateTreeIndex16 OwningNodeIndex);
+	STATETREEMODULE_API void Add(const FStateTreeDelegateListener& Listener, FSimpleDelegate Delegate, UE::StateTree::FActiveFrameID FrameID, UE::StateTree::FActiveStateID StateID, FStateTreeIndex16 OwningNodeIndex);
 
 	/** Removes a delegate bound in the editor from the list. Safe to be called during broadcasting. */
 	STATETREEMODULE_API void Remove(const FStateTreeDelegateListener& Listener);
 
 	/** Removes the listener by predicate. */
-	STATETREEMODULE_API void RemoveByPredicate(TFunctionRef<bool(FStateTreeStateHandle OwningState, FStateTreeIndex16 OwningNodeIndex, const UStateTree* FrameStateTree, FStateTreeStateHandle FrameRootState)> Predicate);
+	STATETREEMODULE_API void RemoveAll(UE::StateTree::FActiveFrameID FrameID);
+
+	/** Removes the listener that match. */
+	STATETREEMODULE_API void RemoveAll(UE::StateTree::FActiveStateID StateID);
 
 	/** Broadcasts the listener by predicate. */
 	STATETREEMODULE_API void BroadcastDelegate(const FStateTreeDelegateDispatcher& Dispatcher, const FStateTreeExecutionState& Exec);
@@ -733,7 +776,7 @@ private:
 	struct FActiveListener
 	{
 		FActiveListener() = default;
-		FActiveListener(const FStateTreeDelegateListener& Listener, FSimpleDelegate InDelegate, TNotNull<const UStateTree*> FrameStateTree, FStateTreeStateHandle FrameRootState, FStateTreeStateHandle OwningState, FStateTreeIndex16 OwningNodeIndex);
+		FActiveListener(const FStateTreeDelegateListener& Listener, FSimpleDelegate InDelegate, UE::StateTree::FActiveFrameID FrameID, UE::StateTree::FActiveStateID StateID, FStateTreeIndex16 OwningNodeIndex);
 
 		bool IsValid() const
 		{
@@ -742,12 +785,9 @@ private:
 
 		FStateTreeDelegateListener Listener;
 		FSimpleDelegate Delegate;
-
-		FStateTreeStateHandle OwningState = FStateTreeStateHandle::Invalid;
+		UE::StateTree::FActiveFrameID FrameID;
+		UE::StateTree::FActiveStateID StateID;
 		FStateTreeIndex16 OwningNodeIndex = FStateTreeIndex16::Invalid;
-
-		FObjectKey FrameStateTree;
-		FStateTreeStateHandle FrameRootState;
 	};
 
 	TArray<FActiveListener> Listeners;
@@ -856,6 +896,15 @@ public:
 		return Result;
 	}
 
+	/** @return the active frame that matches the unique frame ID. */
+	const FStateTreeExecutionFrame* FindActiveFrame(UE::StateTree::FActiveFrameID FrameID) const;
+
+	/** @return the active frame that matches the unique frame ID. */
+	FStateTreeExecutionFrame* FindActiveFrame(UE::StateTree::FActiveFrameID FrameID);
+
+	/** @return the active frame index that matches the unique frame ID. */
+	int32 IndexOfActiveFrame(UE::StateTree::FActiveFrameID FrameID) const;
+
 	/** @return whether it contains any shcedule tick requests. */
 	bool HasScheduledTickRequests() const
 	{
@@ -890,7 +939,7 @@ public:
 	TArray<FStateTreeTransitionDelayedState> DelayedTransitions;
 
 	/** Pending finished tasks that need processing. */
-	TArray<FStateTreeDelayedFinishedTask> DelayedFinishedTasks;
+	TArray<UE::StateTree::FDelayedFinishedTask> DelayedFinishedTasks;
 
 	/** Used by state tree random-based operations. */
 	UPROPERTY()
@@ -921,9 +970,6 @@ public:
 	UPROPERTY(Transient)
 	TInstancedStruct<FStateTreeExecutionExtension> ExecutionExtension;
 
-	/** The current context used to operate on the instance. */
-	FStateTreeExecutionContext* CurrentExecutionContext = nullptr;
-
 	/** The index of the task that failed during enter state. Exit state uses it to call ExitState() symmetrically. */
 	UPROPERTY()
 	FStateTreeIndex16 EnterStateFailedFrameIndex = FStateTreeIndex16::Invalid;
@@ -965,6 +1011,9 @@ public:
 	UE_DEPRECATED(5.6, "Use DelayedCompletedTasks for trigerring completion transitions.")
 	UPROPERTY()
 	FStateTreeStateHandle CompletedStateHandle = FStateTreeStateHandle::Invalid;
+
+	UE_DEPRECATED(5.6, "CurrentExecutionContext is not needed anymore. Use FrameID and StateID.")
+	FStateTreeExecutionContext* CurrentExecutionContext = nullptr;
 #endif //WITH_EDITORONLY_DATA
 };
 
diff --git a/Engine/Plugins/Runtime/StateTree/Source/StateTreeTestSuite/Private/StateTreeTest.cpp b/Engine/Plugins/Runtime/StateTree/Source/StateTreeTestSuite/Private/StateTreeTest.cpp
--- a/Engine/Plugins/Runtime/StateTree/Source/StateTreeTestSuite/Private/StateTreeTest.cpp
+++ b/Engine/Plugins/Runtime/StateTree/Source/StateTreeTestSuite/Private/StateTreeTest.cpp
@@ -1607,6 +1607,179 @@ struct FStateTreeTest_DeferredStop_ExitTask : FStateTreeTest_DeferredStop
 };
 IMPLEMENT_AI_INSTANT_TEST(FStateTreeTest_DeferredStop_ExitTask, "System.StateTree.DeferredStop.ExitTask");
 
+struct FStateTreeTest_FinishTasks : FStateTreeTestBase
+{
+	virtual bool InstantTest() override
+	{
+		UStateTree& StateTree = NewStateTree();
+		UStateTreeEditorData& EditorData = *Cast<UStateTreeEditorData>(StateTree.EditorData);
+
+		/*
+		 - RootA
+			- StateA -> StateB
+			- StateB -> StateA
+		 */
+		UStateTreeState& RootA = EditorData.AddSubTree(FName(TEXT("RootA")));
+		UStateTreeState& StateA = RootA.AddChildState(FName(TEXT("StateA")));
+		UStateTreeState& StateB = RootA.AddChildState(FName(TEXT("StateB")));
+
+		StateA.AddTransition(EStateTreeTransitionTrigger::OnStateCompleted, EStateTreeTransitionType::GotoState, &StateB);
+		StateB.AddTransition(EStateTreeTransitionTrigger::OnStateCompleted, EStateTreeTransitionType::GotoState, &StateA);
+		
+		TStateTreeEditorNode<FStateTreeTestCondition>& BoolCondB = StateB.AddEnterCondition<FStateTreeTestCondition>();
+		BoolCondB.GetNode().bTestConditionResult = false;
+
+		TStateTreeEditorNode<FTestTask_PrintValue>& StateATask = StateA.AddTask<FTestTask_PrintValue>("StateATaskA");
+		StateATask.GetInstanceData().Value = 101;
+		StateATask.GetNode().CustomTickFunc = [](FStateTreeExecutionContext& Context, const FTestTask_PrintValue* Task)
+			{
+				Context.FinishTask(*Task, EStateTreeFinishTaskType::Succeeded);
+				FTestTask_PrintValue::FInstanceDataType& InstanceData = Context.GetInstanceData(*Task);
+				++InstanceData.Value;
+			};
+
+		// Test one finish
+		{
+			{
+				FStateTreeCompilerLog Log;
+				FStateTreeCompiler Compiler(Log);
+				const bool bResult = Compiler.Compile(StateTree);
+				AITEST_TRUE("StateTree should get compiled", bResult);
+			}
+			{
+				EStateTreeRunStatus Status = EStateTreeRunStatus::Unset;
+				FStateTreeInstanceData InstanceData;
+				FTestStateTreeExecutionContext Exec(StateTree, StateTree, InstanceData);
+				const bool bInitSucceeded = Exec.IsValid();
+				AITEST_TRUE("StateTree should init", bInitSucceeded);
+
+				// Start and enter state
+				Status = Exec.Start();
+
+				AITEST_TRUE("StateTree Active States should be in RootA/StateA",
+					Exec.ExpectInActiveStates(RootA.Name, StateA.Name));
+				AITEST_TRUE("StateTree should be running", Status == EStateTreeRunStatus::Running);
+				Exec.LogClear();
+
+				// One FinishTask, it should loop to StateA
+				Status = Exec.Tick(0.1f);
+				AITEST_TRUE("StateTree Active States should be in RootA/StateA",
+					Exec.ExpectInActiveStates(RootA.Name, StateA.Name));
+				AITEST_TRUE("StateTree should be running", Status == EStateTreeRunStatus::Running);
+				AITEST_TRUE("Expect the output tasks", Exec.Expect("StateATaskA", TEXT("Tick101"))
+					.Then("StateATaskA", TEXT("Exitstate102"))
+					//@TODO: instance data should be reset
+					//.Then("StateATaskA", TEXT("EnterState101"))
+					);
+				Exec.LogClear();
+
+				Exec.Stop();
+			}
+		}
+		// Test two finish
+		{
+			{
+				StateATask.GetNode().CustomTickFunc = [](FStateTreeExecutionContext& Context, const FTestTask_PrintValue* Task)
+					{
+						Context.FinishTask(*Task, EStateTreeFinishTaskType::Succeeded);
+						Context.FinishTask(*Task, EStateTreeFinishTaskType::Succeeded);
+						FTestTask_PrintValue::FInstanceDataType& InstanceData = Context.GetInstanceData(*Task);
+						++InstanceData.Value;
+					};
+
+				FStateTreeCompilerLog Log;
+				FStateTreeCompiler Compiler(Log);
+				const bool bResult = Compiler.Compile(StateTree);
+				AITEST_TRUE("StateTree should get compiled", bResult);
+			}
+			{
+				EStateTreeRunStatus Status = EStateTreeRunStatus::Unset;
+				FStateTreeInstanceData InstanceData;
+				FTestStateTreeExecutionContext Exec(StateTree, StateTree, InstanceData);
+				const bool bInitSucceeded = Exec.IsValid();
+				AITEST_TRUE("StateTree should init", bInitSucceeded);
+
+				// Start and enter state
+				Status = Exec.Start();
+
+				AITEST_TRUE("StateTree Active States should be in RootA/StateA",
+					Exec.ExpectInActiveStates(RootA.Name, StateA.Name));
+				AITEST_TRUE("StateTree should be running", Status == EStateTreeRunStatus::Running);
+				Exec.LogClear();
+
+				// One FinishTask, it should loop to StateA
+				Status = Exec.Tick(0.1f);
+				AITEST_TRUE("StateTree Active States should be in RootA/StateA",
+					Exec.ExpectInActiveStates(RootA.Name, StateA.Name));
+				AITEST_TRUE("StateTree should be running", Status == EStateTreeRunStatus::Running);
+				AITEST_TRUE("Expect the output tasks", Exec.Expect("StateATaskA", TEXT("Tick101"))
+					.Then("StateATaskA", TEXT("Exitstate102"))
+					//@TODO: instance data should be reset
+					//.Then("StateATaskA", TEXT("EnterState101"))
+				);
+				Exec.LogClear();
+
+				Exec.Stop();
+			}
+		}
+		// Test finish in exit state
+		{
+			{
+				StateATask.GetNode().CustomTickFunc = [](FStateTreeExecutionContext& Context, const FTestTask_PrintValue* Task)
+					{
+						Context.FinishTask(*Task, EStateTreeFinishTaskType::Succeeded);
+						FTestTask_PrintValue::FInstanceDataType& InstanceData = Context.GetInstanceData(*Task);
+						++InstanceData.Value;
+					};
+				StateATask.GetNode().CustomExitStateFunc = [](FStateTreeExecutionContext& Context, const FTestTask_PrintValue* Task)
+					{
+						Context.FinishTask(*Task, EStateTreeFinishTaskType::Succeeded);
+						FTestTask_PrintValue::FInstanceDataType& InstanceData = Context.GetInstanceData(*Task);
+						++InstanceData.Value;
+					};
+
+				FStateTreeCompilerLog Log;
+				FStateTreeCompiler Compiler(Log);
+				const bool bResult = Compiler.Compile(StateTree);
+				AITEST_TRUE("StateTree should get compiled", bResult);
+			}
+			{
+				EStateTreeRunStatus Status = EStateTreeRunStatus::Unset;
+				FStateTreeInstanceData InstanceData;
+				FTestStateTreeExecutionContext Exec(StateTree, StateTree, InstanceData);
+				const bool bInitSucceeded = Exec.IsValid();
+				AITEST_TRUE("StateTree should init", bInitSucceeded);
+
+				// Start and enter state
+				Status = Exec.Start();
+
+				AITEST_TRUE("StateTree Active States should be in RootA/StateA",
+					Exec.ExpectInActiveStates(RootA.Name, StateA.Name));
+				AITEST_TRUE("StateTree should be running", Status == EStateTreeRunStatus::Running);
+				Exec.LogClear();
+
+				// One FinishTask in exit but should not close StateA again. It should loop to StateA
+				Status = Exec.Tick(0.1f);
+				AITEST_TRUE("StateTree Active States should be in RootA/StateA",
+					Exec.ExpectInActiveStates(RootA.Name, StateA.Name));
+				AITEST_TRUE("StateTree should be running", Status == EStateTreeRunStatus::Running);
+				AITEST_TRUE("Expect the output tasks", Exec.Expect("StateATaskA", TEXT("Tick101"))
+					.Then("StateATaskA", TEXT("ExitState102"))
+					//@TODO: instance data should be reset
+					//.Then("StateATaskA", TEXT("EnterState101"))
+				);
+				AITEST_FALSE("Expect the output tasks", Exec.Expect("StateATaskA", TEXT("ExitState103")));
+				Exec.LogClear();
+
+				Exec.Stop();
+			}
+		}
+
+		return true;
+	}
+};
+IMPLEMENT_AI_INSTANT_TEST(FStateTreeTest_FinishTasks, "System.StateTree.FinishTask");
+
 
 // Test nested tree overrides
 struct FStateTreeTest_NestedOverride : FStateTreeTestBase
diff --git a/Engine/Plugins/Runtime/StateTree/Source/StateTreeTestSuite/Private/StateTreeTestTypes.h b/Engine/Plugins/Runtime/StateTree/Source/StateTreeTestSuite/Private/StateTreeTestTypes.h
--- a/Engine/Plugins/Runtime/StateTree/Source/StateTreeTestSuite/Private/StateTreeTestTypes.h
+++ b/Engine/Plugins/Runtime/StateTree/Source/StateTreeTestSuite/Private/StateTreeTestTypes.h
@@ -233,6 +233,12 @@ struct FTestTask_PrintValueInstanceData
 	
 	UPROPERTY(EditAnywhere, Category = "Parameter")
 	TArray<int32> ArrayValue;
+	
+	UPROPERTY(EditAnywhere, Category = "Parameter")
+	EStateTreeRunStatus EnterStateRunStatus = EStateTreeRunStatus::Running;
+	
+	UPROPERTY(EditAnywhere, Category = "Parameter")
+	EStateTreeRunStatus TickRunStatus = EStateTreeRunStatus::Running;
 
 	FString GetArrayDescription() const
 	{
@@ -276,7 +282,11 @@ struct FTestTask_PrintValue : public FStateTreeTaskBase
 		const FInstanceDataType& InstanceData = Context.GetInstanceData(*this);
 		TestLog.Log(Name,  FString::Printf(TEXT("EnterState%d"), InstanceData.Value));
 		TestLog.Log(Name, FString::Printf(TEXT("EnterState:%s"), *InstanceData.GetArrayDescription()));
-		return EStateTreeRunStatus::Running;
+		if (CustomEnterStateFunc)
+		{
+			CustomEnterStateFunc(Context, this);
+		}
+		return InstanceData.EnterStateRunStatus;
 	}
 
 	virtual void ExitState(FStateTreeExecutionContext& Context, const FStateTreeTransitionResult& Transition) const override
@@ -285,6 +295,10 @@ struct FTestTask_PrintValue : public FStateTreeTaskBase
 		const FInstanceDataType& InstanceData = Context.GetInstanceData(*this);
 		TestLog.Log(Name,  FString::Printf(TEXT("ExitState%d"), InstanceData.Value));
 		TestLog.Log(Name, FString::Printf(TEXT("ExitState:%s"), *InstanceData.GetArrayDescription()));
+		if (CustomExitStateFunc)
+		{
+			CustomExitStateFunc(Context, this);
+		}
 	}
 
 	virtual EStateTreeRunStatus Tick(FStateTreeExecutionContext& Context, const float DeltaTime) const override
@@ -293,8 +307,11 @@ struct FTestTask_PrintValue : public FStateTreeTaskBase
 		const FInstanceDataType& InstanceData = Context.GetInstanceData(*this);
 		TestLog.Log(Name,  FString::Printf(TEXT("Tick%d"), InstanceData.Value));
 		TestLog.Log(Name,  FString::Printf(TEXT("Tick:%s"), *InstanceData.GetArrayDescription()));
-		
-		return EStateTreeRunStatus::Running;
+		if (CustomTickFunc)
+		{
+			CustomTickFunc(Context, this);
+		}
+		return InstanceData.TickRunStatus;
 	}
 
 	virtual void TriggerTransitions(FStateTreeExecutionContext& Context) const override
@@ -312,6 +329,9 @@ struct FTestTask_PrintValue : public FStateTreeTaskBase
 	}
 	
 	TStateTreeExternalDataHandle<FStateTreeTestLog> LogHandle;
+	TFunction<void(FStateTreeExecutionContext&, const FTestTask_PrintValue*)> CustomEnterStateFunc;
+	TFunction<void(FStateTreeExecutionContext&, const FTestTask_PrintValue*)> CustomExitStateFunc;
+	TFunction<void(FStateTreeExecutionContext&, const FTestTask_PrintValue*)> CustomTickFunc;
 };
 
 USTRUCT()
@@ -569,8 +589,11 @@ struct FStateTreeTestCondition : public FStateTreeConditionCommonBase
 	{
 		FInstanceDataType& InstanceData = Context.GetInstanceData(*this);
 		InstanceData.GlobalCounter.fetch_add(InstanceData.Count);
-		return true;
+		return bTestConditionResult;
 	}
+
+	UPROPERTY()
+	bool bTestConditionResult = true;
 };
 
 struct FStateTreeTestRunContext
