Temptative fix for clang compiler.

diff --git a/Engine/Plugins/Runtime/StateTree/Source/StateTreeModule/Public/StateTreeExecutionContext.h b/Engine/Plugins/Runtime/StateTree/Source/StateTreeModule/Public/StateTreeExecutionContext.h
--- a/Engine/Plugins/Runtime/StateTree/Source/StateTreeModule/Public/StateTreeExecutionContext.h
+++ b/Engine/Plugins/Runtime/StateTree/Source/StateTreeModule/Public/StateTreeExecutionContext.h
@@ -948,6 +948,10 @@ protected:
 	/** Helper struct to store the resolved FStateTreeExecutionState::FCompletedState */
 	struct FCompletedState : UE::StateTree::FDelayedFinishedTask
 	{
+		FCompletedState() = default;
+		FCompletedState(const UE::StateTree::FDelayedFinishedTask& Other)
+			: UE::StateTree::FDelayedFinishedTask(Other)
+		{}
 		EStateTreeRunStatus RunStatus = EStateTreeRunStatus::Unset;
 	};
 
