StateTree: Update StateTreeInstanceDataStructRef to use FrameID.
#rnx

diff --git a/Engine/Plugins/Runtime/StateTree/Source/StateTreeModule/Public/StateTreeInstanceData.h b/Engine/Plugins/Runtime/StateTree/Source/StateTreeModule/Public/StateTreeInstanceData.h
--- a/Engine/Plugins/Runtime/StateTree/Source/StateTreeModule/Public/StateTreeInstanceData.h
+++ b/Engine/Plugins/Runtime/StateTree/Source/StateTreeModule/Public/StateTreeInstanceData.h
@@ -571,8 +571,7 @@ struct TStateTreeInstanceDataStructRef
 {
 	TStateTreeInstanceDataStructRef(FStateTreeInstanceData& InInstanceData, const FStateTreeExecutionFrame& CurrentFrame, const FStateTreeDataHandle InDataHandle)
 		: WeakStorage(InInstanceData.GetWeakMutableStorage())
-		, WeakStateTree(CurrentFrame.StateTree)
-		, RootState(CurrentFrame.RootState)
+		, FrameID(CurrentFrame.FrameID)
 		, DataHandle(InDataHandle)
 	{
 		checkf(InDataHandle.GetSource() == EStateTreeDataSourceType::ActiveInstanceData
@@ -580,28 +579,33 @@ struct TStateTreeInstanceDataStructRef
 			TEXT("TStateTreeInstanceDataStructRef supports only struct instance data."));
 	}
 
-	bool IsValid() const { return WeakStateTree.IsValid() && RootState.IsValid() && DataHandle.IsValid(); }
+PRAGMA_DISABLE_DEPRECATION_WARNINGS
+	TStateTreeInstanceDataStructRef(const TStateTreeInstanceDataStructRef& Other) = default;
+	TStateTreeInstanceDataStructRef(TStateTreeInstanceDataStructRef&& Other) = default;
+	TStateTreeInstanceDataStructRef& operator=(TStateTreeInstanceDataStructRef const& Other) = default;
+	TStateTreeInstanceDataStructRef& operator=(TStateTreeInstanceDataStructRef&& Other) = default;
+PRAGMA_ENABLE_DEPRECATION_WARNINGS
+
+	bool IsValid() const
+	{
+		return FrameID.IsValid() && DataHandle.IsValid();
+	}
 
 	T* GetPtr() const
 	{
-		if (!WeakStorage.IsValid())
+		TSharedPtr<FStateTreeInstanceStorage> StoragePtr = WeakStorage.Pin();
+		if (StoragePtr == nullptr)
 		{
 			return nullptr;
 		}
 
-		FStateTreeInstanceStorage& Storage = *WeakStorage.Pin();
-
-		const FStateTreeExecutionState& Exec = Storage.GetExecutionState();
-		const UStateTree* StateTree = WeakStateTree.Get();
-		
-		const FStateTreeExecutionFrame* CurrentFrame = Exec.ActiveFrames.FindByPredicate([StateTree, RootState = RootState](const FStateTreeExecutionFrame& Frame)
-		{
-			return Frame.StateTree == StateTree && Frame.RootState == RootState;
-		});
+		const FStateTreeExecutionState& Exec = StoragePtr->GetExecutionState();
+		const FStateTreeExecutionFrame* CurrentFrame = Exec.FindActiveFrame(FrameID);
 
 		FStructView Struct;
 		if (CurrentFrame)
 		{
+			FStateTreeInstanceStorage& Storage = *(StoragePtr.Get());
 			if (IsHandleSourceValid(Storage, *CurrentFrame, DataHandle))
 			{
 				Struct = GetDataView(Storage, *CurrentFrame, DataHandle);
@@ -610,9 +614,9 @@ struct TStateTreeInstanceDataStructRef
 			{
 				Struct = Storage.GetMutableTemporaryStruct(*CurrentFrame, DataHandle);
 			}
-		}
 
-		check(Struct.GetScriptStruct() == TBaseStructure<T>::Get());
+			check(Struct.GetScriptStruct() == TBaseStructure<T>::Get());
+		}
 		return reinterpret_cast<T*>(Struct.GetMemory());
 	}
 
@@ -651,7 +655,13 @@ protected:
 	}
 	
 	TWeakPtr<FStateTreeInstanceStorage> WeakStorage = nullptr;
+	UE::StateTree::FActiveFrameID FrameID;
+	FStateTreeDataHandle DataHandle = FStateTreeDataHandle::Invalid;
+
+#if WITH_EDITORONLY_DATA
+	UE_DEPRECATED(5.6, "Use the frame id to identify the frame.")
 	TWeakObjectPtr<const UStateTree> WeakStateTree = nullptr;
+	UE_DEPRECATED(5.6, "Use the frame id to identify the frame.")
 	FStateTreeStateHandle RootState = FStateTreeStateHandle::Invalid;
-	FStateTreeDataHandle DataHandle = FStateTreeDataHandle::Invalid;
+#endif
 };
