StateTree: Update transition request to use frame ID. FrameID uniquely identifies the frame and prevents errors in async tasks.

diff --git a/Engine/Plugins/Runtime/StateTree/Source/StateTreeModule/Private/StateTreeExecutionContext.cpp b/Engine/Plugins/Runtime/StateTree/Source/StateTreeModule/Private/StateTreeExecutionContext.cpp
--- a/Engine/Plugins/Runtime/StateTree/Source/StateTreeModule/Private/StateTreeExecutionContext.cpp
+++ b/Engine/Plugins/Runtime/StateTree/Source/StateTreeModule/Private/StateTreeExecutionContext.cpp
@@ -1154,10 +1154,10 @@ void FStateTreeExecutionContext::RequestTransition(const FStateTreeTransitionReq
 		STATETREE_LOG(Verbose, TEXT("Request transition to '%s' at priority %s"), *GetSafeStateName(*RootFrame, Request.TargetState), *UEnum::GetDisplayValueAsText(Request.Priority).ToString());
 
 		FStateTreeTransitionRequest RequestWithSource = Request;
-		RequestWithSource.SourceStateTree = RootFrame->StateTree;
-		RequestWithSource.SourceRootState = RootFrame->ActiveStates[0];
-		RequestWithSource.SourceState = CurrentlyProcessedState;
-		
+		RequestWithSource.SourceFrameID = RootFrame->FrameID;
+		int32 ActiveStateIndex = RootFrame->ActiveStates.IndexOfReverse(CurrentlyProcessedState);
+		RequestWithSource.SourceStateID = ActiveStateIndex != INDEX_NONE ? RootFrame->ActiveStates.StateIDs[ActiveStateIndex] : UE::StateTree::FActiveStateID::Invalid;
+
 		InstanceData.AddTransitionRequest(&Owner, RequestWithSource);
 	}
 
@@ -3891,9 +3891,7 @@ bool FStateTreeExecutionContext::TriggerTransitions()
 	for (const FStateTreeTransitionRequest& Request : InstanceData.GetTransitionRequests())
 	{
 		// Find frame associated with the request.
-		const FStateTreeExecutionFrame* ParentFrame = nullptr;
-		const FStateTreeExecutionFrame* CurrentFrame = FindFrame(Request.SourceStateTree, Request.SourceRootState, Exec.ActiveFrames, ParentFrame);
-
+		const FStateTreeExecutionFrame* CurrentFrame = Exec.FindActiveFrame(Request.SourceFrameID);
 		if (CurrentFrame)
 		{
 			if (RequestTransition(*CurrentFrame, Request.TargetState, Request.Priority, /*TransitionEvent*/nullptr, Request.Fallback))
@@ -5680,20 +5678,8 @@ bool FStateTreeWeakExecutionContext::RequestTransition(FStateTreeStateHandle Tar
 		{
 			FStateTreeTransitionRequest Request;
 
-			if (const FStateTreeExecutionFrame* CurrentFrame = Exec.FindActiveFrame(FrameID))
-			{
-				Request.SourceStateTree = CurrentFrame->StateTree;
-				Request.SourceRootState = CurrentFrame->RootState;
-				const int32 IndexOfActiveState = CurrentFrame->ActiveStates.IndexOfReverse(StateID);
-				Request.SourceState = IndexOfActiveState != INDEX_NONE ? CurrentFrame->ActiveStates[IndexOfActiveState] : FStateTreeStateHandle();
-			}
-			else
-			{
-				const FStateTreeExecutionFrame* RootFrame = &Exec.ActiveFrames[0];
-				Request.SourceStateTree = RootFrame->StateTree;
-				Request.SourceRootState = RootFrame->RootState;
-				Request.SourceState = FStateTreeStateHandle();
-			}
+			Request.SourceFrameID = FrameID;
+			Request.SourceStateID = StateID;
 
 			Request.TargetState = TargetState;
 			Request.Priority = Priority;
diff --git a/Engine/Plugins/Runtime/StateTree/Source/StateTreeModule/Public/StateTreeExecutionTypes.h b/Engine/Plugins/Runtime/StateTree/Source/StateTreeModule/Public/StateTreeExecutionTypes.h
--- a/Engine/Plugins/Runtime/StateTree/Source/StateTreeModule/Public/StateTreeExecutionTypes.h
+++ b/Engine/Plugins/Runtime/StateTree/Source/StateTreeModule/Public/StateTreeExecutionTypes.h
@@ -260,18 +260,6 @@ struct STATETREEMODULE_API FStateTreeTransitionRequest
 	{
 	}
 
-	/** Source state of the transition. Filled in by the StateTree execution context. */
-	UPROPERTY()
-	FStateTreeStateHandle SourceState;
-
-	/** StateTree asset that was active when the transition was requested. Filled in by the StateTree execution context. */
-	UPROPERTY()
-	TObjectPtr<const UStateTree> SourceStateTree = nullptr;
-
-	/** Root state the execution frame where the transition was requested. Filled in by the StateTree execution context. */
-	UPROPERTY()
-	FStateTreeStateHandle SourceRootState = FStateTreeStateHandle::Invalid;
-
 	/** Target state of the transition. */
 	UPROPERTY(EditDefaultsOnly, Category = "Default")
 	FStateTreeStateHandle TargetState;
@@ -283,6 +271,28 @@ struct STATETREEMODULE_API FStateTreeTransitionRequest
 	/** Fallback of the transition if it fails to select the target state */
 	UPROPERTY(EditDefaultsOnly, Category = "Default")
 	EStateTreeSelectionFallback Fallback = EStateTreeSelectionFallback::None;
+
+	/** StateTree frame that was active when the transition was requested. Filled in by the StateTree execution context. */
+	UE::StateTree::FActiveFrameID SourceFrameID;
+	/** StateTree state that was active when the transition was requested. Filled in by the StateTree execution context. */
+	UE::StateTree::FActiveStateID SourceStateID;
+
+#if WITH_EDITORONLY_DATA
+	UE_DEPRECATED(5.6, "Use SourceFrameID to uniquely identify a frame.")
+	/** StateTree asset that was active when the transition was requested. Filled in by the StateTree execution context. */
+	UPROPERTY()
+	TObjectPtr<const UStateTree> SourceStateTree = nullptr;
+
+	UE_DEPRECATED(5.6, "Use SourceFrameID to uniquely identify a frame.")
+	/** Root state the execution frame where the transition was requested. Filled in by the StateTree execution context. */
+	UPROPERTY()
+	FStateTreeStateHandle SourceRootState = FStateTreeStateHandle::Invalid;
+
+	UE_DEPRECATED(5.6, "Use SourceStateID to uniquely identify a state.")
+	/** Source state of the transition. Filled in by the StateTree execution context. */
+	UPROPERTY()
+	FStateTreeStateHandle SourceState;
+#endif
 };
 
 
