Fix compilation issue introduced with CL 39483283

#rnx

diff --git a/Engine/Plugins/Runtime/StateTree/Source/StateTreeModule/Public/StateTreeExecutionTypes.h b/Engine/Plugins/Runtime/StateTree/Source/StateTreeModule/Public/StateTreeExecutionTypes.h
--- a/Engine/Plugins/Runtime/StateTree/Source/StateTreeModule/Public/StateTreeExecutionTypes.h
+++ b/Engine/Plugins/Runtime/StateTree/Source/StateTreeModule/Public/StateTreeExecutionTypes.h
@@ -242,6 +242,7 @@ struct STATETREEMODULE_API FStateTreeTransitionRequest
 {
 	GENERATED_BODY()
 
+PRAGMA_DISABLE_DEPRECATION_WARNINGS
 	FStateTreeTransitionRequest() = default;
 
 	explicit FStateTreeTransitionRequest(const FStateTreeStateLink& InStateLink)
@@ -260,6 +261,12 @@ struct STATETREEMODULE_API FStateTreeTransitionRequest
 	{
 	}
 
+	FStateTreeTransitionRequest(const FStateTreeTransitionRequest&) = default;
+	FStateTreeTransitionRequest(FStateTreeTransitionRequest&&) = default;
+	FStateTreeTransitionRequest& operator=(const FStateTreeTransitionRequest&) = default;
+	FStateTreeTransitionRequest& operator=(FStateTreeTransitionRequest&&) = default;
+PRAGMA_ENABLE_DEPRECATION_WARNINGS
+
 	/** Target state of the transition. */
 	UPROPERTY(EditDefaultsOnly, Category = "Default")
 	FStateTreeStateHandle TargetState;
