StateTree: Save the overrides in the context to prevent the user from modifying the overrides array while ticking. Informs the context implementation when the overrides are set.

diff --git a/Engine/Plugins/AI/MassAI/Source/MassAIBehavior/Private/MassStateTreeExecutionContext.cpp b/Engine/Plugins/AI/MassAI/Source/MassAIBehavior/Private/MassStateTreeExecutionContext.cpp
--- a/Engine/Plugins/AI/MassAI/Source/MassAIBehavior/Private/MassStateTreeExecutionContext.cpp
+++ b/Engine/Plugins/AI/MassAI/Source/MassAIBehavior/Private/MassStateTreeExecutionContext.cpp
@@ -171,6 +171,16 @@ FString FMassExecutionExtension::GetInstanceDescription(const FContextParameters
 	return FString::Printf(TEXT("Entity [%s]"), *Entity.DebugGetDescription());
 }
 
+void FMassExecutionExtension::OnLinkedStateTreeOverridesSet(const FContextParameters& Context, const FStateTreeReferenceOverrides& Overrides)
+{
+	const uint32 NewLinkedStateTreeOverridesHash = GetTypeHash(Overrides);
+	if (LinkedStateTreeOverridesHash != LinkedStateTreeOverridesHash)
+	{
+		LinkedStateTreeOverridesHash = NewLinkedStateTreeOverridesHash;
+		//@todo update mass dependencies
+	}
+}
+
 void FMassStateTreeExecutionContext::BeginDelayedTransition(const FStateTreeTransitionDelayedState& DelayedState)
 {
 	UMassSignalSubsystem* SignalSubsystem = MassEntityExecutionContext->GetMutableSubsystem<UMassSignalSubsystem>();
diff --git a/Engine/Plugins/AI/MassAI/Source/MassAIBehavior/Public/MassStateTreeExecutionContext.h b/Engine/Plugins/AI/MassAI/Source/MassAIBehavior/Public/MassStateTreeExecutionContext.h
--- a/Engine/Plugins/AI/MassAI/Source/MassAIBehavior/Public/MassStateTreeExecutionContext.h
+++ b/Engine/Plugins/AI/MassAI/Source/MassAIBehavior/Public/MassStateTreeExecutionContext.h
@@ -22,8 +22,10 @@ struct FMassExecutionExtension : public FStateTreeExecutionExtension
 
 public:
 	virtual FString GetInstanceDescription(const FContextParameters& Context) const override;
+	virtual void OnLinkedStateTreeOverridesSet(const FContextParameters& Context, const FStateTreeReferenceOverrides& Overrides) override;
 
 	FMassEntityHandle Entity;
+	uint32 LinkedStateTreeOverridesHash = 0;
 };
 
 /**
diff --git a/Engine/Plugins/Runtime/GameplayStateTree/Source/GameplayStateTreeModule/Private/Components/StateTreeComponent.cpp b/Engine/Plugins/Runtime/GameplayStateTree/Source/GameplayStateTreeModule/Private/Components/StateTreeComponent.cpp
--- a/Engine/Plugins/Runtime/GameplayStateTree/Source/GameplayStateTreeModule/Private/Components/StateTreeComponent.cpp
+++ b/Engine/Plugins/Runtime/GameplayStateTree/Source/GameplayStateTreeModule/Private/Components/StateTreeComponent.cpp
@@ -81,7 +81,7 @@ bool UStateTreeComponent::CollectExternalData(const FStateTreeExecutionContext&
 
 bool UStateTreeComponent::SetContextRequirements(FStateTreeExecutionContext& Context, bool bLogErrors)
 {
-	Context.SetLinkedStateTreeOverrides(&LinkedStateTreeOverrides);
+	Context.SetLinkedStateTreeOverrides(LinkedStateTreeOverrides);
 	Context.SetCollectExternalDataCallback(FOnCollectStateTreeExternalData::CreateUObject(this, &UStateTreeComponent::CollectExternalData));
 	return UStateTreeComponentSchema::SetContextRequirements(*this, Context);
 }
diff --git a/Engine/Plugins/Runtime/StateTree/Source/StateTreeModule/Private/StateTreeExecutionContext.cpp b/Engine/Plugins/Runtime/StateTree/Source/StateTreeModule/Private/StateTreeExecutionContext.cpp
--- a/Engine/Plugins/Runtime/StateTree/Source/StateTreeModule/Private/StateTreeExecutionContext.cpp
+++ b/Engine/Plugins/Runtime/StateTree/Source/StateTreeModule/Private/StateTreeExecutionContext.cpp
@@ -414,7 +414,7 @@ FStateTreeExecutionContext::FStateTreeExecutionContext(UObject& InOwner, const U
 FStateTreeExecutionContext::FStateTreeExecutionContext(const FStateTreeExecutionContext& InContextToCopy, const UStateTree& InStateTree, FStateTreeInstanceData& InInstanceData)
 	: FStateTreeExecutionContext(InContextToCopy.Owner, InStateTree, InInstanceData, InContextToCopy.CollectExternalDataDelegate)
 {
-	LinkedStateTreeOverrides = InContextToCopy.LinkedStateTreeOverrides;
+	SetLinkedStateTreeOverrides(InContextToCopy.LinkedAssetStateTreeOverrides);
 	const bool bIsSameSchema = RootStateTree.GetSchema()->GetClass() == InContextToCopy.GetStateTree()->GetSchema()->GetClass();
 	if (bIsSameSchema)
 	{
@@ -445,24 +445,108 @@ FStateTreeExecutionContext::~FStateTreeExecutionContext()
 
 void FStateTreeExecutionContext::SetCollectExternalDataCallback(const FOnCollectStateTreeExternalData& Callback)
 {
+	if (!IsValid())
+	{
+		STATETREE_LOG(Warning, TEXT("%hs: StateTree context is not initialized properly ('%s' using StateTree '%s')"),
+			__FUNCTION__, *GetNameSafe(&Owner), *GetFullNameSafe(&RootStateTree));
+		return;
+	}
+
+	FStateTreeExecutionState& Exec = GetExecState();
+	if (!ensureMsgf(Exec.CurrentPhase == EStateTreeUpdatePhase::Unset, TEXT("%hs can't be called while already in %s ('%s' using StateTree '%s')."),
+		__FUNCTION__, *UEnum::GetDisplayValueAsText(Exec.CurrentPhase).ToString(), *GetNameSafe(&Owner), *GetFullNameSafe(&RootStateTree)))
+	{
+		return;
+	}
+
 	CollectExternalDataDelegate = Callback;
 }
 
 void FStateTreeExecutionContext::SetLinkedStateTreeOverrides(const FStateTreeReferenceOverrides* InLinkedStateTreeOverrides)
 {
-	LinkedStateTreeOverrides = InLinkedStateTreeOverrides;
+	if (InLinkedStateTreeOverrides)
+	{
+		SetLinkedStateTreeOverrides(*InLinkedStateTreeOverrides);
+	}
+	else
+	{
+		SetLinkedStateTreeOverrides(FStateTreeReferenceOverrides());
+	}
 }
 
-const FStateTreeReference* FStateTreeExecutionContext::GetLinkedStateTreeOverrideForTag(const FGameplayTag StateTag) const
+void FStateTreeExecutionContext::SetLinkedStateTreeOverrides(FStateTreeReferenceOverrides InLinkedStateTreeOverrides)
 {
-	if (!LinkedStateTreeOverrides)
+	if (!IsValid())
 	{
-		return nullptr;
+		STATETREE_LOG(Warning, TEXT("%hs: StateTree context is not initialized properly ('%s' using StateTree '%s')"),
+			__FUNCTION__, *GetNameSafe(&Owner), *GetFullNameSafe(&RootStateTree));
+		return;
+	}
+
+	FStateTreeExecutionState& Exec = GetExecState();
+	if (!ensureMsgf(Exec.CurrentPhase == EStateTreeUpdatePhase::Unset, TEXT("%hs can't be called while already in %s ('%s' using StateTree '%s')."),
+		__FUNCTION__, *UEnum::GetDisplayValueAsText(Exec.CurrentPhase).ToString(), *GetNameSafe(&Owner), *GetFullNameSafe(&RootStateTree)))
+	{
+		return;
+	}
+
+	bool bValid = true;
+
+	// Confirms that the overrides schema matches.
+	const TConstArrayView<FStateTreeReferenceOverrideItem> OverrideItems = LinkedAssetStateTreeOverrides.GetOverrideItems();
+	for (const FStateTreeReferenceOverrideItem& Item : OverrideItems)
+	{
+		if (const UStateTree* ItemStateTree = Item.GetStateTreeReference().GetStateTree())
+		{
+			if (!ItemStateTree->IsReadyToRun())
+			{
+				STATETREE_LOG(Error, TEXT("%hs: '%s' using StateTree '%s' trying to set override '%s' but the tree is not initialized properly."),
+					__FUNCTION__, *GetNameSafe(&Owner), *GetFullNameSafe(GetStateTree()), *GetFullNameSafe(ItemStateTree));
+				bValid = false;
+				break;
+			}
+
+			const UStateTreeSchema* OverrideSchema = ItemStateTree->GetSchema();
+			if (ItemStateTree->GetSchema() == nullptr)
+			{
+				STATETREE_LOG(Error, TEXT("%hs: '%s' using StateTree '%s' trying to set override '%s' but the tree does not have a schema."),
+					__FUNCTION__, *GetNameSafe(&Owner), *GetFullNameSafe(GetStateTree()), *GetFullNameSafe(ItemStateTree));
+				bValid = false;
+				break;
+			}
+				
+			const bool bIsSameSchema = RootStateTree.GetSchema()->GetClass() == OverrideSchema->GetClass();
+			if (!bIsSameSchema)
+			{
+				STATETREE_LOG(Error, TEXT("%hs: '%s' using StateTree '%s' trying to set override '%s' but their schemas don't match."),
+					__FUNCTION__, *GetNameSafe(&Owner), *GetFullNameSafe(GetStateTree()), *GetFullNameSafe(Item.GetStateTreeReference().GetStateTree()));
+				bValid = false;
+				break;
+			}
+		}
+	}
+
+	if (bValid)
+	{
+		LinkedAssetStateTreeOverrides = MoveTemp(InLinkedStateTreeOverrides);
+	}
+	else
+	{
+		LinkedAssetStateTreeOverrides.Reset();
 	}
 	
-	for (const FStateTreeReferenceOverrideItem& Item : LinkedStateTreeOverrides->GetOverrideItems())
+	TInstancedStruct<FStateTreeExecutionExtension>& ExecutionExtension = Storage.GetMutableExecutionState().ExecutionExtension;
+	if (ExecutionExtension.IsValid())
 	{
-		if (Item.GetStateTag() == StateTag)
+		ExecutionExtension.GetMutable().OnLinkedStateTreeOverridesSet(FStateTreeExecutionExtension::FContextParameters(Owner, RootStateTree, Storage), LinkedAssetStateTreeOverrides);
+	}
+}
+
+const FStateTreeReference* FStateTreeExecutionContext::GetLinkedStateTreeOverrideForTag(const FGameplayTag StateTag) const
+{
+	for (const FStateTreeReferenceOverrideItem& Item : LinkedAssetStateTreeOverrides.GetOverrideItems())
+	{
+		if (StateTag.MatchesTag(Item.GetStateTag()))
 		{
 			return &Item.GetStateTreeReference();
 		}
@@ -471,7 +555,6 @@ const FStateTreeReference* FStateTreeExecutionContext::GetLinkedStateTreeOverrid
 	return nullptr;
 }
 
-
 bool FStateTreeExecutionContext::AreContextDataViewsValid() const
 {
 	if (!IsValid())
diff --git a/Engine/Plugins/Runtime/StateTree/Source/StateTreeModule/Public/StateTreeExecutionContext.h b/Engine/Plugins/Runtime/StateTree/Source/StateTreeModule/Public/StateTreeExecutionContext.h
--- a/Engine/Plugins/Runtime/StateTree/Source/StateTreeModule/Public/StateTreeExecutionContext.h
+++ b/Engine/Plugins/Runtime/StateTree/Source/StateTreeModule/Public/StateTreeExecutionContext.h
@@ -7,13 +7,13 @@
 #include "StateTreeExecutionTypes.h"
 #include "StateTreeNodeBase.h"
 #include "StateTreeNodeRef.h"
+#include "StateTreeReference.h"
 #include "Experimental/ConcurrentLinearAllocator.h"
 #include "Templates/IsInvocable.h"
 
 struct FGameplayTag;
 struct FInstancedPropertyBag;
 struct FStateTreeExecutionContext;
-struct FStateTreeReferenceOverrides;
 struct FStateTreeEvaluatorBase;
 struct FStateTreeTaskBase;
 struct FStateTreeConditionBase;
@@ -22,7 +22,6 @@ struct FStateTreeEvent;
 struct FStateTreeMinimalExecutionContext;
 struct FStateTreeTransitionRequest;
 struct FStateTreeInstanceDebugId;
-struct FStateTreeReference;
 struct FStateTreeWeakExecutionContext;
 
 /**
@@ -198,12 +197,19 @@ public:
 	/** Sets callback used to collect external data views during State Tree execution. */
 	void SetCollectExternalDataCallback(const FOnCollectStateTreeExternalData& Callback);
 
+	UE_DEPRECATED(5.6, "Use SetLinkedStateTreeOverrides that creates a copy.")
 	/**
 	 * Overrides for linked State Trees. This table is used to override State Tree references on linked states.
 	 * If a linked state's tag is exact match of the tag specified on the table, the reference from the table is used instead.
 	 */
 	void SetLinkedStateTreeOverrides(const FStateTreeReferenceOverrides* InLinkedStateTreeOverrides);
 
+	/**
+	 * Overrides for linked State Trees. This table is used to override State Tree references on linked states.
+	 * If a linked state's tag is exact match of the tag specified on the table, the reference from the table is used instead.
+	 */
+	void SetLinkedStateTreeOverrides(FStateTreeReferenceOverrides InLinkedStateTreeOverrides);
+
 	/** @return the first state tree reference set by SetLinkedStateTreeOverrides that matches the StateTag. Or null if not found. */
 	const FStateTreeReference* GetLinkedStateTreeOverrideForTag(const FGameplayTag StateTag) const;
 
@@ -925,8 +931,13 @@ protected:
 	/** Events queue to use, cached for less indirections. */
 	TSharedPtr<FStateTreeEventQueue> EventQueue;
 
+#if WITH_EDITORONLY_DATA
+	UE_DEPRECATED(5.6, "Use the LinkedStateTreeOverrides copy")
 	/** Pointer to linked state tree overrides. */
 	const FStateTreeReferenceOverrides* LinkedStateTreeOverrides = nullptr;
+#endif
+	/** Current linked state tree overrides. */
+	FStateTreeReferenceOverrides LinkedAssetStateTreeOverrides;
 	
 	/** Data view of the context data. */
 	TArray<FStateTreeDataView, TConcurrentLinearArrayAllocator<FDefaultBlockAllocationTag>> ContextAndExternalDataViews;
diff --git a/Engine/Plugins/Runtime/StateTree/Source/StateTreeModule/Public/StateTreeExecutionExtension.h b/Engine/Plugins/Runtime/StateTree/Source/StateTreeModule/Public/StateTreeExecutionExtension.h
--- a/Engine/Plugins/Runtime/StateTree/Source/StateTreeModule/Public/StateTreeExecutionExtension.h
+++ b/Engine/Plugins/Runtime/StateTree/Source/StateTreeModule/Public/StateTreeExecutionExtension.h
@@ -8,6 +8,7 @@ class UObject;
 class UStateTree;
 struct FStateTreeInstanceStorage;
 struct FStateTreeTransitionDelayedState;
+struct FStateTreeReferenceOverrides;
 
 
 /** Used by the execution context or a weak execution context to extend their functionalities. */
@@ -41,4 +42,9 @@ public:
 	virtual void ScheduleNextTick(const FContextParameters& Context)
 	{
 	}
+
+	/** Callback when the overrides are set to the execution context . */
+	virtual void OnLinkedStateTreeOverridesSet(const FContextParameters& Context, const FStateTreeReferenceOverrides& Overrides)
+	{
+	}
 };
diff --git a/Engine/Plugins/Runtime/StateTree/Source/StateTreeModule/Public/StateTreeReference.h b/Engine/Plugins/Runtime/StateTree/Source/StateTreeModule/Public/StateTreeReference.h
--- a/Engine/Plugins/Runtime/StateTree/Source/StateTreeModule/Public/StateTreeReference.h
+++ b/Engine/Plugins/Runtime/StateTree/Source/StateTreeModule/Public/StateTreeReference.h
@@ -82,14 +82,14 @@ struct STATETREEMODULE_API FStateTreeReference
 	void PostSerialize(const FArchive& Ar);
 
 protected:
-	UPROPERTY(EditAnywhere, Category = "")
+	UPROPERTY(EditAnywhere, Category = "StateTree")
 	TObjectPtr<UStateTree> StateTree = nullptr;
 
-	UPROPERTY(EditAnywhere, Category = "", meta = (FixedLayout))
+	UPROPERTY(EditAnywhere, Category = "StateTree", meta = (FixedLayout))
 	FInstancedPropertyBag Parameters;
 
 	/** Array of overridden properties. Non-overridden properties will inherit the values from the StateTree default parameters. */
-	UPROPERTY(EditAnywhere, Category = "")
+	UPROPERTY(EditAnywhere, Category = "StateTree")
 	TArray<FGuid> PropertyOverrides;
 
 	friend class FStateTreeReferenceDetails;
@@ -130,15 +130,19 @@ struct STATETREEMODULE_API FStateTreeReferenceOverrideItem
 	{
 		return StateTreeReference;
 	}
+
+	friend uint32 GetTypeHash(const FStateTreeReferenceOverrideItem& Item)
+	{
+		return HashCombine(GetTypeHash(Item.StateTag), GetTypeHash(Item.StateTreeReference.GetStateTree()));
+	}
 	
 private:
-
 	/** Exact tag used to match against a tag on a linked State Tree state. */
-	UPROPERTY(EditAnywhere, Category = "")
+	UPROPERTY(EditAnywhere, Category = "StateTree")
 	FGameplayTag StateTag;
 
 	/** State Tree and parameters to replace the linked state asset with. */
-	UPROPERTY(EditAnywhere, Category = "", meta=(SchemaCanBeOverriden))
+	UPROPERTY(EditAnywhere, Category = "StateTree", meta=(SchemaCanBeOverriden))
 	FStateTreeReference StateTreeReference;
 	
 	friend class FStateTreeReferenceOverridesDetails;
@@ -156,9 +160,9 @@ struct STATETREEMODULE_API FStateTreeReferenceOverrides
 	/** Removes all overrides. */
 	void Reset()
 	{
-		OverrideItems.Reset();	
+		OverrideItems.Empty();
 	}
-	
+
 	/** Adds or replaces override for a selected tag. */
 	void AddOverride(const FGameplayTag StateTag, const FStateTreeReference& StateTreeReference)
 	{
@@ -196,11 +200,16 @@ struct STATETREEMODULE_API FStateTreeReferenceOverrides
 
 	TConstArrayView<FStateTreeReferenceOverrideItem> GetOverrideItems() const
 	{
-		return OverrideItems;		
+		return OverrideItems;
 	}
-	
+
+	friend uint32 GetTypeHash(const FStateTreeReferenceOverrides& Overrides)
+	{
+		return GetTypeHash(Overrides.OverrideItems);
+	}
+
 private:
-	UPROPERTY(EditAnywhere, Category = "")
+	UPROPERTY(EditAnywhere, Category = "StateTree")
 	TArray<FStateTreeReferenceOverrideItem> OverrideItems;
 
 	friend class FStateTreeReferenceOverridesDetails;
diff --git a/Engine/Plugins/Runtime/StateTree/Source/StateTreeTestSuite/Private/StateTreeTest.cpp b/Engine/Plugins/Runtime/StateTree/Source/StateTreeTestSuite/Private/StateTreeTest.cpp
--- a/Engine/Plugins/Runtime/StateTree/Source/StateTreeTestSuite/Private/StateTreeTest.cpp
+++ b/Engine/Plugins/Runtime/StateTree/Source/StateTreeTestSuite/Private/StateTreeTest.cpp
@@ -1858,7 +1858,7 @@ struct FStateTreeTest_NestedOverride : FStateTreeTestBase
 			Overrides.AddOverride(Tag, OverrideRef);
 			
 			FTestStateTreeExecutionContext Exec(StateTree, StateTree, InstanceData);
-			Exec.SetLinkedStateTreeOverrides(&Overrides);
+			Exec.SetLinkedStateTreeOverrides(Overrides);
 			
 			const bool bInitSucceeded = Exec.IsValid();
 			AITEST_TRUE("StateTree should init", bInitSucceeded);
