StateTree: Added support delegate listener binding with the input category. It forces them to be bound in the editor.

diff --git a/Engine/Plugins/Runtime/StateTree/Source/StateTreeEditorModule/Private/StateTreeCompiler.cpp b/Engine/Plugins/Runtime/StateTree/Source/StateTreeEditorModule/Private/StateTreeCompiler.cpp
--- a/Engine/Plugins/Runtime/StateTree/Source/StateTreeEditorModule/Private/StateTreeCompiler.cpp
+++ b/Engine/Plugins/Runtime/StateTree/Source/StateTreeEditorModule/Private/StateTreeCompiler.cpp
@@ -2250,7 +2250,9 @@ bool FStateTreeCompiler::GetAndValidateBindings(const FStateTreeBindableStructDe
 			if (Usage == EStateTreePropertyUsage::Input)
 			{
 				// Make sure that an Input property is bound unless marked optional.
-				if (bIsOptional == false && !IsPropertyBound(PropertyName, OutValidatedBindings.CopyBindings))
+				if (bIsOptional == false
+					&& !IsPropertyBound(PropertyName, OutValidatedBindings.CopyBindings)
+					&& !IsPropertyBound(PropertyName, OutValidatedBindings.DelegateListeners))
 				{
 					Log.Reportf(EMessageSeverity::Error, TargetStruct,
 						TEXT("Input property '%s' on %s is expected to have a binding."),
