Make ST draw elements that would be smaller than the horizontal size reduction.

#tests PIE

diff --git a/Engine/Plugins/Runtime/StateTree/Source/StateTreeEditorModule/Private/Debugger/SStateTreeDebuggerEventTimelineView.cpp b/Engine/Plugins/Runtime/StateTree/Source/StateTreeEditorModule/Private/Debugger/SStateTreeDebuggerEventTimelineView.cpp
--- a/Engine/Plugins/Runtime/StateTree/Source/StateTreeEditorModule/Private/Debugger/SStateTreeDebuggerEventTimelineView.cpp
+++ b/Engine/Plugins/Runtime/StateTree/Source/StateTreeEditorModule/Private/Debugger/SStateTreeDebuggerEventTimelineView.cpp
@@ -54,14 +54,18 @@ int32 SStateTreeDebuggerEventTimelineView::PaintEvents(const FGeometry& Allotted
 				const float XEnd = RangeToScreen.InputToLocalX(static_cast<float>(Window.TimeEnd));
 				const float Y = (AllottedGeometry.Size.Y - EventSize.Y) / 2.f;
 
-				// window bar
+				// Window bar
 				constexpr float HSizeReduction = 2.f;
 				constexpr float VSizeReduction = 2.f;
 				constexpr float VerticalOffset = 0.5f * VSizeReduction;
+
+				// Ensure at least 1 pixel remains after horizontal size reduction, else don't reduce size
+				const float XDelta = (XEnd - XStart);
+				const float Width = XDelta - 1 > HSizeReduction ? XDelta - HSizeReduction : XDelta;
 				FSlateDrawElement::MakeBox
 				(OutDrawElements,
 					LayerId++,
-					AllottedGeometry.ToPaintGeometry(FVector2f((XEnd - XStart) - HSizeReduction, EventSize.Y - VSizeReduction), FSlateLayoutTransform(FVector2f(XStart, Y + VerticalOffset))),
+					AllottedGeometry.ToPaintGeometry(FVector2f(Width, EventSize.Y - VSizeReduction), FSlateLayoutTransform(FVector2f(XStart, Y + VerticalOffset))),
 					Brush,
 					ESlateDrawEffect::None,
 					Window.Color);
@@ -72,7 +76,7 @@ int32 SStateTreeDebuggerEventTimelineView::PaintEvents(const FGeometry& Allotted
 					FSlateFontInfo FontInfo = FCoreStyle::GetDefaultFontStyle("Regular", 8);
 					const TSharedRef<FSlateFontMeasure> FontMeasureService = FSlateApplication::Get().GetRenderer()->GetFontMeasureService();
 					constexpr int32 TextLeftOffset = 8;
-					const int32 HorizontalOffset = FMath::RoundToInt((XEnd - XStart) - TextLeftOffset);
+					const int32 HorizontalOffset = FMath::RoundToInt(XDelta - TextLeftOffset);
 					const int32 LastWholeCharacterIndex = FontMeasureService->FindLastWholeCharacterIndexBeforeOffset(Description, FontInfo, HorizontalOffset);
 
 					if (LastWholeCharacterIndex >= 0)
