StateTree: Fixed accessing memory deallocated memory.
#pr

diff --git a/Engine/Plugins/Runtime/StateTree/Source/StateTreeEditorModule/Private/StateTreeCompiler.cpp b/Engine/Plugins/Runtime/StateTree/Source/StateTreeEditorModule/Private/StateTreeCompiler.cpp
--- a/Engine/Plugins/Runtime/StateTree/Source/StateTreeEditorModule/Private/StateTreeCompiler.cpp
+++ b/Engine/Plugins/Runtime/StateTree/Source/StateTreeEditorModule/Private/StateTreeCompiler.cpp
@@ -1287,9 +1287,6 @@ bool FStateTreeCompiler::CreateBindingsForNodes(TConstArrayView<FStateTreeEditor
 
 		FStateTreeNodeBase& Node = Nodes[NodeIndex++].GetMutable<FStateTreeNodeBase>();
 
-		const FStateTreeBindableStructDesc* BindableStruct = BindingsCompiler.GetSourceStructDescByID(EditorNode.ID);
-		check(BindableStruct);
-
 		FStateTreeIndex16 PropertyFunctionsBegin(Nodes.Num());
 		if (!CreatePropertyFunctionsForStruct(EditorNode.ID))
 		{
@@ -1319,10 +1316,14 @@ bool FStateTreeCompiler::CreateBindingsForNodes(TConstArrayView<FStateTreeEditor
 			InstanceView = FStateTreeDataView(Instance);
 		}
 
-		if (!CreateBindingsForStruct(*BindableStruct, InstanceView, PropertyFunctionsBegin, PropertyFunctionsEnd, Node.BindingsBatch))
 		{
-			bSucceeded = false;
-			continue;
+			const FStateTreeBindableStructDesc* BindableStruct = BindingsCompiler.GetSourceStructDescByID(EditorNode.ID);
+			check(BindableStruct);
+			if (!CreateBindingsForStruct(*BindableStruct, InstanceView, PropertyFunctionsBegin, PropertyFunctionsEnd, Node.BindingsBatch))
+			{
+				bSucceeded = false;
+				continue;
+			}
 		}
 	}
 
