StateTree: Removed property's custom label if it's member is bound.
#pr
[RN] Behavior was wrongly introduced in 5.5. Rollback to previous behavior.

diff --git a/Engine/Plugins/Runtime/StateTree/Source/StateTreeEditorModule/Private/Customizations/StateTreeEditorNodeDetails.cpp b/Engine/Plugins/Runtime/StateTree/Source/StateTreeEditorModule/Private/Customizations/StateTreeEditorNodeDetails.cpp
--- a/Engine/Plugins/Runtime/StateTree/Source/StateTreeEditorModule/Private/Customizations/StateTreeEditorNodeDetails.cpp
+++ b/Engine/Plugins/Runtime/StateTree/Source/StateTreeEditorModule/Private/Customizations/StateTreeEditorNodeDetails.cpp
@@ -174,12 +174,11 @@ namespace UE::StateTreeEditor::Internal
 			FDetailWidgetRow Row;
 			ChildRow.GetDefaultWidgets(NameWidget, ValueWidget, Row);
 
-			const bool bHasChildPropertyBinding = EditorPropBindings->HasBinding(Path, FPropertyBindingBindingCollection::ESearchMode::Includes);
 			const bool bValidUsage = Usage == EStateTreePropertyUsage::Input || Usage == EStateTreePropertyUsage::Output || Usage == EStateTreePropertyUsage::Context;
 			const bool bIsDelegateDispatcher = UE::StateTreeEditor::Internal::IsDelegateDispatcherProperty(*Property);
 
-			if (bHasChildPropertyBinding || bValidUsage || bIsDelegateDispatcher)
-			{				
+			if (bValidUsage || bIsDelegateDispatcher)
+			{
 				FEdGraphPinType PinType;
 				const UEdGraphSchema_K2* Schema = GetDefault<UEdGraphSchema_K2>();
 
@@ -270,12 +269,6 @@ namespace UE::StateTreeEditor::Internal
 						IconColor = FLinearColor::White;
 					}
 				}
-				else if (ensure(bHasChildPropertyBinding))
-				{
-					Icon = FCoreStyle::Get().GetBrush("Icons.Link");
-					Text = LOCTEXT("PropertyMemberIsBoundTooltip", "{Member(s) Bound_B}");
-					TextColor = FSlateColor::UseSubduedForeground();
-				}
 				
 				ChildRow
 					.CustomWidget(true)
