StateTree: Update TStateTreeInstanceDataStructRef to access data in temporary frames. When selecting a state the evaluator and global task data are not in the active frame.

diff --git a/Engine/Plugins/Runtime/StateTree/Source/StateTreeModule/Public/StateTreeInstanceData.h b/Engine/Plugins/Runtime/StateTree/Source/StateTreeModule/Public/StateTreeInstanceData.h
--- a/Engine/Plugins/Runtime/StateTree/Source/StateTreeModule/Public/StateTreeInstanceData.h
+++ b/Engine/Plugins/Runtime/StateTree/Source/StateTreeModule/Public/StateTreeInstanceData.h
@@ -617,6 +617,19 @@ PRAGMA_ENABLE_DEPRECATION_WARNINGS
 
 			check(Struct.GetScriptStruct() == TBaseStructure<T>::Get());
 		}
+		else
+		{
+			// When selecting a state, the frame is not in the active list.
+			FStateTreeTemporaryInstanceData* ExistingInstance = StoragePtr->GetMutableTemporaryInstances().FindByPredicate([FrameID = FrameID, DataHandle = DataHandle](const FStateTreeTemporaryInstanceData& TempInstance)
+				{
+					return TempInstance.FrameID == FrameID
+						&& TempInstance.DataHandle == DataHandle;
+				});
+			if (ExistingInstance)
+			{
+				Struct = FStructView(ExistingInstance->Instance);
+			}
+		}
 		return reinterpret_cast<T*>(Struct.GetMemory());
 	}
 
