StateTree: Update pending delayed transitions to use the state ID instead of the {frame. state} pair.

diff --git a/Engine/Plugins/Runtime/StateTree/Source/StateTreeModule/Private/StateTreeExecutionContext.cpp b/Engine/Plugins/Runtime/StateTree/Source/StateTreeModule/Private/StateTreeExecutionContext.cpp
--- a/Engine/Plugins/Runtime/StateTree/Source/StateTreeModule/Private/StateTreeExecutionContext.cpp
+++ b/Engine/Plugins/Runtime/StateTree/Source/StateTreeModule/Private/StateTreeExecutionContext.cpp
@@ -2706,6 +2706,7 @@ void FStateTreeExecutionContext::ExitState(const FStateTreeTransitionResult& Tra
 		for (int32 StateIndex = CurrentFrame.ActiveStates.Num() - 1; StateIndex >= 0; StateIndex--)
 		{
 			const FStateTreeStateHandle CurrentHandle = CurrentFrame.ActiveStates[StateIndex];
+			const UE::StateTree::FActiveStateID CurrentStateID = CurrentFrame.ActiveStates.StateIDs[StateIndex];
 			const FCompactStateTreeState& State = CurrentStateTree->States[CurrentHandle.Index];
 
 			const FExitStateCall& ExitCall = ExitStateCalls[CallIndex--];
@@ -2723,9 +2724,11 @@ void FStateTreeExecutionContext::ExitState(const FStateTreeTransitionResult& Tra
 
 				// Remove any delayed transitions that belong to this state.
 				Exec.DelayedTransitions.RemoveAllSwap(
-					[StateTree = CurrentFrame.StateTree, Begin = State.TransitionsBegin, End = State.TransitionsBegin + State.TransitionsNum](const FStateTreeTransitionDelayedState& DelayedState)
+					[CurrentStateID, Begin = State.TransitionsBegin, End = State.TransitionsBegin + State.TransitionsNum](const FStateTreeTransitionDelayedState& DelayedState)
 					{
-						return  DelayedState.StateTree == StateTree && DelayedState.TransitionIndex.Get() >= Begin && DelayedState.TransitionIndex.Get() < End;
+						return  DelayedState.StateID == CurrentStateID
+							&& DelayedState.TransitionIndex.Get() >= Begin
+							&& DelayedState.TransitionIndex.Get() < End;
 					});
 
 				CurrentTransition.CurrentState = CurrentHandle;
@@ -3985,22 +3988,15 @@ bool FStateTreeExecutionContext::TriggerTransitions()
 	}
 	InstanceData.ResetTransitionRequests();
 
-
 	//
 	// Collect expired delayed transitions
 	//
-	struct FExpiredTransitionDelayedState
-	{
-		FStateTreeTransitionDelayedState Transition;
-		int32 FrameIndex;
-	};
-	TArray<FExpiredTransitionDelayedState, TInlineAllocator<8>> ExpiredTransitionsDelayed;
+	TArray<FStateTreeTransitionDelayedState, TInlineAllocator<8>> ExpiredTransitionsDelayed;
 	for (TArray<FStateTreeTransitionDelayedState>::TIterator It = Exec.DelayedTransitions.CreateIterator(); It; ++It)
 	{
 		if (It->TimeLeft <= 0.0f)
 		{
-			int32 FrameIndex = Exec.ActiveFrames.IndexOfByPredicate([StateTree = It->StateTree](const FStateTreeExecutionFrame& Other){ return Other.StateTree == StateTree; });
-			ExpiredTransitionsDelayed.Emplace(FExpiredTransitionDelayedState{ .Transition = MoveTemp(*It), .FrameIndex = FrameIndex });
+			ExpiredTransitionsDelayed.Emplace(MoveTemp(*It));
 			It.RemoveCurrentSwap();
 		}
 	}
@@ -4012,16 +4008,18 @@ bool FStateTreeExecutionContext::TriggerTransitions()
 	{
 		FTransitionHandler() = default;
 		
-		FTransitionHandler(const uint8 InFrameIndex, const FStateTreeStateHandle InStateHandle, const EStateTreeTransitionPriority InPriority)
+		FTransitionHandler(const int32 InFrameIndex, const FStateTreeStateHandle InStateHandle, const UE::StateTree::FActiveStateID InStateID, const EStateTreeTransitionPriority InPriority)
 			: StateHandle(InStateHandle)
+			, StateID(InStateID)
 			, TaskIndex(FStateTreeIndex16::Invalid)
 			, FrameIndex(InFrameIndex)
 			, Priority(InPriority)
 		{
 		}
 
-		FTransitionHandler(const uint8 InFrameIndex, const FStateTreeStateHandle InStateHandle, const FStateTreeIndex16 InTaskIndex, const EStateTreeTransitionPriority InPriority)
+		FTransitionHandler(const int32 InFrameIndex, const FStateTreeStateHandle InStateHandle, const UE::StateTree::FActiveStateID InStateID, const FStateTreeIndex16 InTaskIndex, const EStateTreeTransitionPriority InPriority)
 			: StateHandle(InStateHandle)
+			, StateID(InStateID)
 			, TaskIndex(InTaskIndex)
 			, FrameIndex(InFrameIndex)
 			, Priority(InPriority)
@@ -4029,8 +4027,9 @@ bool FStateTreeExecutionContext::TriggerTransitions()
 		}
 
 		FStateTreeStateHandle StateHandle;
+		UE::StateTree::FActiveStateID StateID;
 		FStateTreeIndex16 TaskIndex = FStateTreeIndex16::Invalid;
-		uint8 FrameIndex = 0;
+		int32 FrameIndex = 0;
 		EStateTreeTransitionPriority Priority = EStateTreeTransitionPriority::Normal;
 
 		bool operator<(const FTransitionHandler& Other) const
@@ -4055,6 +4054,7 @@ bool FStateTreeExecutionContext::TriggerTransitions()
 			for (int32 StateIndex = CurrentFrame.ActiveStates.Num() - 1; StateIndex >= 0; StateIndex--)
 			{
 				const FStateTreeStateHandle StateHandle = CurrentFrame.ActiveStates[StateIndex];
+				const UE::StateTree::FActiveStateID StateID = CurrentFrame.ActiveStates.StateIDs[StateIndex];
 				const FCompactStateTreeState& State = CurrentStateTree->States[StateHandle.Index];
 
 				// Do not process any transitions from a disabled state
@@ -4072,7 +4072,7 @@ bool FStateTreeExecutionContext::TriggerTransitions()
 						const FStateTreeTaskBase& Task = CurrentStateTree->Nodes[TaskIndex].Get<const FStateTreeTaskBase>();
 						if (Task.bShouldAffectTransitions && Task.bTaskEnabled)
 						{
-							TransitionHandlers.Emplace((uint8)FrameIndex, StateHandle, FStateTreeIndex16(TaskIndex), Task.TransitionHandlingPriority);
+							TransitionHandlers.Emplace(FrameIndex, StateHandle, StateID, FStateTreeIndex16(TaskIndex), Task.TransitionHandlingPriority);
 							bAdded = true;
 						}
 					}
@@ -4080,9 +4080,9 @@ bool FStateTreeExecutionContext::TriggerTransitions()
 				}
 
 				// Has expired transition delayed.
-				const bool bHasActiveTransitionDelayed = ExpiredTransitionsDelayed.ContainsByPredicate([FrameIndex, StateHandle](const FExpiredTransitionDelayedState& Other)
+				const bool bHasActiveTransitionDelayed = ExpiredTransitionsDelayed.ContainsByPredicate([StateID](const FStateTreeTransitionDelayedState& Other)
 				{
-					return Other.FrameIndex == FrameIndex && Other.Transition.StateHandle == StateHandle;
+					return Other.StateID == StateID;
 				});
 				
 				// Regular transitions on state
@@ -4090,7 +4090,7 @@ bool FStateTreeExecutionContext::TriggerTransitions()
 				//or Expired transition delayed
 				if (State.ShouldTickTransitions(bHasEvents, bHasBroadcastedDelegates) || State.bHasTransitionTasks || bHasActiveTransitionDelayed)
 				{
-					TransitionHandlers.Emplace((uint8)FrameIndex, StateHandle, EStateTreeTransitionPriority::Normal);
+					TransitionHandlers.Emplace(FrameIndex, StateHandle, StateID, EStateTreeTransitionPriority::Normal);
 				}
 			}
 
@@ -4105,7 +4105,7 @@ bool FStateTreeExecutionContext::TriggerTransitions()
 						const FStateTreeTaskBase& Task = CurrentStateTree->Nodes[TaskIndex].Get<const FStateTreeTaskBase>();
 						if (Task.bShouldAffectTransitions && Task.bTaskEnabled)
 						{
-							TransitionHandlers.Emplace((uint8)FrameIndex, FStateTreeStateHandle(), FStateTreeIndex16(TaskIndex), Task.TransitionHandlingPriority);
+							TransitionHandlers.Emplace(FrameIndex, FStateTreeStateHandle(), UE::StateTree::FActiveStateID::Invalid, FStateTreeIndex16(TaskIndex), Task.TransitionHandlingPriority);
 							bAdded = true;
 						}
 					}
@@ -4126,6 +4126,7 @@ bool FStateTreeExecutionContext::TriggerTransitions()
 		const int32 FrameIndex = Handler.FrameIndex;
 		FStateTreeExecutionFrame* CurrentParentFrame = FrameIndex > 0 ? &Exec.ActiveFrames[FrameIndex - 1] : nullptr;
 		FStateTreeExecutionFrame& CurrentFrame = Exec.ActiveFrames[FrameIndex];
+		const UE::StateTree::FActiveFrameID CurrentFrameID = CurrentFrame.FrameID;
 		const UStateTree* CurrentStateTree = CurrentFrame.StateTree;
 
 		FCurrentlyProcessedFrameScope FrameScope(*this, CurrentParentFrame, CurrentFrame);
@@ -4160,6 +4161,7 @@ bool FStateTreeExecutionContext::TriggerTransitions()
 		}
 		else if (Handler.StateHandle.IsValid())
 		{
+			check(Handler.StateID.IsValid());
 			const FCompactStateTreeState& State = CurrentStateTree->States[Handler.StateHandle.Index];
 
 			// Transitions
@@ -4191,20 +4193,20 @@ bool FStateTreeExecutionContext::TriggerTransitions()
 				if (Transition.HasDelay())
 				{
 					bool bTriggeredDelayedTransition = false;
-					for (const FExpiredTransitionDelayedState& DelayedTransition : ExpiredTransitionsDelayed)
+					for (const FStateTreeTransitionDelayedState& DelayedTransition : ExpiredTransitionsDelayed)
 					{
-						if (DelayedTransition.FrameIndex == FrameIndex && DelayedTransition.Transition.TransitionIndex == FStateTreeIndex16(TransitionIndex))
+						if (DelayedTransition.StateID == Handler.StateID && DelayedTransition.TransitionIndex == FStateTreeIndex16(TransitionIndex))
 						{
 							STATETREE_LOG(Verbose, TEXT("Passed delayed transition from '%s' (%s) -> '%s'"),
 								*GetSafeStateName(CurrentFrame, CurrentFrame.ActiveStates.Last()), *State.Name.ToString(), *GetSafeStateName(CurrentFrame, Transition.State));
 
 							// Trigger Delayed Transition when the delay has passed.
-							if (RequestTransition(CurrentFrame, Transition.State, Transition.Priority, &DelayedTransition.Transition.CapturedEvent, Transition.Fallback))
+							if (RequestTransition(CurrentFrame, Transition.State, Transition.Priority, &DelayedTransition.CapturedEvent, Transition.Fallback))
 							{
 								// If the transition was successfully requested with a specific event, consume and remove the event, it's been used.
-								if (DelayedTransition.Transition.CapturedEvent.IsValid() && Transition.bConsumeEventOnSelect)
+								if (DelayedTransition.CapturedEvent.IsValid() && Transition.bConsumeEventOnSelect)
 								{
-									ConsumeEvent(DelayedTransition.Transition.CapturedEvent);
+									ConsumeEvent(DelayedTransition.CapturedEvent);
 								}
 
 								NextTransitionSource = FStateTreeTransitionSource(CurrentFrame.StateTree, FStateTreeIndex16(TransitionIndex), Transition.State, Transition.Priority);
@@ -4275,9 +4277,11 @@ bool FStateTreeExecutionContext::TriggerTransitions()
 								TransitionEventHash = GetTypeHash(*TransitionEvent->Get());
 							}
 
-							const bool bIsDelayedTransitionExisting = Exec.DelayedTransitions.ContainsByPredicate([StateTree = CurrentFrame.StateTree, TransitionIndex, TransitionEventHash](const FStateTreeTransitionDelayedState& DelayedState)
+							const bool bIsDelayedTransitionExisting = Exec.DelayedTransitions.ContainsByPredicate([CurrentStateID =  Handler.StateID, TransitionIndex, TransitionEventHash](const FStateTreeTransitionDelayedState& DelayedState)
 							{
-								return DelayedState.StateTree == StateTree && DelayedState.TransitionIndex.Get() == TransitionIndex && DelayedState.CapturedEventHash == TransitionEventHash;
+								return DelayedState.StateID == CurrentStateID
+									&& DelayedState.TransitionIndex.Get() == TransitionIndex
+									&& DelayedState.CapturedEventHash == TransitionEventHash;
 							});
 
 							if (!bIsDelayedTransitionExisting)
@@ -4287,8 +4291,11 @@ bool FStateTreeExecutionContext::TriggerTransitions()
 								if (DelayDuration > 0.0f)
 								{
 									FStateTreeTransitionDelayedState& DelayedState = Exec.DelayedTransitions.AddDefaulted_GetRef();
+									DelayedState.StateID = Handler.StateID;
+PRAGMA_DISABLE_DEPRECATION_WARNINGS
 									DelayedState.StateTree = CurrentFrame.StateTree;
 									DelayedState.StateHandle = Handler.StateHandle;
+PRAGMA_ENABLE_DEPRECATION_WARNINGS
 									DelayedState.TransitionIndex = FStateTreeIndex16(TransitionIndex);
 									DelayedState.TimeLeft = DelayDuration;
 									if (TransitionEvent && TransitionEvent->IsValid())
diff --git a/Engine/Plugins/Runtime/StateTree/Source/StateTreeModule/Public/StateTreeExecutionTypes.h b/Engine/Plugins/Runtime/StateTree/Source/StateTreeModule/Public/StateTreeExecutionTypes.h
--- a/Engine/Plugins/Runtime/StateTree/Source/StateTreeModule/Public/StateTreeExecutionTypes.h
+++ b/Engine/Plugins/Runtime/StateTree/Source/StateTreeModule/Public/StateTreeExecutionTypes.h
@@ -642,6 +642,10 @@ struct STATETREEMODULE_API FStateTreeTransitionDelayedState
 {
 	GENERATED_BODY()
 
+	/** The state ID that triggers the transition. */
+	UE::StateTree::FActiveStateID StateID;
+
+	UE_DEPRECATED(5.6, "StateTree is unused. Use StateID instead.")
 	UPROPERTY()
 	TObjectPtr<const UStateTree> StateTree = nullptr;
 
@@ -654,6 +658,7 @@ struct STATETREEMODULE_API FStateTreeTransitionDelayedState
 	UPROPERTY()
 	uint32 CapturedEventHash = 0u;
 
+	UE_DEPRECATED(5.6, "StateHandle is unused. Use StateID instead.")
 	UPROPERTY()
 	FStateTreeStateHandle StateHandle;
 	
@@ -900,6 +905,7 @@ public:
 	/** Finds all delayed transition states for a specific transition and removes them. Returns their copies. */
 	TArray<FStateTreeTransitionDelayedState, TInlineAllocator<8>> FindAndRemoveExpiredDelayedTransitions(const UStateTree* OwnerStateTree, const FStateTreeIndex16 TransitionIndex)
 	{
+PRAGMA_DISABLE_DEPRECATION_WARNINGS
 		TArray<FStateTreeTransitionDelayedState, TInlineAllocator<8>> Result;
 		for (TArray<FStateTreeTransitionDelayedState>::TIterator It = DelayedTransitions.CreateIterator(); It; ++It)
 		{
@@ -911,6 +917,7 @@ public:
 		}
 
 		return Result;
+PRAGMA_ENABLE_DEPRECATION_WARNINGS
 	}
 
 	/** @return the active frame that matches the unique frame ID. */
