Fix compilation error introduced with CL 39565986.

#rnx

diff --git a/Engine/Plugins/Runtime/StateTree/Source/StateTreeModule/Public/StateTreeExecutionTypes.h b/Engine/Plugins/Runtime/StateTree/Source/StateTreeModule/Public/StateTreeExecutionTypes.h
--- a/Engine/Plugins/Runtime/StateTree/Source/StateTreeModule/Public/StateTreeExecutionTypes.h
+++ b/Engine/Plugins/Runtime/StateTree/Source/StateTreeModule/Public/StateTreeExecutionTypes.h
@@ -642,6 +642,14 @@ struct STATETREEMODULE_API FStateTreeTransitionDelayedState
 {
 	GENERATED_BODY()
 
+PRAGMA_DISABLE_DEPRECATION_WARNINGS
+	FStateTreeTransitionDelayedState() = default;
+	FStateTreeTransitionDelayedState(const FStateTreeTransitionDelayedState&) = default;
+	FStateTreeTransitionDelayedState(FStateTreeTransitionDelayedState&&) = default;
+	FStateTreeTransitionDelayedState& operator=(const FStateTreeTransitionDelayedState&) = default;
+	FStateTreeTransitionDelayedState& operator=(FStateTreeTransitionDelayedState&&) = default;
+PRAGMA_ENABLE_DEPRECATION_WARNINGS
+
 	/** The state ID that triggers the transition. */
 	UE::StateTree::FActiveStateID StateID;
 
