StateTree: Remove selection when the debugging session closes or when the tab containing the debugging view closes.

diff --git a/Engine/Plugins/Runtime/StateTree/Source/StateTreeEditorModule/Private/Debugger/SStateTreeDebuggerView.cpp b/Engine/Plugins/Runtime/StateTree/Source/StateTreeEditorModule/Private/Debugger/SStateTreeDebuggerView.cpp
--- a/Engine/Plugins/Runtime/StateTree/Source/StateTreeEditorModule/Private/Debugger/SStateTreeDebuggerView.cpp
+++ b/Engine/Plugins/Runtime/StateTree/Source/StateTreeEditorModule/Private/Debugger/SStateTreeDebuggerView.cpp
@@ -214,6 +214,8 @@ SStateTreeDebuggerView::~SStateTreeDebuggerView()
 	Debugger->OnNewInstance.Unbind();
 	Debugger->OnSelectedInstanceCleared.Unbind();
 
+	Debugger->ClearSelection();
+
 	FEditorDelegates::BeginPIE.RemoveAll(this);
 	FEditorDelegates::EndPIE.RemoveAll(this);
 	FEditorDelegates::PausePIE.RemoveAll(this);
diff --git a/Engine/Plugins/Runtime/StateTree/Source/StateTreeEditorModule/Private/StateTreeEditorModeToolkit.cpp b/Engine/Plugins/Runtime/StateTree/Source/StateTreeEditorModule/Private/StateTreeEditorModeToolkit.cpp
--- a/Engine/Plugins/Runtime/StateTree/Source/StateTreeEditorModule/Private/StateTreeEditorModeToolkit.cpp
+++ b/Engine/Plugins/Runtime/StateTree/Source/StateTreeEditorModule/Private/StateTreeEditorModeToolkit.cpp
@@ -131,7 +131,12 @@ void FStateTreeEditorModeToolkit::RequestModeUITabs()
 				ensure(Args.GetTabId() == UAssetEditorUISubsystem::TopRightTabID);
 				TSharedRef<SDockTab> Tab = SNew(SDockTab)
 				.Label(LOCTEXT("DebuggerTitle", "Debugger"))
-				.TabColorScale(GetTabColorScale());
+				.TabColorScale(GetTabColorScale())
+				.OnTabClosed(SDockTab::FOnTabClosedCallback::CreateLambda([](TSharedRef<SDockTab> Tab)
+					{
+						// Destroy the SStateTreeDebuggerView
+						Tab->SetContent(SNullWidget::NullWidget);
+					}));
 
 				WeakDebuggerTab = Tab;
 				UpdateDebuggerView();
@@ -600,7 +605,7 @@ void FStateTreeEditorModeToolkit::UpdateStateTreeOutliner()
 #if WITH_STATETREE_TRACE_DEBUGGER
 void FStateTreeEditorModeToolkit::UpdateDebuggerView()
 {
-	DebuggerView = SNullWidget::NullWidget;
+	TSharedRef<SWidget> DebuggerView = SNullWidget::NullWidget;
 	TSharedPtr<SDockTab> SharedDebuggerTab = WeakDebuggerTab.Pin();
 
 	if (SharedDebuggerTab)
@@ -608,7 +613,7 @@ void FStateTreeEditorModeToolkit::UpdateDebuggerView()
 		// Clear any references the previous tab might have to a previous debugger view.
 		// The view will clear the shared debugger's bindings on dtor.
 		// We don't want it to clear newly created bindings from our new view.
-		SharedDebuggerTab->SetContent(SNullWidget::NullWidget);
+		SharedDebuggerTab->SetContent(DebuggerView);
 	}
 	
 	if (UStateTree* StateTree = WeakEditorMode->GetStateTree())
@@ -621,7 +626,7 @@ void FStateTreeEditorModeToolkit::UpdateDebuggerView()
 
 	if (SharedDebuggerTab)
 	{
-		SharedDebuggerTab->SetContent(DebuggerView.ToSharedRef());
+		SharedDebuggerTab->SetContent(DebuggerView);
 	}
 }
 #endif // WITH_STATETREE_TRACE_DEBUGGER
diff --git a/Engine/Plugins/Runtime/StateTree/Source/StateTreeEditorModule/Private/StateTreeViewModel.cpp b/Engine/Plugins/Runtime/StateTree/Source/StateTreeEditorModule/Private/StateTreeViewModel.cpp
--- a/Engine/Plugins/Runtime/StateTree/Source/StateTreeEditorModule/Private/StateTreeViewModel.cpp
+++ b/Engine/Plugins/Runtime/StateTree/Source/StateTreeEditorModule/Private/StateTreeViewModel.cpp
@@ -1325,7 +1325,11 @@ void FStateTreeViewModel::BindToDebuggerDelegates()
 #if WITH_STATETREE_TRACE_DEBUGGER
 	Debugger->OnActiveStatesChanged.BindSPLambda(this, [this](const FStateTreeTraceActiveStates& NewActiveStates)
 	{
-		if (const UStateTree* OuterStateTree = GetStateTree())
+		if (NewActiveStates.PerAssetStates.Num() == 0)
+		{
+			ActiveStates.Empty();
+		}
+		else if (const UStateTree* OuterStateTree = GetStateTree())
 		{
 			for (const FStateTreeTraceActiveStates::FAssetActiveStates& AssetActiveStates : NewActiveStates.PerAssetStates)
 			{
diff --git a/Engine/Plugins/Runtime/StateTree/Source/StateTreeEditorModule/Public/StateTreeEditorModeToolkit.h b/Engine/Plugins/Runtime/StateTree/Source/StateTreeEditorModule/Public/StateTreeEditorModeToolkit.h
--- a/Engine/Plugins/Runtime/StateTree/Source/StateTreeEditorModule/Public/StateTreeEditorModeToolkit.h
+++ b/Engine/Plugins/Runtime/StateTree/Source/StateTreeEditorModule/Public/StateTreeEditorModeToolkit.h
@@ -52,7 +52,6 @@ protected:
 
 #if WITH_STATETREE_TRACE_DEBUGGER
 	void UpdateDebuggerView();
-	TSharedPtr<SWidget> DebuggerView = nullptr;	
 	TWeakPtr<SDockTab> WeakDebuggerTab = nullptr;
 public:
 	static const FLazyName DebuggerTabId;
