[StateTreeDebugger] refactored event buffering to support instance execution for multiple worker threads.
Buffered data flushing (when starting a trace session) is now handled by callbacks from the main thread.
Locks are used with the assumption that the struct associated to a given state tree instance can only be used by a single worker thread in a given frame. Based on that assumption the locks are only used to protect creation of new entries when new instances are producing events. Note that those "Write" operations are a lot less frequent that the read accesses.

diff --git a/Engine/Plugins/Runtime/StateTree/Source/StateTreeModule/Private/Debugger/StateTreeTrace.cpp b/Engine/Plugins/Runtime/StateTree/Source/StateTreeModule/Private/Debugger/StateTreeTrace.cpp
--- a/Engine/Plugins/Runtime/StateTree/Source/StateTreeModule/Private/Debugger/StateTreeTrace.cpp
+++ b/Engine/Plugins/Runtime/StateTree/Source/StateTreeModule/Private/Debugger/StateTreeTrace.cpp
@@ -121,9 +121,7 @@ UE_TRACE_EVENT_END()
 
 namespace UE::StateTreeTrace
 {
-
-static double GWorldTime = 0.0;
-
+struct FBufferedDataList;
 FDelegateHandle GOnWorldTickStartDelegateHandle;
 FDelegateHandle GTracingStateChangedDelegateHandle;
 
@@ -141,13 +139,6 @@ struct FPhaseTraceStatusPair
 	bool bTraced = false;
 };
 
-/** Struct to keep track of the list of stacked phases for a given statetree instance. */
-struct FPhaseStack
-{
-	FStateTreeInstanceDebugId InstanceId;
-	TArray<FPhaseTraceStatusPair> Stack;
-};
-
 /**
  * Struct to hold data for asset debug id events until we are ready to trace the events (i.e. traces are active and channel is enabled).
  */
@@ -181,24 +172,7 @@ struct FInstanceEventBufferedData
 	struct FActiveStates
 	{
 		FActiveStates() = default;
-		explicit FActiveStates(const TConstArrayView<FStateTreeExecutionFrame> ActiveFrames)
-		{
-			for (const FStateTreeExecutionFrame& Frame : ActiveFrames)
-			{
-				const FStateTreeIndex16 AssetDebugId = FindOrAddDebugIdForAsset(Frame.StateTree.Get());
-
-				const int32 RequiredSize = StatesIndices.Num() + Frame.ActiveStates.Num();
-				StatesIndices.Reserve(RequiredSize);
-				AssetDebugIds.Reserve(RequiredSize);
-
-				for (const FStateTreeStateHandle StateHandle : Frame.ActiveStates)
-				{
-					StatesIndices.Add(StateHandle.Index);
-					AssetDebugIds.Add(AssetDebugId.Get());
-				}
-			}
-		}
-
+		explicit FActiveStates(FBufferedDataList& BufferedData, const TConstArrayView<FStateTreeExecutionFrame> ActiveFrames);
 		bool IsValid() const { return StatesIndices.Num() > 0 && StatesIndices.Num() == AssetDebugIds.Num(); }
 		void Output(const FStateTreeInstanceDebugId InInstanceId) const
 		{
@@ -253,143 +227,357 @@ struct FInstanceEventBufferedData
 	FString InstanceName;
 	TWeakObjectPtr<const UStateTree> WeakStateTree;
 	FStateTreeInstanceDebugId InstanceId;
+
+	/** Stack to keep track of all received phase events so matching "closing" events (i.e., Pop) will control "IF" and "WHEN" a given phase will be sent. */
+	TArray<FPhaseTraceStatusPair> PhaseStack;
 };
 
 /** Struct to keep track of the buffered event data and flush them. */
 struct FBufferedDataList
 {
-	double TracedRecordingWorldTime = -1;
-	
+	/** Traces required events (world time, phase events, etc.) for a given instance. */
+	void Flush_AnyThread(const FStateTreeInstanceDebugId InstanceId);
+
+	/** Traces buffered events, if needed, when traces are enabled. */
+	void OnTracesStarted_GameThread();
+
+	/** Traces closing events and resets some data for subsequent traces. */
+	void OnStoppingTraces_GameThread();
+
+	/**  Returns an existing ID or create one for a given StateTree asset. */
+	FStateTreeIndex16 FindOrAddDebugIdForAsset_AnyThread(const UStateTree* StateTree);
+
 	/**
-	 * Stacks to keep track of all received phase events so other events will control when and if a given phase trace will be sent.
-	 * This is per thread since it is possible to update execution contexts on multiple threads.
+	 * Keeps track of Pushed/Popped phase for a given instance
+	 * @return The pair that got popped when processing an event of type 'Pop' for a currently active event ('Push').
 	 */
-	TArray<FPhaseStack> PhaseStacks;
+	TOptional<FPhaseTraceStatusPair> UpdatePhaseScope_AnyThread(const FStateTreeInstanceDebugId InstanceId, const EStateTreeUpdatePhase Phase, const EStateTreeTraceEventType EventType, const FStateTreeStateHandle StateHandle);
+
+	/** Converts the provided list of execution frames to traceable data format and outputs it to the trace. */
+	void OutputActiveStates_AnyThread(const FStateTreeInstanceDebugId InstanceId, const TConstArrayView<FStateTreeExecutionFrame> ActiveFrames);
+
+	/** Converts the provided list of execution frames to traceable data format and updates the buffered data for the provided instance. */
+	void UpdateActiveStates_AnyThread(const FStateTreeInstanceDebugId InstanceId, const TConstArrayView<FStateTreeExecutionFrame> ActiveFrames);
+
+	/** Keeps track of created/destroyed statetree instances */
+	void UpdateInstanceLifetime_AnyThread(const FStateTreeInstanceDebugId InstanceId, const UStateTree* StateTree, const TCHAR* InstanceName, const EStateTreeTraceEventType EventType);
+
+	/** Sets the world time that might need to be traced by worker threads. */
+	void SetWorldTime_GameThread(const double WorldTime);
+
+private:
+	static constexpr double UNINITIALIZED_WORLD_TIME = -1;
+
+	/** World time provided by the game thread at the beginning of the frame */
+	double RecordingWorldTime = UNINITIALIZED_WORLD_TIME;
+
+	/** World time traced only once by the first worker thread that needs to trace an event */
+	std::atomic<double> TracedRecordingWorldTime = UNINITIALIZED_WORLD_TIME;
 
 	/** List of asset debug ids events that will be output if channel gets enabled. */
-	TArray<FAssetDebugIdEventBufferedData> AssetDebugIdEvents;
+	TArray<FAssetDebugIdEventBufferedData> AssetDebugIdMap;
+	FRWLock AssetDebugIdMapLock;
 
 	/** List of lifetime events that will be output if channel gets enabled in the Push - Pop lifetime window of an instance. */
-	TMap<FStateTreeInstanceDebugId, FInstanceEventBufferedData> InstanceLifetimeEvents;
+	TMap<FStateTreeInstanceDebugId, FInstanceEventBufferedData> BufferedEventsMap;
+	FRWLock BufferedEventsMapLock;
 
-	/** Flag use to prevent reentrant calls */
-	bool bFlushing = false;
+	/** MT access detector used to validate that AssetDebugIdEvents and InstanceLifetimeEvents are never flushed while getting accessed by worker threads. */
+	UE_MT_DECLARE_RW_ACCESS_DETECTOR(AssetDebugIdMapMTDetector);
+
+	/** MT access detector used to validate that AssetDebugIdEvents and InstanceLifetimeEvents are never flushed while getting accessed by worker threads. */
+	UE_MT_DECLARE_RW_ACCESS_DETECTOR(BufferedEventsMapMTDetector);
+
+	/** Flag use to prevent reentrant calls when buffered events gets flushed when starting a new trace (i.e., OnTracesStarted) */
+	bool bFlushingLifetimeEvents = false;
 
 	uint16 NextAssetDebugId = 1;
 	int32 CurrentVersion = 0;
 	int32 FlushedVersion = -1;
 
-	void Flush(const FStateTreeInstanceDebugId InstanceId)
+
+	/** Called by TraceBufferedEvents from the OutputXYZ methods to make sure current world time was traced. */
+	void TraceWorldTimeIfNeeded_AnyThread();
+
+	/**
+	 * Called by TraceBufferedEvents from the OutputXYZ methods to flush pending phase events.
+	 * Phases popped before TraceStackedPhases gets called will never produce any trace since
+	 * they will not be required for the analysis.
+	 */
+	void TraceStackedPhases_AnyThread(const FStateTreeInstanceDebugId InstanceId);
+};
+
+FInstanceEventBufferedData::FActiveStates::FActiveStates(FBufferedDataList& BufferedData, const TConstArrayView<FStateTreeExecutionFrame> ActiveFrames)
+{
+	for (const FStateTreeExecutionFrame& Frame : ActiveFrames)
 	{
-		if (bFlushing)
+		const FStateTreeIndex16 AssetDebugId = BufferedData.FindOrAddDebugIdForAsset_AnyThread(Frame.StateTree.Get());
+
+		const int32 RequiredSize = StatesIndices.Num() + Frame.ActiveStates.Num();
+		StatesIndices.Reserve(RequiredSize);
+		AssetDebugIds.Reserve(RequiredSize);
+
+		for (const FStateTreeStateHandle StateHandle : Frame.ActiveStates)
 		{
-			return;
+			StatesIndices.Add(StateHandle.Index);
+			AssetDebugIds.Add(AssetDebugId.Get());
 		}
+	}
+}
 
-		TGuardValue<bool> GuardReentry(bFlushing, true);
+void FBufferedDataList::Flush_AnyThread(const FStateTreeInstanceDebugId InstanceId)
+{
+	if (InstanceId.IsValid())
+	{
+		TraceWorldTimeIfNeeded_AnyThread();
+		TraceStackedPhases_AnyThread(InstanceId);
+	}
+}
 
-		const bool bTraceLifetimeEvents = FlushedVersion != CurrentVersion;
-		FlushedVersion = CurrentVersion;
+void FBufferedDataList::TraceWorldTimeIfNeeded_AnyThread()
+{
+	checkf(RecordingWorldTime != UNINITIALIZED_WORLD_TIME
+	       , TEXT("Expecting world time to always be set at the beginning of the frame before any worker thread attempts to trace."));
+	double Expected = UNINITIALIZED_WORLD_TIME;
+	if (TracedRecordingWorldTime.compare_exchange_strong(Expected, RecordingWorldTime))
+	{
+		UE_TRACE_LOG(StateTreeDebugger, WorldTimestampEvent, StateTreeDebugChannel)
+				<< WorldTimestampEvent.WorldTime(RecordingWorldTime);
+	}
+}
+
+void FBufferedDataList::TraceStackedPhases_AnyThread(const FStateTreeInstanceDebugId InstanceId)
+{
+	// Trace pushes phase events and marked them as traced only if not already traced and our channel is enabled.
+	// We need PhaseEvent:Pop to be sent only in this case to enforce complementary events in case of
+	// late recording (e.g. recording gets started, or channel gets enabled, while simulation is already running and instances were ticked)
+	if (!UE_TRACE_CHANNELEXPR_IS_ENABLED(StateTreeDebugChannel))
+	{
+		return;
+	}
 
-		// Trace asset events first since they are required for instance lifetime event types.
-		// Events are preserved in case the trace session is stopped and then a new one gets started
-		// in the same game session. In which case we need to output the ids to that new trace.
-		if (bTraceLifetimeEvents || !IsInGameThread())
+	TArray<FPhaseTraceStatusPair> PhaseStackToTrace;
+	{
+		// We can use a read scope even if we modify entries since the lock
+		// is only required to access the buffered data struct associated 'InstanceId'
+		// which can only happen in a single worker thread each frame.
+		TReadScopeLock ReadSection(BufferedEventsMapLock);
+		UE_MT_SCOPED_READ_ACCESS(BufferedEventsMapMTDetector);
+		if (FInstanceEventBufferedData* Data = BufferedEventsMap.Find(InstanceId))
 		{
-			for (const FAssetDebugIdEventBufferedData& AssetDebugIdEventData : AssetDebugIdEvents)
+			// Keep local copy to trace outside the critical section and only perform the 
+			// minimal required changes while we hold the lock (i.e., updating that status)
+			PhaseStackToTrace = Data->PhaseStack;
+			for (FPhaseTraceStatusPair& StackEntry : Data->PhaseStack)
 			{
-				AssetDebugIdEventData.Trace();
+				StackEntry.bTraced = true;
 			}
 		}
+	}
 
-		TraceWorldTime();
-
-		// Then trace instance lifetime events since they are required for other event types.
-		// They are associated to an older world time but to simplify the logic on the analysis side
-		// we will send them as if the instances were created at the beginning of the recording.
-		// Events are also preserved for the same reason as AssetDebugIdEvents.
-		if (bTraceLifetimeEvents)
+	// We can now safely send PhaseEvents from the local copy
+	for (FPhaseTraceStatusPair& StackEntry : PhaseStackToTrace)
+	{
+		if (StackEntry.bTraced == false)
 		{
-			for (TPair<FStateTreeInstanceDebugId, FInstanceEventBufferedData>& Pair : InstanceLifetimeEvents)
-			{
-				Pair.Value.Trace();
-			}
+			UE_TRACE_LOG(StateTreeDebugger, PhaseEvent, StateTreeDebugChannel)
+					<< PhaseEvent.Cycle(FPlatformTime::Cycles64())
+					<< PhaseEvent.InstanceId(InstanceId.Id)
+					<< PhaseEvent.InstanceSerial(InstanceId.SerialNumber)
+					<< PhaseEvent.Phase(static_cast<std::underlying_type_t<EStateTreeUpdatePhase>>(StackEntry.Phase))
+					<< PhaseEvent.StateIndex(StackEntry.StateHandle.Index)
+					<< PhaseEvent.EventType(static_cast<std::underlying_type_t<EStateTreeTraceEventType>>(EStateTreeTraceEventType::Push));
 		}
+	}
+}
+
+void FBufferedDataList::OnTracesStarted_GameThread()
+{
+	checkf(IsInGameThread(), TEXT("Expecting to only be called by the statetree delegate on the main thread before worker threads trace events."));
 
-		if (InstanceId.IsValid())
+	// Trace asset events first since they are required for instance lifetime event types.
+	// Events are preserved in case the trace session is stopped and then a new one gets started
+	// in the same game session. In which case we need to output AssetDebugIdEvents to that new trace.
+	const bool bFlushingLifetimeEventsRequired = FlushedVersion != CurrentVersion;
+	TGuardValue<bool> GuardReentry(bFlushingLifetimeEvents, bFlushingLifetimeEventsRequired);
+
+	if (bFlushingLifetimeEventsRequired)
+	{
+		FlushedVersion = CurrentVersion;
+
+		UE_MT_SCOPED_READ_ACCESS(BufferedEventsMapMTDetector);
+		for (const FAssetDebugIdEventBufferedData& AssetDebugIdEventData : AssetDebugIdMap)
 		{
-			TraceStackedPhases(InstanceId);
+			AssetDebugIdEventData.Trace();
 		}
 	}
 
-	/**
-	 * Called by TraceBufferedEvents from the OutputXYZ methods to make sure we have the current world time was sent.
-	 */
-	void TraceWorldTime()
+	TraceWorldTimeIfNeeded_AnyThread();
+
+	// Then trace instance lifetime events since they are required for other event types.
+	// They are associated to an older world time but to simplify the logic on the analysis side
+	// we will send them as if the instances were created at the beginning of the recording.
+	// Those events are also preserved for the same reason as AssetDebugIdEvents.
+	if (bFlushingLifetimeEventsRequired)
 	{
-		if (TracedRecordingWorldTime != GWorldTime)
+		UE_MT_SCOPED_READ_ACCESS(BufferedEventsMapMTDetector);
+		for (TPair<FStateTreeInstanceDebugId, FInstanceEventBufferedData>& Pair : BufferedEventsMap)
 		{
-			TracedRecordingWorldTime = GWorldTime;
-			UE_TRACE_LOG(StateTreeDebugger, WorldTimestampEvent, StateTreeDebugChannel)
-				<< WorldTimestampEvent.WorldTime(GWorldTime);
+			Pair.Value.Trace();
 		}
 	}
+}
 
-	/**
-	 * Called by TraceBufferedEvents from the OutputXYZ methods to flush pending phase events.
-	 * Phases popped before TraceStackedPhases gets called will never produce any trace since
-	 * they will not be required for the analysis.
-	 */
-	void TraceStackedPhases(const FStateTreeInstanceDebugId InstanceId)
+void FBufferedDataList::OnStoppingTraces_GameThread()
+{
+	checkf(IsInGameThread(), TEXT("Expecting to only be called by the statetree delegate on the main thread after worker threads traced events."));
+
+	UE_MT_SCOPED_READ_ACCESS(BufferedEventsMapMTDetector);
+	for (TPair<FStateTreeInstanceDebugId, FInstanceEventBufferedData>& Pair : BufferedEventsMap)
+	{
+		Pair.Value.CloseRecording(TracedRecordingWorldTime);
+	}
+
+	// Bump version so shareable data will be flush in the next trace (e.g. Asset ids, instance lifetime events, etc.)
+	++CurrentVersion;
+
+	// Reset for next trace
+	RecordingWorldTime = UNINITIALIZED_WORLD_TIME;
+}
+
+FStateTreeIndex16 FBufferedDataList::FindOrAddDebugIdForAsset_AnyThread(const UStateTree* StateTree)
+{
+	FStateTreeIndex16 AssetDebugId;
+	if (!ensure(StateTree != nullptr))
+	{
+		return AssetDebugId;
+	}
+
+	// Return DebugId from existing entry (read-only operation)
 	{
-		for (FPhaseStack& PhaseStack : PhaseStacks)
+		TReadScopeLock ReadSection(AssetDebugIdMapLock);
+		UE_MT_SCOPED_READ_ACCESS(AssetDebugIdMapMTDetector);
+		const FAssetDebugIdEventBufferedData* ExistingPair = AssetDebugIdMap.FindByPredicate([StateTree](const FAssetDebugIdEventBufferedData& BufferedData)
 		{
-			if (PhaseStack.InstanceId == InstanceId)
-			{
-				for (FPhaseTraceStatusPair& StackEntry : PhaseStack.Stack)
-				{
-					// Trace push phase event and marked as traced only if not already traced and our channel is enabled.
-					// We need the pop phase event to be sent only in this case to enforce complementary events in case of
-					// late recording (e.g. recording started, or channel enabled, after simulation is running and instances are ticked)
-					if (StackEntry.bTraced == false && UE_TRACE_CHANNELEXPR_IS_ENABLED(StateTreeDebugChannel))
-					{
-						UE_TRACE_LOG(StateTreeDebugger, PhaseEvent, StateTreeDebugChannel)
-							<< PhaseEvent.Cycle(FPlatformTime::Cycles64())
-							<< PhaseEvent.InstanceId(InstanceId.Id)
-							<< PhaseEvent.InstanceSerial(InstanceId.SerialNumber)
-							<< PhaseEvent.Phase(static_cast<std::underlying_type_t<EStateTreeUpdatePhase>>(StackEntry.Phase))
-							<< PhaseEvent.StateIndex(StackEntry.StateHandle.Index)
-							<< PhaseEvent.EventType(static_cast<std::underlying_type_t<EStateTreeTraceEventType>>(EStateTreeTraceEventType::Push));
-
-						StackEntry.bTraced = true;
-					}
-				}
-				break;
-			}
+			return BufferedData.WeakStateTree == StateTree;
+		});
+
+		if (ExistingPair != nullptr)
+		{
+			return ExistingPair->AssetDebugId;
 		}
 	}
 
-	void OnStoppingTraces()
+	// Assign new DebugId, store it and trace it (write operation)
+	{
+		TWriteScopeLock WriteSection(AssetDebugIdMapLock);
+		UE_MT_SCOPED_WRITE_ACCESS(AssetDebugIdMapMTDetector);
+		AssetDebugId = FStateTreeIndex16(++NextAssetDebugId);
+		AssetDebugIdMap.Emplace(StateTree, AssetDebugId);
+	}
+
+	OutputAssetDebugIdEvent(StateTree, AssetDebugId);
+
+	return AssetDebugId;
+}
+
+TOptional<FPhaseTraceStatusPair> FBufferedDataList::UpdatePhaseScope_AnyThread(const FStateTreeInstanceDebugId InstanceId, const EStateTreeUpdatePhase Phase
+	, const EStateTreeTraceEventType EventType, const FStateTreeStateHandle StateHandle)
+{
+	TOptional<FPhaseTraceStatusPair> RemovedPair;
+	auto UpdateStackFunc = [&RemovedPair, Phase, EventType, StateHandle](FInstanceEventBufferedData& Data)
 	{
-		for (TPair<FStateTreeInstanceDebugId, FInstanceEventBufferedData>& Pair : InstanceLifetimeEvents)
+		if (EventType == EStateTreeTraceEventType::Push)
 		{
-			Pair.Value.CloseRecording(TracedRecordingWorldTime);
+			Data.PhaseStack.Push(FPhaseTraceStatusPair(Phase, StateHandle));
 		}
+		else if (ensureAlwaysMsgf(Data.PhaseStack.IsEmpty() == false, TEXT("Not expected to pop phases that never got pushed.")))
+		{
+			RemovedPair = Data.PhaseStack.Pop();
+		}
+	};
 
-		// Bump version so shareable data will be flush in the next trace (e.g. Asset ids, instance lifetime events, etc.)
-		CurrentVersion++;
+	// Update existing data (read-only operation on the event container)
+	{
+		TReadScopeLock ReadSection(BufferedEventsMapLock);
+		UE_MT_SCOPED_READ_ACCESS(BufferedEventsMapMTDetector);
+		FInstanceEventBufferedData* ExistingBufferedData = BufferedEventsMap.Find(InstanceId);
+		if (ExistingBufferedData != nullptr)
+		{
+			UpdateStackFunc(*ExistingBufferedData);
+			return RemovedPair;
+		}
+	}
 
-		// Force world time to trace on the next event
-		GWorldTime = -1;
+	// Need to create a new entry in the container
+	{
+		TWriteScopeLock WriteSection(BufferedEventsMapLock);
+		UE_MT_SCOPED_WRITE_ACCESS(BufferedEventsMapMTDetector);
+		FInstanceEventBufferedData& NewBufferedData = BufferedEventsMap.Add(InstanceId);
+		UpdateStackFunc(NewBufferedData);
 	}
-};
+	return RemovedPair;
+}
+
+void FBufferedDataList::OutputActiveStates_AnyThread(const FStateTreeInstanceDebugId InstanceId, const TConstArrayView<FStateTreeExecutionFrame> ActiveFrames)
+{
+	const FInstanceEventBufferedData::FActiveStates ActiveStates = FInstanceEventBufferedData::FActiveStates(*this, ActiveFrames);
+	ActiveStates.Output(InstanceId);
+}
+
+void FBufferedDataList::UpdateActiveStates_AnyThread(const FStateTreeInstanceDebugId InstanceId, const TConstArrayView<FStateTreeExecutionFrame> ActiveFrames)
+{
+	TReadScopeLock ReadSection(BufferedEventsMapLock);
+	UE_MT_SCOPED_READ_ACCESS(BufferedEventsMapMTDetector);
+
+	// We keep only the most recent active states since all we need to know is the last active states of the instance
+	// when we start receiving events once the channel gets enabled.
+	if (FInstanceEventBufferedData* ExistingBufferedData = BufferedEventsMap.Find(InstanceId))
+	{
+		ExistingBufferedData->ActiveStates = FInstanceEventBufferedData::FActiveStates(*this, ActiveFrames);
+	}
+}
+
+void FBufferedDataList::UpdateInstanceLifetime_AnyThread(const FStateTreeInstanceDebugId InstanceId, const UStateTree* StateTree, const TCHAR* InstanceName
+	, const EStateTreeTraceEventType EventType)
+{
+	// This can be called by worker threads from statetree trace macros and by the game thread during the flush.
+	// For the latter we don't want to queue again events that are getting flushed.
+	if (!bFlushingLifetimeEvents)
+	{
+		TWriteScopeLock WriteSection(BufferedEventsMapLock);
+		UE_MT_SCOPED_WRITE_ACCESS(BufferedEventsMapMTDetector);
+
+		if (EventType == EStateTreeTraceEventType::Push)
+		{
+			BufferedEventsMap.Emplace(InstanceId, FInstanceEventBufferedData(StateTree, InstanceId, InstanceName));
+		}
+		else if (EventType == EStateTreeTraceEventType::Pop)
+		{
+			BufferedEventsMap.Remove(InstanceId);
+		}
+		else
+		{
+			ensureMsgf(false, TEXT("Unexpected EventType '%s' for instance lifetime event."), *UEnum::GetDisplayValueAsText(EventType).ToString());
+		}
+	}
+}
+
+void FBufferedDataList::SetWorldTime_GameThread(const double WorldTime)
+{
+	checkf(IsInGameThread(), TEXT("Expecting to only be called by the world delegate on the main thread before worker threads trace events."));
+	RecordingWorldTime = WorldTime;
+
+	// Reset traced time so the first worker thread that will need it will set and trace it
+	TracedRecordingWorldTime = UNINITIALIZED_WORLD_TIME;
+}
 
 /**
  * Buffered events (e.g. lifetime, active state, scoped phase) in case channel is not active yet or phase are empty and don't need to be traced.
- * This is per thread since it is possible to update execution contexts on multiple threads.
- * @note The current implementation of the lifetime events doesn't properly support same instance getting ticked in different threads.
+ * This can be accessed from any thread and based on the assumption that a given instance data can never be accessed by multiple thread simultaneously in the same frame
+ * so we only need to protect access to the main maps.
  */
-thread_local FBufferedDataList GBufferedEvents;
+FBufferedDataList GBufferedEvents;
 
 /**
  * Pushed or pops an entry on the Phase stack for a given Instance.
@@ -397,57 +585,29 @@ thread_local FBufferedDataList GBufferedEvents;
  */
 void OutputPhaseScopeEvent(const FStateTreeInstanceDebugId InstanceId, const EStateTreeUpdatePhase Phase, const EStateTreeTraceEventType EventType, const FStateTreeStateHandle StateHandle)
 {
-	TArray<FPhaseStack>& PhaseStacks = GBufferedEvents.PhaseStacks;
-	int32 ExistingStackIndex = PhaseStacks.IndexOfByPredicate([InstanceId](const FPhaseStack& PhaseStack){ return PhaseStack.InstanceId == InstanceId; });
+	const TOptional<FPhaseTraceStatusPair> RemovedPair = GBufferedEvents.UpdatePhaseScope_AnyThread(InstanceId, Phase, EventType, StateHandle);
 
-	if (EventType == EStateTreeTraceEventType::Push)
-	{
-		if (ExistingStackIndex == INDEX_NONE)
-		{
-			ExistingStackIndex = PhaseStacks.AddDefaulted();
-		}
-		FPhaseStack& PhaseStack = PhaseStacks[ExistingStackIndex];
-		PhaseStack.InstanceId = InstanceId;
-		PhaseStack.Stack.Push(FPhaseTraceStatusPair(Phase, StateHandle));
-	}
-	else if (ensureMsgf(ExistingStackIndex != INDEX_NONE, TEXT("Not expected to pop phases for an instance that never pushed a phase.")))
+	// Phase was previously traced (i.e., other events were traced in that scope so we need to trace the closing (i.e., Pop) event)
+	if (RemovedPair.IsSet()
+		&& ensureMsgf(RemovedPair.GetValue().Phase == Phase, TEXT("Not expected to pop a phase that is not on the top of the stack."))
+		&& RemovedPair.GetValue().bTraced)
 	{
-		FPhaseStack& PhaseStack = PhaseStacks[ExistingStackIndex];
-
-		if (ensureMsgf(PhaseStack.Stack.IsEmpty() == false, TEXT("Not expected to pop phases that never got pushed.")) &&
-			ensureMsgf(PhaseStack.InstanceId == InstanceId, TEXT("Not expected to pop phases for an instance that is not the one currently assigned to the stack.")))
-		{
-			const FPhaseTraceStatusPair RemovedPair = PhaseStack.Stack.Pop();
-			ensureMsgf(RemovedPair.Phase == Phase, TEXT("Not expected to pop a phase that is not on the top of the stack."));
-
-			// Clear associated InstanceId when removing last entry from the stack.
-			if (PhaseStack.Stack.IsEmpty())
-			{
-				PhaseStacks.RemoveAt(ExistingStackIndex, /*Count*/1, EAllowShrinking::No);
-			}
-
-			// Phase was previously traced (i.e. other events were traced in that scope so we need to trace the closing (i.e. Pop) event.
-			if (RemovedPair.bTraced)
-			{
-				UE_TRACE_LOG(StateTreeDebugger, PhaseEvent, StateTreeDebugChannel)
-				<< PhaseEvent.Cycle(FPlatformTime::Cycles64())
-				<< PhaseEvent.InstanceId(InstanceId.Id)
-				<< PhaseEvent.InstanceSerial(InstanceId.SerialNumber)
-				<< PhaseEvent.Phase(static_cast<std::underlying_type_t<EStateTreeUpdatePhase>>(Phase))
-				<< PhaseEvent.StateIndex(StateHandle.Index)
-				<< PhaseEvent.EventType(static_cast<std::underlying_type_t<EStateTreeTraceEventType>>(EStateTreeTraceEventType::Pop));
-			}
-		}
+		UE_TRACE_LOG(StateTreeDebugger, PhaseEvent, StateTreeDebugChannel)
+			<< PhaseEvent.Cycle(FPlatformTime::Cycles64())
+			<< PhaseEvent.InstanceId(InstanceId.Id)
+			<< PhaseEvent.InstanceSerial(InstanceId.SerialNumber)
+			<< PhaseEvent.Phase(static_cast<std::underlying_type_t<EStateTreeUpdatePhase>>(Phase))
+			<< PhaseEvent.StateIndex(StateHandle.Index)
+			<< PhaseEvent.EventType(static_cast<std::underlying_type_t<EStateTreeTraceEventType>>(EStateTreeTraceEventType::Pop));
 	}
 }
 
-
 /**
  * Called by the OutputXYZ methods to flush pending events (e.g. Push or WorldTime).
  */
 void TraceBufferedEvents(const FStateTreeInstanceDebugId InstanceId)
 {
-	GBufferedEvents.Flush(InstanceId);
+	GBufferedEvents.Flush_AnyThread(InstanceId);
 }
 
 void SerializeDataViewToArchive(FBufferArchive& Ar, const FStateTreeDataView DataView)
@@ -492,16 +652,21 @@ void RegisterGlobalDelegates()
 	GOnWorldTickStartDelegateHandle = FWorldDelegates::OnWorldTickStart.AddLambda([](const UWorld* TickedWorld, ELevelTick TickType, float DeltaTime)
 		{
 #if OBJECT_TRACE_ENABLED
-			GWorldTime = FObjectTrace::GetWorldElapsedTime(TickedWorld);
+			GBufferedEvents.SetWorldTime_GameThread(FObjectTrace::GetWorldElapsedTime(TickedWorld));
 #endif// OBJECT_TRACE_ENABLED
 		});
 
 	GTracingStateChangedDelegateHandle = UE::StateTree::Delegates::OnTracingStateChanged.AddLambda([](const EStateTreeTraceStatus TraceStatus)
 	{
+		if (TraceStatus == EStateTreeTraceStatus::TracesStarted)
+		{
+			GBufferedEvents.OnTracesStarted_GameThread();
+		}
+
 		// Traces are about to be stopped so allow the buffered events to react.
 		if (TraceStatus == EStateTreeTraceStatus::StoppingTrace)
 		{
-			GBufferedEvents.OnStoppingTraces();
+			GBufferedEvents.OnStoppingTraces_GameThread();
 		}
 	});
 }
@@ -515,29 +680,10 @@ void UnregisterGlobalDelegates()
 	GTracingStateChangedDelegateHandle.Reset();
 }
 
+// Deprecated
 FStateTreeIndex16 FindOrAddDebugIdForAsset(const UStateTree* StateTree)
 {
-	FStateTreeIndex16 AssetDebugId;
-	const FAssetDebugIdEventBufferedData* ExistingPair = GBufferedEvents.AssetDebugIdEvents.FindByPredicate([StateTree](const FAssetDebugIdEventBufferedData& BufferedData)
-	{
-		return BufferedData.WeakStateTree == StateTree;
-	});
-
-	if (ExistingPair == nullptr)
-	{
-		if (ensure(StateTree != nullptr))
-		{
-			AssetDebugId = FStateTreeIndex16(GBufferedEvents.NextAssetDebugId++);
-			GBufferedEvents.AssetDebugIdEvents.Emplace(StateTree, AssetDebugId);
-			OutputAssetDebugIdEvent(StateTree, AssetDebugId);
-		}
-	}
-	else
-	{
-		AssetDebugId = ExistingPair->AssetDebugId;
-	}
-
-	return AssetDebugId;
+	return GBufferedEvents.FindOrAddDebugIdForAsset_AnyThread(StateTree);
 }
 
 void OutputAssetDebugIdEvent(
@@ -573,7 +719,7 @@ void OutputInstanceLifetimeEvent(
 	{
 		TraceBufferedEvents(InstanceId);
 
-		const FStateTreeIndex16 AssetDebugId = FindOrAddDebugIdForAsset(StateTree);
+		const FStateTreeIndex16 AssetDebugId = GBufferedEvents.FindOrAddDebugIdForAsset_AnyThread(StateTree);
 
 		UE_TRACE_LOG(StateTreeDebugger, InstanceEvent, StateTreeDebugChannel)
 			<< InstanceEvent.Cycle(FPlatformTime::Cycles64())
@@ -584,24 +730,10 @@ void OutputInstanceLifetimeEvent(
 			<< InstanceEvent.AssetDebugId(AssetDebugId.Get());
 	}
 
-	// Buffer these events regardless of the status of the channel since they will be used
+	// Update buffered events regardless of the status of the channel since they will be used
 	// when flushing buffered event when a late recording is started or more than one trace are started
 	// during the same game session (i.e. Start Traces -> Stop Traces -> Start Traces).
-	if (!GBufferedEvents.bFlushing)
-	{
-		if (EventType == EStateTreeTraceEventType::Push)
-		{
-			GBufferedEvents.InstanceLifetimeEvents.Emplace(InstanceId, FInstanceEventBufferedData(StateTree, InstanceId, InstanceName));
-		}
-		else if (EventType == EStateTreeTraceEventType::Pop)
-		{
-			GBufferedEvents.InstanceLifetimeEvents.Remove(InstanceId);
-		}
-		else
-		{
-			ensureMsgf(false, TEXT("Unexpected EventType '%s' for instance lifetime event."), *UEnum::GetDisplayValueAsText(EventType).ToString());
-		}
-	}
+	GBufferedEvents.UpdateInstanceLifetime_AnyThread(InstanceId, StateTree, InstanceName, EventType);
 }
 
 void OutputInstanceAssetEvent(
@@ -615,7 +747,7 @@ void OutputInstanceAssetEvent(
 	{
 		TraceBufferedEvents(InstanceId);
 
-		const FStateTreeIndex16 AssetDebugId = FindOrAddDebugIdForAsset(StateTree);
+		const FStateTreeIndex16 AssetDebugId = GBufferedEvents.FindOrAddDebugIdForAsset_AnyThread(StateTree);
 
 		UE_TRACE_LOG(StateTreeDebugger, InstanceFrameEvent, StateTreeDebugChannel)
 			<< InstanceFrameEvent.Cycle(FPlatformTime::Cycles64())
@@ -623,8 +755,10 @@ void OutputInstanceAssetEvent(
 			<< InstanceFrameEvent.InstanceSerial(InstanceId.SerialNumber)
 			<< InstanceFrameEvent.AssetDebugId(AssetDebugId.Get());
 	}
-	// No need to buffer since frame event are sent each time a FrameScope is used by the execution context
-	// and we don't expect the trace channel to be enabled/disabled during a single execution context update.
+
+	// No need to buffer anything here since these events are sent each time a FCurrentlyProcessedFrameScope
+	// is used by a FStateTreeExecutionContext, and we don't expect the trace channel to be enabled/disabled
+	// during a single execution context update.
 }
 
 void OutputInstanceFrameEvent(
@@ -773,17 +907,11 @@ void OutputActiveStatesEventTrace(
 	{
 		TraceBufferedEvents(InstanceId);
 
-		const FInstanceEventBufferedData::FActiveStates ActiveStates(ActiveFrames);
-		ActiveStates.Output(InstanceId);
+		GBufferedEvents.OutputActiveStates_AnyThread(InstanceId, ActiveFrames);
 	}
 	else
 	{
-		// We keep only the most recent active states since this is all we need to know in which state was the instance
-		// when we start receiving the events once the channel is enabled.
-		if (FInstanceEventBufferedData* ExisingBufferedData = GBufferedEvents.InstanceLifetimeEvents.Find(InstanceId))
-		{
-			ExisingBufferedData->ActiveStates= FInstanceEventBufferedData::FActiveStates(ActiveFrames);
-		}
+		GBufferedEvents.UpdateActiveStates_AnyThread(InstanceId, ActiveFrames);
 	}
 }
 
diff --git a/Engine/Plugins/Runtime/StateTree/Source/StateTreeModule/Private/StateTreeExecutionContext.cpp b/Engine/Plugins/Runtime/StateTree/Source/StateTreeModule/Private/StateTreeExecutionContext.cpp
--- a/Engine/Plugins/Runtime/StateTree/Source/StateTreeModule/Private/StateTreeExecutionContext.cpp
+++ b/Engine/Plugins/Runtime/StateTree/Source/StateTreeModule/Private/StateTreeExecutionContext.cpp
@@ -832,10 +832,6 @@ EStateTreeRunStatus FStateTreeExecutionContext::Stop(EStateTreeRunStatus Complet
 
 	TGuardValue<bool> ScheduledNextTickScope(bAllowedToScheduleNextTick, false);
 
-	// Set scoped phase only for properly initialized context with valid Instance data
-	// since we need it to output the InstanceId 
-	UE_STATETREE_DEBUG_SCOPED_PHASE(this, EStateTreeUpdatePhase::StopTree);
-
 	// Make sure that we return a valid completion status (i.e. Succeeded, Failed or Stopped)
 	if (CompletionStatus == EStateTreeRunStatus::Unset
 		|| CompletionStatus == EStateTreeRunStatus::Running)
@@ -855,6 +851,10 @@ EStateTreeRunStatus FStateTreeExecutionContext::Stop(EStateTreeRunStatus Complet
 		return EStateTreeRunStatus::Running;
 	}
 
+	// Set phase only after validating that we are in a valid context since we need an InstancedId for the debug
+	// We intentionally use ENTER_PHASE instead of SCOPED_PHASE here to make sure we exit the phase before the Instance
+	UE_STATETREE_DEBUG_ENTER_PHASE(this, EStateTreeUpdatePhase::StopTree);
+
 	// No need to clear on exit since we reset all the instance data before leaving the function.
 	Exec.CurrentPhase = EStateTreeUpdatePhase::StopTree;
 
@@ -877,6 +877,7 @@ EStateTreeRunStatus FStateTreeExecutionContext::Stop(EStateTreeRunStatus Complet
 
 	// Trace before resetting the instance data since it is required to provide all the event information
 	UE_STATETREE_DEBUG_ACTIVE_STATES_EVENT(this, {});
+	UE_STATETREE_DEBUG_EXIT_PHASE(this, EStateTreeUpdatePhase::StopTree);
 	UE_STATETREE_DEBUG_INSTANCE_EVENT(this, EStateTreeTraceEventType::Pop);
 
 	// Destruct all allocated instance data (does not shrink the buffer). This will invalidate Exec too.
diff --git a/Engine/Plugins/Runtime/StateTree/Source/StateTreeModule/Public/Debugger/StateTreeTrace.h b/Engine/Plugins/Runtime/StateTree/Source/StateTreeModule/Public/Debugger/StateTreeTrace.h
--- a/Engine/Plugins/Runtime/StateTree/Source/StateTreeModule/Public/Debugger/StateTreeTrace.h
+++ b/Engine/Plugins/Runtime/StateTree/Source/StateTreeModule/Public/Debugger/StateTreeTrace.h
@@ -27,6 +27,7 @@ namespace UE::StateTreeTrace
 {
 	void RegisterGlobalDelegates();
 	void UnregisterGlobalDelegates();
+	UE_DEPRECATED(5.6, "This method will no longer be exposed publicly.")
 	FStateTreeIndex16 FindOrAddDebugIdForAsset(const UStateTree* StateTree);
 	void OutputPhaseScopeEvent(FStateTreeInstanceDebugId InstanceId, EStateTreeUpdatePhase Phase, EStateTreeTraceEventType EventType, FStateTreeStateHandle StateHandle);
 	void OutputAssetDebugIdEvent(const UStateTree* StateTree, FStateTreeIndex16 AssetDebugId);
