StateTree: Fixed uninitialized variable in CompactStateTreeView.
#pr

#rnx

diff --git a/Engine/Plugins/Runtime/StateTree/Source/StateTreeEditorModule/Private/Customizations/Widgets/SCompactStateTreeView.cpp b/Engine/Plugins/Runtime/StateTree/Source/StateTreeEditorModule/Private/Customizations/Widgets/SCompactStateTreeView.cpp
--- a/Engine/Plugins/Runtime/StateTree/Source/StateTreeEditorModule/Private/Customizations/Widgets/SCompactStateTreeView.cpp
+++ b/Engine/Plugins/Runtime/StateTree/Source/StateTreeEditorModule/Private/Customizations/Widgets/SCompactStateTreeView.cpp
@@ -117,7 +117,7 @@ void SCompactStateTreeView::SetSelection(TConstArrayView<FGuid> Selection)
 	StateItemTree->ClearSelection();
 	StateItemTree->SetItemSelection(SelectedStates, true);
 
-	if (bShowLinkedStates)
+	if (bShowLinkedStates && StateTreeViewModel)
 	{
 		ResetLinkedStates();
 
diff --git a/Engine/Plugins/Runtime/StateTree/Source/StateTreeEditorModule/Private/Customizations/Widgets/SCompactStateTreeView.h b/Engine/Plugins/Runtime/StateTree/Source/StateTreeEditorModule/Private/Customizations/Widgets/SCompactStateTreeView.h
--- a/Engine/Plugins/Runtime/StateTree/Source/StateTreeEditorModule/Private/Customizations/Widgets/SCompactStateTreeView.h
+++ b/Engine/Plugins/Runtime/StateTree/Source/StateTreeEditorModule/Private/Customizations/Widgets/SCompactStateTreeView.h
@@ -29,6 +29,7 @@ public:
 		, _SelectionMode(ESelectionMode::Single)
 		, _SelectableStatesOnly(false)
 		, _SubtreesOnly(false)
+		, _ShowLinkedStates(false)
 	{}
 		/** Currently selected struct, initially highlighted. */
 		SLATE_ARGUMENT(const UStateTreeEditorData*, StateTreeEditorData)
