[PropertyBindings] fixed memory leak in FPropertyBindingBindingCollection::ResetObjects when the binding is a complex type with a destructor since InitializeValue() does not check for CPF_NoDestructor. The issue was seen with a TArray leaking its buffer when it was reset.
PR #12810: Fix for a memory leak and crash in StateTree.

diff --git a/Engine/Plugins/Runtime/PropertyBindingUtils/Source/PropertyBindingUtils/Private/PropertyBindingBindingCollection.cpp b/Engine/Plugins/Runtime/PropertyBindingUtils/Source/PropertyBindingUtils/Private/PropertyBindingBindingCollection.cpp
--- a/Engine/Plugins/Runtime/PropertyBindingUtils/Source/PropertyBindingUtils/Private/PropertyBindingBindingCollection.cpp
+++ b/Engine/Plugins/Runtime/PropertyBindingUtils/Source/PropertyBindingUtils/Private/PropertyBindingBindingCollection.cpp
@@ -914,7 +914,7 @@ bool FPropertyBindingBindingCollection::PerformCopy(const FPropertyBindingCopyIn
 		static_cast<const FBoolProperty*>(Copy.TargetLeafProperty)->SetPropertyValue(TargetAddress, static_cast<const FBoolProperty*>(Copy.SourceLeafProperty)->GetPropertyValue(SourceAddress));
 		break;
 	case EPropertyCopyType::CopyStruct:
-		// If SourceProperty == nullptr (pointing to the struct source directly), the GetAddress() did the right thing and is pointing the the beginning of the struct. 
+		// If SourceProperty == nullptr (pointing to the struct source directly), the GetAddress() did the right thing and is pointing to the beginning of the struct. 
 		static_cast<const FStructProperty*>(Copy.TargetLeafProperty)->Struct->CopyScriptStruct(TargetAddress, SourceAddress);
 		break;
 	case EPropertyCopyType::CopyObject:
@@ -1089,7 +1089,7 @@ void FPropertyBindingBindingCollection::PerformResetObjects(const FPropertyBindi
 	switch (Copy.Type)
 	{
 	case EPropertyCopyType::CopyComplex:
-		Copy.TargetLeafProperty->InitializeValue(TargetAddress);
+		Copy.TargetLeafProperty->ClearValue(TargetAddress);
 		break;
 	case EPropertyCopyType::CopyStruct:
 		static_cast<const FStructProperty*>(Copy.TargetLeafProperty)->Struct->ClearScriptStruct(TargetAddress);
@@ -1110,7 +1110,7 @@ void FPropertyBindingBindingCollection::PerformResetObjects(const FPropertyBindi
 		FScriptArrayHelper TargetArrayHelper(TargetArrayProperty, TargetAddress);
 		for (int32 ElementIndex = 0; ElementIndex < TargetArrayHelper.Num(); ++ElementIndex)
 		{
-			TargetArrayProperty->Inner->InitializeValue(TargetArrayHelper.GetRawPtr(ElementIndex));
+			TargetArrayProperty->Inner->ClearValue(TargetArrayHelper.GetRawPtr(ElementIndex));
 		}
 		break;
 	}
@@ -1137,7 +1137,7 @@ bool FPropertyBindingBindingCollection::ResetObjects(const FPropertyBindingIndex
 
 	check(TargetStructView.IsValid());
 	check(TargetStructView.GetStruct() == Batch.TargetStruct.Get().Struct);
-	
+
 	for (int32 i = Batch.BindingsBegin.Get(); i != Batch.BindingsEnd.Get(); i++)
 	{
 		const FPropertyBindingCopyInfo& Copy = PropertyCopies[i];
@@ -1146,10 +1146,13 @@ bool FPropertyBindingBindingCollection::ResetObjects(const FPropertyBindingIndex
 		{
 			continue;
 		}
-		
-		uint8* TargetAddress = GetAddress(TargetStructView, Copy.TargetIndirection, Copy.TargetLeafProperty);
-		check(TargetAddress != nullptr);
-		PerformResetObjects(Copy, TargetAddress);
+
+		// Validate target address since resetting a previous bindings might invalidate some subsequent bindings
+		// targeting inner values (e.g., array got reset and bindings to item no longer need to be reset)
+		if (uint8* TargetAddress = GetAddress(TargetStructView, Copy.TargetIndirection, Copy.TargetLeafProperty))
+		{
+			PerformResetObjects(Copy, TargetAddress);
+		}
 	}
 
 	return true;
diff --git a/Engine/Plugins/Runtime/StateTree/Source/StateTreeTestSuite/Private/StateTreeBindingTest.cpp b/Engine/Plugins/Runtime/StateTree/Source/StateTreeTestSuite/Private/StateTreeBindingTest.cpp
--- a/Engine/Plugins/Runtime/StateTree/Source/StateTreeTestSuite/Private/StateTreeBindingTest.cpp
+++ b/Engine/Plugins/Runtime/StateTree/Source/StateTreeTestSuite/Private/StateTreeBindingTest.cpp
@@ -55,6 +55,14 @@ struct FStateTreeTest_BindingsCompiler : FStateTreeTestBase
 		PropertyBindings.Add(MakeBinding(SourceADesc.ID, TEXT("Item.B"), TargetDesc.ID, TEXT("Array[1].B")));
 		PropertyBindings.Add(MakeBinding(SourceADesc.ID, TEXT("Array"), TargetDesc.ID, TEXT("Array")));
 
+		PropertyBindings.Add(MakeBinding(SourceBDesc.ID, TEXT("Item"), TargetDesc.ID, TEXT("FixedArray[1]")));
+		PropertyBindings.Add(MakeBinding(SourceADesc.ID, TEXT("Item.B"), TargetDesc.ID, TEXT("FixedArray[1].B")));
+		PropertyBindings.Add(MakeBinding(SourceADesc.ID, TEXT("FixedArray"), TargetDesc.ID, TEXT("FixedArray")));
+
+		PropertyBindings.Add(MakeBinding(SourceBDesc.ID, TEXT("Item"), TargetDesc.ID, TEXT("CArray[1]")));
+		PropertyBindings.Add(MakeBinding(SourceADesc.ID, TEXT("Item.B"), TargetDesc.ID, TEXT("CArray[1].B")));
+		PropertyBindings.Add(MakeBinding(SourceADesc.ID, TEXT("CArray"), TargetDesc.ID, TEXT("CArray")));
+
 		int32 CopyBatchIndex = INDEX_NONE;
 		const bool bCompileBatchResult = BindingCompiler.CompileBatch(TargetDesc, PropertyBindings, FStateTreeIndex16::Invalid, FStateTreeIndex16::Invalid, CopyBatchIndex);
 		AITEST_TRUE("CompileBatch should succeed", bCompileBatchResult);
@@ -70,11 +78,21 @@ struct FStateTreeTest_BindingsCompiler : FStateTreeTestBase
 		SourceA.Array.AddDefaulted_GetRef().A = 1;
 		SourceA.Array.AddDefaulted_GetRef().B = 2;
 
+		constexpr int32 FixedArraySize = 4;
+		SourceA.FixedArray.SetNum(FixedArraySize, EAllowShrinking::No);
+		SourceA.FixedArray[0].A = 1;
+		SourceA.FixedArray[1].B = 2;
+
+		SourceA.CArray[0].A = 1;
+		SourceA.CArray[0].B = 2;
+
 		FStateTreeTest_PropertyCopy SourceB;
 		SourceB.Item.A = 41;
 		SourceB.Item.B = 42;
+		SourceB.FixedArray.SetNum(FixedArraySize, EAllowShrinking::No);
 
 		FStateTreeTest_PropertyCopy Target;
+		Target.FixedArray.SetNum(FixedArraySize, EAllowShrinking::No);
 
 		AITEST_TRUE("SourceAIndex should be less than max number of source structs.", SourceAIndex < Bindings.GetNumBindableStructDescriptors());
 		AITEST_TRUE("SourceBIndex should be less than max number of source structs.", SourceBIndex < Bindings.GetNumBindableStructDescriptors());
@@ -83,7 +101,7 @@ struct FStateTreeTest_BindingsCompiler : FStateTreeTestBase
 		SourceViews.SetNum(Bindings.GetNumBindableStructDescriptors());
 		SourceViews[SourceAIndex] = FStateTreeDataView(FStructView::Make(SourceA));
 		SourceViews[SourceBIndex] = FStateTreeDataView(FStructView::Make(SourceB));
-		FStateTreeDataView TargetView(FStructView::Make(Target));
+		FPropertyBindingDataView TargetView(FStructView::Make(Target));
 
 		bool bCopyResult = true;
 		for (const FPropertyBindingCopyInfo& Copy : Bindings.Super::GetBatchCopies(FPropertyBindingIndex16(CopyBatchIndex)))
@@ -93,6 +111,12 @@ struct FStateTreeTest_BindingsCompiler : FStateTreeTestBase
 		AITEST_TRUE("CopyTo should succeed", bCopyResult);
 
 		// Due to binding sorting, we expect them to execute in this order (sorted based on target access, earliest to latest)
+		// SourceA.CArray -> Target.CArray
+		// SourceB.Item -> Target.CArray[1]
+		// SourceA.Item.B -> Target.CArray[1].B
+		// SourceA.FixedArray -> Target.FixedArray
+		// SourceB.Item -> Target.FixedArray[1]
+		// SourceA.Item.B -> Target.FixedArray[1].B
 		// SourceA.Array -> Target.Array
 		// SourceB.Item -> Target.Array[1]
 		// SourceA.Item.B -> Target.Array[1].B
@@ -102,7 +126,26 @@ struct FStateTreeTest_BindingsCompiler : FStateTreeTestBase
 		AITEST_EQUAL("Expect Target.Array[0].B copied from SourceA.Array[0].B", Target.Array[0].B, SourceA.Array[0].B);
 		AITEST_EQUAL("Expect Target.Array[1].A copied from SourceB.Item.A", Target.Array[1].A, SourceB.Item.A);
 		AITEST_EQUAL("Expect Target.Array[1].B copied from SourceA.Item.B", Target.Array[1].B, SourceA.Item.B);
+
+		AITEST_EQUAL("Expect TargetArray to be copied from SourceA", Target.FixedArray.Num(), SourceA.FixedArray.Num());
+		AITEST_EQUAL("Expect Target.FixedArray[0].A copied from SourceA.FixedArray[0].A", Target.FixedArray[0].A, SourceA.FixedArray[0].A);
+		AITEST_EQUAL("Expect Target.FixedArray[0].B copied from SourceA.FixedArray[0].B", Target.FixedArray[0].B, SourceA.FixedArray[0].B);
+		AITEST_EQUAL("Expect Target.FixedArray[1].A copied from SourceB.Item.A", Target.FixedArray[1].A, SourceB.Item.A);
+		AITEST_EQUAL("Expect Target.FixedArray[1].B copied from SourceA.Item.B", Target.FixedArray[1].B, SourceA.Item.B);
+		AITEST_EQUAL("Expect Target.FixedArray to not have changed size", Target.FixedArray.Num(), FixedArraySize);
+
+		AITEST_EQUAL("Expect Target.CArray[0].A copied from SourceA.CArray[0].A", Target.CArray[0].A, SourceA.CArray[0].A);
+		AITEST_EQUAL("Expect Target.CArray[0].B copied from SourceA.CArray[0].B", Target.CArray[0].B, SourceA.CArray[0].B);
+		AITEST_EQUAL("Expect Target.CArray[1].A copied from SourceB.Item.A", Target.CArray[1].A, SourceB.Item.A);
+		AITEST_EQUAL("Expect Target.CArray[1].B copied from SourceA.Item.B", Target.CArray[1].B, SourceA.Item.B);
 		
+		const int32 NumAllocated_FStateTreeTest_PropertyStructB_BeforeReset = FStateTreeTest_PropertyStructB::NumConstructed;
+		bool bResetResult = Bindings.FPropertyBindingBindingCollection::ResetObjects(FPropertyBindingIndex16(CopyBatchIndex), TargetView);
+		AITEST_TRUE("ResetObjects should succeed", bResetResult);
+		AITEST_EQUAL("Expect Target dynamic array to be empty", Target.Array.Num(), 0);
+		AITEST_EQUAL("Expect Target fixed size Array to not have changed size.", Target.FixedArray.Num(), FixedArraySize);
+		AITEST_NOT_EQUAL("Expect the count of constructed FStateTreeTest_PropertyStructB to be smaller after calling ResetObjects", FStateTreeTest_PropertyStructB::NumConstructed, NumAllocated_FStateTreeTest_PropertyStructB_BeforeReset);
+
 		return true;
 	}
 };
diff --git a/Engine/Plugins/Runtime/StateTree/Source/StateTreeTestSuite/Private/StateTreeTestTypes.h b/Engine/Plugins/Runtime/StateTree/Source/StateTreeTestSuite/Private/StateTreeTestTypes.h
--- a/Engine/Plugins/Runtime/StateTree/Source/StateTreeTestSuite/Private/StateTreeTestTypes.h
+++ b/Engine/Plugins/Runtime/StateTree/Source/StateTreeTestSuite/Private/StateTreeTestTypes.h
@@ -616,8 +616,19 @@ struct FStateTreeTest_PropertyStructB
 {
 	GENERATED_BODY()
 
+	FStateTreeTest_PropertyStructB()
+	{
+		NumConstructed++;
+	}
+	~FStateTreeTest_PropertyStructB()
+	{
+		NumConstructed--;
+	}
+
 	UPROPERTY(EditAnywhere, Category = "")
 	int32 B = 0;
+
+	static inline int32 NumConstructed = 0;
 };
 
 USTRUCT()
@@ -706,6 +717,12 @@ struct FStateTreeTest_PropertyCopy
 
 	UPROPERTY(EditAnywhere, Category = "")
 	TArray<FStateTreeTest_PropertyStruct> Array;
+
+	UPROPERTY(EditAnywhere, Category = "", EditFixedSize)
+	TArray<FStateTreeTest_PropertyStruct> FixedArray;
+
+	UPROPERTY(EditAnywhere, Category = "")
+	FStateTreeTest_PropertyStruct CArray[3];
 };
 
 USTRUCT()
