[StateTree] added setter to keep the update phase of the execution state and the debug macros in sync. This also fixes an issue where StartPhase was not popped soon enough when processing a deferred stop when starting the execution.

diff --git a/Engine/Plugins/Runtime/StateTree/Source/StateTreeModule/Private/StateTreeExecutionContext.cpp b/Engine/Plugins/Runtime/StateTree/Source/StateTreeModule/Private/StateTreeExecutionContext.cpp
--- a/Engine/Plugins/Runtime/StateTree/Source/StateTreeModule/Private/StateTreeExecutionContext.cpp
+++ b/Engine/Plugins/Runtime/StateTree/Source/StateTreeModule/Private/StateTreeExecutionContext.cpp
@@ -647,6 +647,26 @@ EStateTreeRunStatus FStateTreeExecutionContext::Start(const FInstancedPropertyBa
 	return Start(FStartParameters{.GlobalParameters = InitialParameters, .RandomSeed = ParamRandomSeed });
 }
 
+void FStateTreeExecutionContext::SetUpdatePhaseInExecutionState(FStateTreeExecutionState& ExecutionState, const EStateTreeUpdatePhase UpdatePhase) const
+{
+	if (ExecutionState.CurrentPhase == UpdatePhase)
+	{
+		return;
+	}
+
+	if (ExecutionState.CurrentPhase != EStateTreeUpdatePhase::Unset)
+	{
+		UE_STATETREE_DEBUG_EXIT_PHASE(this, ExecutionState.CurrentPhase);
+	}
+
+	ExecutionState.CurrentPhase = UpdatePhase;
+
+	if (ExecutionState.CurrentPhase != EStateTreeUpdatePhase::Unset)
+	{
+		UE_STATETREE_DEBUG_ENTER_PHASE(this, ExecutionState.CurrentPhase);
+	}
+}
+
 EStateTreeRunStatus FStateTreeExecutionContext::Start(FStartParameters Parameters)
 {
 	CSV_SCOPED_TIMING_STAT_EXCLUSIVE(StateTree_Start);
@@ -714,16 +734,12 @@ EStateTreeRunStatus FStateTreeExecutionContext::Start(FStartParameters Parameter
 
 	// Must sent instance creation event first 
 	UE_STATETREE_DEBUG_INSTANCE_EVENT(this, EStateTreeTraceEventType::Push);
-	
-	// Set scoped phase only for properly initialized context with valid Instance data
-	// since we need it to output the InstanceId
-	UE_STATETREE_DEBUG_SCOPED_PHASE(this, EStateTreeUpdatePhase::StartTree);
 
 	STATETREE_LOG(VeryVerbose, TEXT("%hs: Starting State Tree %s on owner '%s'."),
 		__FUNCTION__, *GetFullNameSafe(&RootStateTree), *GetNameSafe(&Owner));
 
 	// From this point any calls to Stop should be deferred.
-	Exec.CurrentPhase = EStateTreeUpdatePhase::StartTree;
+	SetUpdatePhaseInExecutionState(Exec, EStateTreeUpdatePhase::StartTree);
 
 	// Start evaluators and global tasks. Fail the execution if any global task fails.
 	FStateTreeIndex16 LastInitializedTaskIndex;
@@ -796,8 +812,8 @@ EStateTreeRunStatus FStateTreeExecutionContext::Start(FStartParameters Parameter
 		Exec.TreeRunStatus = GlobalTasksRunStatus;
 	}
 
-	// Reset phase since we are now safe to stop.
-	Exec.CurrentPhase = EStateTreeUpdatePhase::Unset;
+	// Reset phase since we are now safe to stop and before potentially stopping the instance.
+	SetUpdatePhaseInExecutionState(Exec, EStateTreeUpdatePhase::Unset);
 
 	// Use local for resulting run state since Stop will reset the instance data.
 	EStateTreeRunStatus Result = Exec.TreeRunStatus;
@@ -851,12 +867,8 @@ EStateTreeRunStatus FStateTreeExecutionContext::Stop(EStateTreeRunStatus Complet
 		return EStateTreeRunStatus::Running;
 	}
 
-	// Set phase only after validating that we are in a valid context since we need an InstancedId for the debug
-	// We intentionally use ENTER_PHASE instead of SCOPED_PHASE here to make sure we exit the phase before the Instance
-	UE_STATETREE_DEBUG_ENTER_PHASE(this, EStateTreeUpdatePhase::StopTree);
-
 	// No need to clear on exit since we reset all the instance data before leaving the function.
-	Exec.CurrentPhase = EStateTreeUpdatePhase::StopTree;
+	SetUpdatePhaseInExecutionState(Exec, EStateTreeUpdatePhase::StopTree);
 
 	EStateTreeRunStatus Result = Exec.TreeRunStatus;
 	
@@ -920,8 +932,7 @@ EStateTreeRunStatus FStateTreeExecutionContext::TickPrelude()
 	}
 
 	// From this point any calls to Stop should be deferred.
-	Exec.CurrentPhase = EStateTreeUpdatePhase::TickStateTree;
-	UE_STATETREE_DEBUG_ENTER_PHASE(this, EStateTreeUpdatePhase::TickStateTree);
+	SetUpdatePhaseInExecutionState(Exec, EStateTreeUpdatePhase::TickStateTree);
 
 	UpdateCompletedStateList();
 
@@ -934,8 +945,7 @@ EStateTreeRunStatus FStateTreeExecutionContext::TickPostlude()
 	FStateTreeExecutionState& Exec = GetExecState();
 
 	// Reset phase since we are now safe to stop.
-	UE_STATETREE_DEBUG_EXIT_PHASE(this, EStateTreeUpdatePhase::TickStateTree);
-	Exec.CurrentPhase = EStateTreeUpdatePhase::Unset;
+	SetUpdatePhaseInExecutionState(Exec, EStateTreeUpdatePhase::Unset);
 
 	// Use local for resulting run state since Stop will reset the instance data.
 	EStateTreeRunStatus Result = Exec.TreeRunStatus;
diff --git a/Engine/Plugins/Runtime/StateTree/Source/StateTreeModule/Public/StateTreeExecutionContext.h b/Engine/Plugins/Runtime/StateTree/Source/StateTreeModule/Public/StateTreeExecutionContext.h
--- a/Engine/Plugins/Runtime/StateTree/Source/StateTreeModule/Public/StateTreeExecutionContext.h
+++ b/Engine/Plugins/Runtime/StateTree/Source/StateTreeModule/Public/StateTreeExecutionContext.h
@@ -863,6 +863,9 @@ protected:
 		return InstanceDataStorage->GetExecutionState();
 	}
 
+	/** Updates the update phase of the statetree execution state. */
+	void SetUpdatePhaseInExecutionState(FStateTreeExecutionState& ExecutionState, EStateTreeUpdatePhase UpdatePhase) const;
+
 	/** @return String describing state status for logging and debug. */
 	FString GetStateStatusString(const FStateTreeExecutionState& ExecState) const;
 
