Synchronize ST scrubber with rewind debugger.

#tests PIE, editor, avalanche

diff --git a/Engine/Plugins/Runtime/StateTree/Source/StateTreeEditorModule/Private/Debugger/StateTreeRewindDebuggerExtensions.cpp b/Engine/Plugins/Runtime/StateTree/Source/StateTreeEditorModule/Private/Debugger/StateTreeRewindDebuggerExtensions.cpp
new file mode 100644
--- /dev/null
+++ b/Engine/Plugins/Runtime/StateTree/Source/StateTreeEditorModule/Private/Debugger/StateTreeRewindDebuggerExtensions.cpp
@@ -0,0 +1,68 @@
+// Copyright Epic Games, Inc. All Rights Reserved.
+
+#if WITH_STATETREE_TRACE_DEBUGGER
+
+#include "Debugger/StateTreeRewindDebuggerExtensions.h"
+
+#include "Debugger/StateTreeTrace.h"
+#include "Debugger/StateTreeTraceProvider.h"
+#include "StateTreeDelegates.h"
+#include "Trace/Trace.h"
+
+namespace UE::StateTreeDebugger
+{
+
+void FRewindDebuggerExtension::Update(float DeltaTime, IRewindDebugger* RewindDebugger)
+{
+	const TraceServices::IAnalysisSession* AnalysisSession = RewindDebugger->GetAnalysisSession();
+	TraceServices::FAnalysisSessionReadScope SessionReadScope(*AnalysisSession);
+	
+	const IStateTreeTraceProvider* TraceProvider = AnalysisSession->ReadProvider<IStateTreeTraceProvider>(FStateTreeTraceProvider::ProviderName);
+
+	const TraceServices::IFrameProvider& FrameProvider = TraceServices::ReadFrameProvider(*RewindDebugger->GetAnalysisSession());
+	TraceServices::FFrame Frame;
+
+	// Require some debug frame to exist before doing any processing, currently the frame itself is not used
+	if (FrameProvider.GetFrameFromTime(ETraceFrameType::TraceFrameType_Game, RewindDebugger->CurrentTraceTime(), Frame))
+	{
+		auto SetScrubberTimeline = [&]()
+		{
+			if (RewindDebugger->IsPIESimulating())
+			{
+				return;
+			}
+
+			const double CurrentScrubTime = RewindDebugger->GetScrubTime();
+			if (LastScrubTime != CurrentScrubTime)
+			{
+				UE::StateTree::Delegates::OnTracingTimelineScrubbed.Broadcast(CurrentScrubTime);
+				LastScrubTime = CurrentScrubTime;
+			}
+		};
+		SetScrubberTimeline();
+	}
+}
+
+void FRewindDebuggerRuntimeExtension::RecordingStarted()
+{ 
+	if (!UE_TRACE_CHANNELEXPR_IS_ENABLED(StateTreeDebugChannel))
+	{
+		UE::Trace::ToggleChannel(TEXT("StateTreeDebugChannel"), true);
+		UE::StateTree::Delegates::OnTracingStateChanged.Broadcast(EStateTreeTraceStatus::TracesStarted);
+	}
+}
+
+void FRewindDebuggerRuntimeExtension::RecordingStopped()
+{
+	// Shouldn't fire as channel will already be disabled when rewind debugger stops, but just as safeguard
+	if (UE_TRACE_CHANNELEXPR_IS_ENABLED(StateTreeDebugChannel))
+	{
+		UE::StateTree::Delegates::OnTracingStateChanged.Broadcast(EStateTreeTraceStatus::StoppingTrace);
+		UE::Trace::ToggleChannel(TEXT("StateTreeDebugChannel"), false);
+	}
+}
+
+
+} // namespace UE::StateTreeDebugger
+
+#endif // WITH_STATETREE_TRACE_DEBUGGER
\ No newline at end of file
diff --git a/Engine/Plugins/Runtime/StateTree/Source/StateTreeEditorModule/Private/Debugger/StateTreeRewindDebuggerExtensions.h b/Engine/Plugins/Runtime/StateTree/Source/StateTreeEditorModule/Private/Debugger/StateTreeRewindDebuggerExtensions.h
new file mode 100644
--- /dev/null
+++ b/Engine/Plugins/Runtime/StateTree/Source/StateTreeEditorModule/Private/Debugger/StateTreeRewindDebuggerExtensions.h
@@ -0,0 +1,35 @@
+// Copyright Epic Games, Inc. All Rights Reserved.
+#pragma once
+
+#if WITH_STATETREE_TRACE_DEBUGGER
+
+#include "IRewindDebugger.h"
+#include "IRewindDebuggerExtension.h"
+#include "RewindDebuggerRuntimeInterface/IRewindDebuggerRuntimeExtension.h"
+
+namespace UE::StateTreeDebugger
+{
+
+class FRewindDebuggerExtension : public IRewindDebuggerExtension
+{
+public:
+
+	virtual void Update(float DeltaTime, IRewindDebugger* RewindDebugger) override;
+
+private:
+
+	/** Last scrub time we received. Used to avoid redundant updates. */
+	double LastScrubTime = 0.0;
+};
+
+class FRewindDebuggerRuntimeExtension : public IRewindDebuggerRuntimeExtension
+{
+public:
+
+	virtual void RecordingStarted() override;
+	virtual void RecordingStopped() override;
+};
+
+} // namespace UE::StateTreeDebugger
+
+#endif // WITH_STATETREE_TRACE_DEBUGGER
\ No newline at end of file
diff --git a/Engine/Plugins/Runtime/StateTree/Source/StateTreeEditorModule/Private/StateTreeEditorModule.cpp b/Engine/Plugins/Runtime/StateTree/Source/StateTreeEditorModule/Private/StateTreeEditorModule.cpp
--- a/Engine/Plugins/Runtime/StateTree/Source/StateTreeEditorModule/Private/StateTreeEditorModule.cpp
+++ b/Engine/Plugins/Runtime/StateTree/Source/StateTreeEditorModule/Private/StateTreeEditorModule.cpp
@@ -20,11 +20,13 @@
 #include "Customizations/StateTreeEventDescDetails.h"
 #include "Customizations/StateTreeBindingExtension.h"
 #include "Customizations/StateTreeBlueprintPropertyRefDetails.h"
+#include "IRewindDebuggerExtension.h"
 #include "PropertyEditorModule.h"
 #include "StateTree.h"
 #include "StateTreeCompiler.h"
 #include "StateTreeCompilerLog.h"
 #include "Debugger/StateTreeDebuggerCommands.h"
+#include "Debugger/StateTreeRewindDebuggerExtensions.h"
 #include "StateTreeDelegates.h"
 #include "StateTreeEditingSubsystem.h"
 #include "StateTreeEditor.h"
@@ -57,6 +59,10 @@ void FStateTreeEditorModule::StartupModule()
 
 #if WITH_STATETREE_TRACE_DEBUGGER
 	FStateTreeDebuggerCommands::Register();
+	RewindDebuggerExtension = MakePimpl<UE::StateTreeDebugger::FRewindDebuggerExtension>();
+	IModularFeatures::Get().RegisterModularFeature(IRewindDebuggerExtension::ModularFeatureName, RewindDebuggerExtension.Get());
+	RewindDebuggerRuntimeExtension = MakePimpl<UE::StateTreeDebugger::FRewindDebuggerRuntimeExtension>();
+	IModularFeatures::Get().RegisterModularFeature(IRewindDebuggerRuntimeExtension::ModularFeatureName, RewindDebuggerRuntimeExtension.Get());
 #endif // WITH_STATETREE_TRACE_DEBUGGER
 
 	MenuExtensibilityManager = MakeShareable(new FExtensibilityManager);
@@ -90,6 +96,8 @@ void FStateTreeEditorModule::ShutdownModule()
 	UE::StateTree::Delegates::OnRequestCompile.Unbind();
 
 #if WITH_STATETREE_TRACE_DEBUGGER
+	IModularFeatures::Get().UnregisterModularFeature(IRewindDebuggerRuntimeExtension::ModularFeatureName, RewindDebuggerRuntimeExtension.Get());
+	IModularFeatures::Get().UnregisterModularFeature(IRewindDebuggerExtension::ModularFeatureName, RewindDebuggerExtension.Get());
 	FStateTreeDebuggerCommands::Unregister();
 #endif // WITH_STATETREE_TRACE_DEBUGGER
 	
diff --git a/Engine/Plugins/Runtime/StateTree/Source/StateTreeEditorModule/Public/StateTreeEditorModule.h b/Engine/Plugins/Runtime/StateTree/Source/StateTreeEditorModule/Public/StateTreeEditorModule.h
--- a/Engine/Plugins/Runtime/StateTree/Source/StateTreeEditorModule/Public/StateTreeEditorModule.h
+++ b/Engine/Plugins/Runtime/StateTree/Source/StateTreeEditorModule/Public/StateTreeEditorModule.h
@@ -4,6 +4,7 @@
 
 #include "Modules/ModuleInterface.h"
 #include "Logging/LogMacros.h"
+#include "Templates/PimplPtr.h"
 #include "Toolkits/AssetEditorToolkit.h"
 
 class UStateTree;
@@ -11,6 +12,12 @@ class UUserDefinedStruct;
 class IStateTreeEditor;
 struct FStateTreeNodeClassCache;
 
+namespace UE::StateTreeDebugger
+{
+	class FRewindDebuggerExtension;
+	class FRewindDebuggerRuntimeExtension;
+}
+
 STATETREEEDITORMODULE_API DECLARE_LOG_CATEGORY_EXTERN(LogStateTreeEditor, Log, All);
 
 /**
@@ -41,6 +48,10 @@ protected:
 	TSharedPtr<FExtensibilityManager> MenuExtensibilityManager;
 	TSharedPtr<FExtensibilityManager> ToolBarExtensibilityManager;
 	TSharedPtr<FStateTreeNodeClassCache> NodeClassCache;
+#if WITH_STATETREE_TRACE_DEBUGGER
+	TPimplPtr<UE::StateTreeDebugger::FRewindDebuggerExtension> RewindDebuggerExtension;
+	TPimplPtr<UE::StateTreeDebugger::FRewindDebuggerRuntimeExtension> RewindDebuggerRuntimeExtension;
+#endif  // WITH_STATETREE_TRACE_DEBUGGER
 
 	FDelegateHandle OnUserDefinedStructReinstancedHandle;
 	FOnRegisterLayoutExtensions RegisterLayoutExtensions;
diff --git a/Engine/Plugins/Runtime/StateTree/Source/StateTreeEditorModule/StateTreeEditorModule.Build.cs b/Engine/Plugins/Runtime/StateTree/Source/StateTreeEditorModule/StateTreeEditorModule.Build.cs
--- a/Engine/Plugins/Runtime/StateTree/Source/StateTreeEditorModule/StateTreeEditorModule.Build.cs
+++ b/Engine/Plugins/Runtime/StateTree/Source/StateTreeEditorModule/StateTreeEditorModule.Build.cs
@@ -49,9 +49,13 @@ namespace UnrealBuildTool.Rules
 				"SourceCodeAccess",
 				"ToolMenus",
 				"ToolWidgets",
+				"TraceLog",
+				"TraceServices",
 				"ApplicationCore",
 				"DeveloperSettings",
+				"GameplayInsights",
 				"RewindDebuggerInterface",
+				"RewindDebuggerRuntimeInterface",
 				"DetailCustomizations",
 				"AppFramework",
 				"Kismet",
diff --git a/Engine/Plugins/Runtime/StateTree/Source/StateTreeModule/Private/Debugger/StateTreeDebugger.cpp b/Engine/Plugins/Runtime/StateTree/Source/StateTreeModule/Private/Debugger/StateTreeDebugger.cpp
--- a/Engine/Plugins/Runtime/StateTree/Source/StateTreeModule/Private/Debugger/StateTreeDebugger.cpp
+++ b/Engine/Plugins/Runtime/StateTree/Source/StateTreeModule/Private/Debugger/StateTreeDebugger.cpp
@@ -86,6 +86,11 @@ FStateTreeDebugger::FStateTreeDebugger()
 				RequestAnalysisOfLatestTrace();
 			}
 		});
+
+	TracingTimelineScrubbedHandle = UE::StateTree::Delegates::OnTracingTimelineScrubbed.AddLambda([this](const double InScrubTime)
+		{
+			SetScrubTime(InScrubTime);
+		});
 }
 
 FStateTreeDebugger::~FStateTreeDebugger()
@@ -93,6 +98,9 @@ FStateTreeDebugger::~FStateTreeDebugger()
 	UE::StateTree::Delegates::OnTracingStateChanged.Remove(TracingStateChangedHandle);
 	TracingStateChangedHandle.Reset();
 
+	UE::StateTree::Delegates::OnTracingTimelineScrubbed.Remove(TracingTimelineScrubbedHandle);
+	TracingTimelineScrubbedHandle.Reset();
+
 	StopSessionAnalysis();
 }
 
diff --git a/Engine/Plugins/Runtime/StateTree/Source/StateTreeModule/Private/StateTreeDelegates.cpp b/Engine/Plugins/Runtime/StateTree/Source/StateTreeModule/Private/StateTreeDelegates.cpp
--- a/Engine/Plugins/Runtime/StateTree/Source/StateTreeModule/Private/StateTreeDelegates.cpp
+++ b/Engine/Plugins/Runtime/StateTree/Source/StateTreeModule/Private/StateTreeDelegates.cpp
@@ -22,4 +22,8 @@ FOnRequestEditorHash OnRequestEditorHash;
 FOnTracingStateChanged OnTracingStateChanged;
 #endif // WITH_STATETREE_TRACE
 
+#if WITH_STATETREE_TRACE_DEBUGGER
+FOnTracingTimelineScrubbed OnTracingTimelineScrubbed;
+#endif // WITH_STATETREE_TRACE_DEBUGGER
+
 }; // UE::StateTree::Delegates
diff --git a/Engine/Plugins/Runtime/StateTree/Source/StateTreeModule/Public/Debugger/StateTreeDebugger.h b/Engine/Plugins/Runtime/StateTree/Source/StateTreeModule/Public/Debugger/StateTreeDebugger.h
--- a/Engine/Plugins/Runtime/StateTree/Source/StateTreeModule/Public/Debugger/StateTreeDebugger.h
+++ b/Engine/Plugins/Runtime/StateTree/Source/StateTreeModule/Public/Debugger/StateTreeDebugger.h
@@ -341,6 +341,11 @@ private:
 	 * Delegate Handle bound to UE::StateTree::Delegates::OnTracingStateChanged
 	 */
 	FDelegateHandle TracingStateChangedHandle;
+
+	/**
+	 * Delegate Handle bound to UE::StateTree::Delegates::OnTracingTimelineScrubbed
+	 */
+	FDelegateHandle TracingTimelineScrubbedHandle;
 };
 
 #endif // WITH_STATETREE_TRACE_DEBUGGER
diff --git a/Engine/Plugins/Runtime/StateTree/Source/StateTreeModule/Public/Debugger/StateTreeTraceProvider.h b/Engine/Plugins/Runtime/StateTree/Source/StateTreeModule/Public/Debugger/StateTreeTraceProvider.h
--- a/Engine/Plugins/Runtime/StateTree/Source/StateTreeModule/Public/Debugger/StateTreeTraceProvider.h
+++ b/Engine/Plugins/Runtime/StateTree/Source/StateTreeModule/Public/Debugger/StateTreeTraceProvider.h
@@ -16,7 +16,7 @@ namespace UE::StateTreeDebugger { struct FInstanceDescriptor; }
 class FStateTreeTraceProvider : public IStateTreeTraceProvider
 {
 public:
-	static FName ProviderName;
+	STATETREEMODULE_API static FName ProviderName;
 
 	explicit FStateTreeTraceProvider(TraceServices::IAnalysisSession& InSession);
 	
diff --git a/Engine/Plugins/Runtime/StateTree/Source/StateTreeModule/Public/StateTreeDelegates.h b/Engine/Plugins/Runtime/StateTree/Source/StateTreeModule/Public/StateTreeDelegates.h
--- a/Engine/Plugins/Runtime/StateTree/Source/StateTreeModule/Public/StateTreeDelegates.h
+++ b/Engine/Plugins/Runtime/StateTree/Source/StateTreeModule/Public/StateTreeDelegates.h
@@ -82,4 +82,12 @@ extern STATETREEMODULE_API FOnTracingStateChanged OnTracingStateChanged;
 
 #endif // WITH_STATETREE_TRACE
 
+#if WITH_STATETREE_TRACE_DEBUGGER
+
+/** Called by the StateTree module whenever tracing timeline is scrubbed on Rewind Debugger */
+DECLARE_MULTICAST_DELEGATE_OneParam(FOnTracingTimelineScrubbed, const double /* InScrubTime */);
+extern STATETREEMODULE_API FOnTracingTimelineScrubbed OnTracingTimelineScrubbed;
+
+#endif // WITH_STATETREE_TRACE_DEBUGGER
+
 }; // UE::StateTree::Delegates
diff --git a/Engine/Plugins/Runtime/StateTree/StateTree.uplugin b/Engine/Plugins/Runtime/StateTree/StateTree.uplugin
--- a/Engine/Plugins/Runtime/StateTree/StateTree.uplugin
+++ b/Engine/Plugins/Runtime/StateTree/StateTree.uplugin
@@ -32,6 +32,13 @@
 		}
 	],
 	"Plugins": [
+		{
+			"Name": "GameplayInsights",
+			"Enabled": true,
+			"TargetAllowList": [
+				"Editor"
+			]
+		},
 		{
 			"Name": "PropertyAccessEditor",
 			"Enabled": true
