Fix StateTree debug color not randomizing well.

#tests PIE

diff --git a/Engine/Plugins/Runtime/StateTree/Source/StateTreeEditorModule/Private/Debugger/StateTreeDebuggerTrack.cpp b/Engine/Plugins/Runtime/StateTree/Source/StateTreeEditorModule/Private/Debugger/StateTreeDebuggerTrack.cpp
--- a/Engine/Plugins/Runtime/StateTree/Source/StateTreeEditorModule/Private/Debugger/StateTreeDebuggerTrack.cpp
+++ b/Engine/Plugins/Runtime/StateTree/Source/StateTreeEditorModule/Private/Debugger/StateTreeDebuggerTrack.cpp
@@ -95,7 +95,7 @@ bool FStateTreeDebuggerInstanceTrack::UpdateInternal()
 				FString StatePath = Event.GetValueString(*StateTree);
 				UE::StateTreeDebugger::FFrameSpan Span = EventCollection.FrameSpans[SpanIndex];
 				SStateTreeDebuggerEventTimelineView::FTimelineEventData::EventWindow& Window = EventData->Windows.AddDefaulted_GetRef();
-				Window.Color = MakeRandomColor(GetTypeHash(StatePath));
+				Window.Color = MakeRandomColor(CityHash32(reinterpret_cast<const char*>(*StatePath), StatePath.Len() * sizeof(FString::ElementType)));
 				Window.Description = FText::FromString(StatePath);
 				Window.TimeStart = Span.GetWorldTimeStart();
 
