StateTree: For linked assets, test if the schema is the same. Add unit test for invalid SetLinkedStateTreeOverrides.

diff --git a/Engine/Plugins/Runtime/StateTree/Source/StateTreeModule/Private/StateTreeExecutionContext.cpp b/Engine/Plugins/Runtime/StateTree/Source/StateTreeModule/Private/StateTreeExecutionContext.cpp
--- a/Engine/Plugins/Runtime/StateTree/Source/StateTreeModule/Private/StateTreeExecutionContext.cpp
+++ b/Engine/Plugins/Runtime/StateTree/Source/StateTreeModule/Private/StateTreeExecutionContext.cpp
@@ -493,8 +493,8 @@ void FStateTreeExecutionContext::SetLinkedStateTreeOverrides(FStateTreeReference
 	bool bValid = true;
 
 	// Confirms that the overrides schema matches.
-	const TConstArrayView<FStateTreeReferenceOverrideItem> OverrideItems = LinkedAssetStateTreeOverrides.GetOverrideItems();
-	for (const FStateTreeReferenceOverrideItem& Item : OverrideItems)
+	const TConstArrayView<FStateTreeReferenceOverrideItem> InOverrideItems = InLinkedStateTreeOverrides.GetOverrideItems();
+	for (const FStateTreeReferenceOverrideItem& Item : InOverrideItems)
 	{
 		if (const UStateTree* ItemStateTree = Item.GetStateTreeReference().GetStateTree())
 		{
@@ -506,6 +506,14 @@ void FStateTreeExecutionContext::SetLinkedStateTreeOverrides(FStateTreeReference
 				break;
 			}
 
+			if (!RootStateTree.HasCompatibleContextData(*ItemStateTree))
+			{
+				STATETREE_LOG(Error, TEXT("%hs: '%s' using StateTree '%s' trying to set override '%s' but the tree context data is not compatible."),
+					__FUNCTION__, *GetNameSafe(&Owner), *GetFullNameSafe(GetStateTree()), *GetFullNameSafe(ItemStateTree));
+				bValid = false;
+				break;
+			}
+
 			const UStateTreeSchema* OverrideSchema = ItemStateTree->GetSchema();
 			if (ItemStateTree->GetSchema() == nullptr)
 			{
@@ -526,19 +534,25 @@ void FStateTreeExecutionContext::SetLinkedStateTreeOverrides(FStateTreeReference
 		}
 	}
 
+	bool bChanged = false;
 	if (bValid)
 	{
 		LinkedAssetStateTreeOverrides = MoveTemp(InLinkedStateTreeOverrides);
+		bChanged = LinkedAssetStateTreeOverrides.GetOverrideItems().Num() > 0;
 	}
-	else
+	else if (LinkedAssetStateTreeOverrides.GetOverrideItems().Num() > 0)
 	{
 		LinkedAssetStateTreeOverrides.Reset();
+		bChanged = true;
 	}
 	
-	TInstancedStruct<FStateTreeExecutionExtension>& ExecutionExtension = Storage.GetMutableExecutionState().ExecutionExtension;
-	if (ExecutionExtension.IsValid())
+	if (bChanged)
 	{
-		ExecutionExtension.GetMutable().OnLinkedStateTreeOverridesSet(FStateTreeExecutionExtension::FContextParameters(Owner, RootStateTree, Storage), LinkedAssetStateTreeOverrides);
+		TInstancedStruct<FStateTreeExecutionExtension>& ExecutionExtension = Storage.GetMutableExecutionState().ExecutionExtension;
+		if (ExecutionExtension.IsValid())
+		{
+			ExecutionExtension.GetMutable().OnLinkedStateTreeOverridesSet(FStateTreeExecutionExtension::FContextParameters(Owner, RootStateTree, Storage), LinkedAssetStateTreeOverrides);
+		}
 	}
 }
 
@@ -4988,7 +5002,8 @@ bool FStateTreeExecutionContext::SelectStateInternal(
 			}
 
 			// The linked state tree should have compatible context requirements.
-			if (!NextLinkedStateAsset->HasCompatibleContextData(RootStateTree))
+			if (!NextLinkedStateAsset->HasCompatibleContextData(RootStateTree)
+				|| NextLinkedStateAsset->GetSchema()->GetClass() != RootStateTree.GetSchema()->GetClass())
 			{
 				STATETREE_LOG(Error, TEXT("%hs: The linked State Tree '%s' does not have compatible schema, trying to select state %s from '%s'.  '%s' using StateTree '%s'."),
 					__FUNCTION__, *GetFullNameSafe(NextLinkedStateAsset), *GetSafeStateName(CurrentFrame, NextStateHandle), *GetStateStatusString(Exec), *GetNameSafe(&Owner), *GetFullNameSafe(CurrentFrame.StateTree));
diff --git a/Engine/Plugins/Runtime/StateTree/Source/StateTreeTestSuite/Private/StateTreeTest.cpp b/Engine/Plugins/Runtime/StateTree/Source/StateTreeTestSuite/Private/StateTreeTest.cpp
--- a/Engine/Plugins/Runtime/StateTree/Source/StateTreeTestSuite/Private/StateTreeTest.cpp
+++ b/Engine/Plugins/Runtime/StateTree/Source/StateTreeTestSuite/Private/StateTreeTest.cpp
@@ -1789,6 +1789,7 @@ struct FStateTreeTest_NestedOverride : FStateTreeTestBase
 		FStateTreeCompilerLog Log;
 		
 		const FGameplayTag Tag = GetTestTag1();
+		const FGameplayTag Tag2 = GetTestTag2();
 
 		// Asset 2
 		UStateTree& StateTree2 = NewStateTree();
@@ -1796,11 +1797,11 @@ struct FStateTreeTest_NestedOverride : FStateTreeTestBase
 		EditorData2.RootParameters.Parameters.AddProperty(FName(TEXT("Int")), EPropertyBagPropertyType::Int32);
 		UStateTreeState& Root2 = EditorData2.AddSubTree(FName(TEXT("Root2")));
 		TStateTreeEditorNode<FTestTask_Stand>& TaskRoot2 = Root2.AddTask<FTestTask_Stand>(FName(TEXT("TaskRoot2")));
-
-		FStateTreeCompiler Compiler2(Log);
-		const bool bResult2 = Compiler2.Compile(StateTree2);
-		AITEST_TRUE("StateTree2 should get compiled", bResult2);
-
+		{
+			FStateTreeCompiler Compiler2(Log);
+			const bool bResult2 = Compiler2.Compile(StateTree2);
+			AITEST_TRUE("StateTree2 should get compiled", bResult2);
+		}
 		
 		// Asset 3
 		UStateTree& StateTree3 = NewStateTree();
@@ -1808,10 +1809,27 @@ struct FStateTreeTest_NestedOverride : FStateTreeTestBase
 		EditorData3.RootParameters.Parameters.AddProperty(FName(TEXT("Float")), EPropertyBagPropertyType::Float); // Different parameters
 		UStateTreeState& Root3 = EditorData3.AddSubTree(FName(TEXT("Root3")));
 		TStateTreeEditorNode<FTestTask_Stand>& TaskRoot3 = Root3.AddTask<FTestTask_Stand>(FName(TEXT("TaskRoot3")));
+		{
+			FStateTreeCompiler Compiler3(Log);
+			const bool bResult3 = Compiler3.Compile(StateTree3);
+			AITEST_TRUE("StateTree3 should get compiled", bResult3);
+		}
+		// Wrong Asset 4
+		UStateTree* StateTree4 = NewObject<UStateTree>(&GetWorld());;
+		{
+			UStateTreeEditorData* EditorData = NewObject<UStateTreeEditorData>(StateTree4);
+			check(EditorData);
+			StateTree4->EditorData = EditorData;
+			EditorData->Schema = NewObject<UStateTreeTestSchema2>();
+			
+			EditorData->RootParameters.Parameters.AddProperty(FName(TEXT("Float")), EPropertyBagPropertyType::Float); // Different parameters
+			UStateTreeState& Root4 = EditorData->AddSubTree(FName(TEXT("Root4")));
+			TStateTreeEditorNode<FTestTask_Stand>& TaskRoot4 = Root3.AddTask<FTestTask_Stand>(FName(TEXT("TaskRoot4")));
 
-		FStateTreeCompiler Compiler3(Log);
-		const bool bResult3 = Compiler3.Compile(StateTree3);
-		AITEST_TRUE("StateTree3 should get compiled", bResult3);
+			FStateTreeCompiler Compiler4(Log);
+			const bool bResult4 = Compiler4.Compile(*StateTree4);
+			AITEST_TRUE("StateTree4 should get compiled", bResult4);
+		}
 
 		// Main asset
 		UStateTree& StateTree = NewStateTree();
@@ -1866,6 +1884,37 @@ struct FStateTreeTest_NestedOverride : FStateTreeTestBase
 			Status = Exec.Start();
 			AITEST_EQUAL("Start should complete with Running", Status, EStateTreeRunStatus::Running);
 			AITEST_TRUE("StateTree should enter TaskRoot3", Exec.Expect(TaskRoot3.GetName(), EnterStateStr));
+			AITEST_FALSE("StateTree should not enter TaskRoot2", Exec.Expect(TaskRoot2.GetName(), EnterStateStr));
+
+			Exec.Stop();
+		}
+
+		// With wrong overrides
+		{
+			EStateTreeRunStatus Status = EStateTreeRunStatus::Unset;
+			FStateTreeInstanceData InstanceData;
+
+			FStateTreeReferenceOverrides Overrides;
+			FStateTreeReference OverrideRef3;
+			OverrideRef3.SetStateTree(&StateTree3);
+			Overrides.AddOverride(Tag, OverrideRef3);
+			FStateTreeReference OverrideRef4;
+			OverrideRef4.SetStateTree(StateTree4);
+			Overrides.AddOverride(Tag2, OverrideRef4);
+
+			FTestStateTreeExecutionContext Exec(StateTree, StateTree, InstanceData);
+
+			GetTestRunner().AddExpectedMessage("their schemas don't match", ELogVerbosity::Error, EAutomationExpectedMessageFlags::Contains, 1, false);
+			Exec.SetLinkedStateTreeOverrides(Overrides);
+			AITEST_TRUE("Start should complete with Running", GetTestRunner().HasMetExpectedErrors());
+
+			const bool bInitSucceeded = Exec.IsValid();
+			AITEST_TRUE("StateTree should init", bInitSucceeded);
+
+			Status = Exec.Start();
+			AITEST_EQUAL("Start should complete with Running", Status, EStateTreeRunStatus::Running);
+			AITEST_TRUE("StateTree should enter TaskRoot2", Exec.Expect(TaskRoot2.GetName(), EnterStateStr));
+			AITEST_FALSE("StateTree should not enter TaskRoot3", Exec.Expect(TaskRoot3.GetName(), EnterStateStr));
 
 			Exec.Stop();
 		}
diff --git a/Engine/Plugins/Runtime/StateTree/Source/StateTreeTestSuite/Private/StateTreeTest.h b/Engine/Plugins/Runtime/StateTree/Source/StateTreeTestSuite/Private/StateTreeTest.h
--- a/Engine/Plugins/Runtime/StateTree/Source/StateTreeTestSuite/Private/StateTreeTest.h
+++ b/Engine/Plugins/Runtime/StateTree/Source/StateTreeTestSuite/Private/StateTreeTest.h
@@ -14,3 +14,9 @@ class UStateTreeTestSchema : public UStateTreeSchema
 	virtual bool IsExternalItemAllowed(const UStruct& InStruct) const override { return true; }
 };
 
+UCLASS(HideDropdown)
+class UStateTreeTestSchema2 : public UStateTreeSchema
+{
+	GENERATED_BODY()
+};
+
