StateTree: Add a BP native make node for StateTreeReference. Add a k2 node to set the properties in the state tree reference. The default values are not accepted for the property pins. That detects whether the user wants to override the value or keep the state tree default parameters. The default cannot stay in sync with the state tree asset.

diff --git a/Engine/Plugins/Runtime/StateTree/Source/StateTreeEditorModule/Private/K2Node_StateTreeReference.cpp b/Engine/Plugins/Runtime/StateTree/Source/StateTreeEditorModule/Private/K2Node_StateTreeReference.cpp
new file mode 100644
--- /dev/null
+++ b/Engine/Plugins/Runtime/StateTree/Source/StateTreeEditorModule/Private/K2Node_StateTreeReference.cpp
@@ -0,0 +1,417 @@
+// Copyright Epic Games, Inc. All Rights Reserved.
+
+#include "K2Node_StateTreeReference.h"
+
+#include "BlueprintActionDatabaseRegistrar.h"
+#include "BlueprintNodeSpawner.h"
+#include "EditorCategoryUtils.h"
+#include "EdGraph/EdGraphPin.h"
+#include "EdGraphSchema_K2.h"
+#include "K2Node_Variable.h"
+#include "KismetCompiler.h"
+#include "KismetCompilerMisc.h"
+#include "Kismet2/BlueprintEditorUtils.h"
+#include "StateTree.h"
+#include "StateTreeDelegates.h"
+#include "StateTreeFunctionLibrary.h"
+#include "StateTreeReference.h"
+#include "StructUtils/PropertyBag.h"
+
+#include "KismetNodes/SGraphNodeK2Default.h"
+#include "SGraphPin.h"
+#include "Styling/AppStyle.h"
+
+#define LOCTEXT_NAMESPACE "K2Node_StateTreeReference"
+
+namespace UE::StateTreeEditor::Private
+{
+	static FLazyName StateTreePinName = "BA2CE32D97D46A3A524AC510A794C3C";
+
+	bool IsPropertyPin(const UEdGraphPin* Pin)
+	{
+		if (Pin->ParentPin)
+		{
+			return IsPropertyPin(Pin->ParentPin);
+		}
+
+		return Pin->Direction == EEdGraphPinDirection::EGPD_Input
+			&& Pin->PinType.PinCategory != UEdGraphSchema_K2::PC_Exec
+			&& Pin->PinName != StateTreePinName;
+	}
+
+	bool CanUseProperty(const UEdGraphPin* Pin)
+	{
+		// A property needs to be linked to be considered.
+		//The default value won't matches the value in the state tree asset.
+		return !Pin->bOrphanedPin
+			&& Pin->ParentPin == nullptr
+			&& Pin->LinkedTo.Num() > 0;
+	}
+
+	bool DoRenamedPinsMatch(UStateTree* StateTree, const UEdGraphPin* NewPin, const UEdGraphPin* OldPin)
+	{
+		if (StateTree == nullptr
+			|| NewPin == nullptr
+			|| OldPin == nullptr
+			|| OldPin->Direction != NewPin->Direction)
+		{
+			return false;
+		}
+
+		const FInstancedPropertyBag& Parameters = StateTree->GetDefaultParameters();
+		if (!Parameters.IsValid())
+		{
+			return false;
+		}
+
+		const bool bCompatible = GetDefault<UEdGraphSchema_K2>()->ArePinTypesCompatible(NewPin->PinType, OldPin->PinType);
+		if (!bCompatible)
+		{
+			return false;
+		}
+
+		UScriptStruct* Struct = const_cast<UScriptStruct*>(Parameters.GetValue().GetScriptStruct());
+		return UK2Node_Variable::DoesRenamedVariableMatch(OldPin->PinName, NewPin->PinName, Struct);
+	}
+}
+
+UK2Node_MakeStateTreeReference::UK2Node_MakeStateTreeReference()
+{
+	if (!HasAnyFlags(RF_ClassDefaultObject))
+	{
+		ParametersChangedHandle = UE::StateTree::Delegates::OnPostCompile.AddUObject(this, &UK2Node_MakeStateTreeReference::HandleStateTreeCompiled);
+	}
+}
+
+void UK2Node_MakeStateTreeReference::BeginDestroy()
+{
+	UE::StateTree::Delegates::OnPostCompile.Remove(ParametersChangedHandle);
+	Super::BeginDestroy();
+}
+
+FText UK2Node_MakeStateTreeReference::GetNodeTitle(ENodeTitleType::Type TitleType) const
+{
+	auto MakeNodeTitle = []()
+		{
+			FFormatNamedArguments Args;
+			Args.Add(TEXT("StructName"), FStateTreeReference::StaticStruct()->GetDisplayNameText());
+			return FText::Format(LOCTEXT("MakeNodeTitle", "Make {StructName}"), Args);
+		};
+	static FText LocalCachedNodeTitle = MakeNodeTitle();
+	return LocalCachedNodeTitle;
+}
+
+FText UK2Node_MakeStateTreeReference::GetTooltipText() const
+{
+	auto MakeNodeTooltip = []()
+		{
+			FFormatNamedArguments Args;
+			Args.Add(TEXT("StructName"), FStateTreeReference::StaticStruct()->GetDisplayNameText());
+			return FText::Format(LOCTEXT("MakeNodeTitle", "Adds a node that create a  {StructName} from its members"), Args);
+		};
+	static FText LocalCachedTooltip = MakeNodeTooltip();
+	return LocalCachedTooltip;
+}
+
+FText UK2Node_MakeStateTreeReference::GetMenuCategory() const
+{
+	static FText LocalCachedCategory = FEditorCategoryUtils::GetCommonCategory(FCommonEditorCategory::Struct);
+	return LocalCachedCategory;
+}
+
+FSlateIcon UK2Node_MakeStateTreeReference::GetIconAndTint(FLinearColor& OutColor) const
+{
+	static FSlateIcon LocalCachedIcon(FAppStyle::GetAppStyleSetName(), "GraphEditor.MakeStruct_16x");
+	return LocalCachedIcon;
+}
+
+FLinearColor UK2Node_MakeStateTreeReference::GetNodeTitleColor() const
+{
+	auto MakeTitleColor = []() -> FLinearColor
+		{
+			const UEdGraphSchema_K2* K2Schema = GetDefault<UEdGraphSchema_K2>();
+			FEdGraphPinType PinType;
+			PinType.PinCategory = UEdGraphSchema_K2::PC_Struct;
+			PinType.PinSubCategoryObject = FStateTreeReference::StaticStruct();
+			return K2Schema->GetPinTypeColor(PinType);
+		};
+
+	static FLinearColor LocalCachedTitleColor = MakeTitleColor();
+	return LocalCachedTitleColor;
+}
+
+void UK2Node_MakeStateTreeReference::AllocateDefaultPins()
+{
+	Super::AllocateDefaultPins();
+
+	CreatePin(EGPD_Input, UEdGraphSchema_K2::PC_Exec, UEdGraphSchema_K2::PN_Execute);
+	CreatePin(EGPD_Output, UEdGraphSchema_K2::PC_Exec, UEdGraphSchema_K2::PN_Then);
+	UEdGraphPin* ReturnValuePin = CreatePin(EGPD_Output, UEdGraphSchema_K2::PC_Struct, FStateTreeReference::StaticStruct(), UEdGraphSchema_K2::PN_ReturnValue);
+	ReturnValuePin->PinFriendlyName = FStateTreeReference::StaticStruct()->GetDisplayNameText();
+
+	UEdGraphPin* StateTreePin = CreatePin(EGPD_Input, UEdGraphSchema_K2::PC_Object, UStateTree::StaticClass(), UE::StateTreeEditor::Private::StateTreePinName);
+	StateTreePin->bNotConnectable = true;
+	StateTreePin->PinFriendlyName = LOCTEXT("StateTreePinName", "State Tree");
+
+	CreatePropertyPins();
+}
+
+void UK2Node_MakeStateTreeReference::CreatePropertyPins()
+{
+	if (StateTree)
+	{
+		const FInstancedPropertyBag& Parameters = StateTree->GetDefaultParameters();
+		if (Parameters.IsValid())
+		{
+			UScriptStruct* Struct = const_cast<UScriptStruct*>(Parameters.GetValue().GetScriptStruct());
+			FOptionalPinManager OptionalPinManager;
+			OptionalPinManager.RebuildPropertyList(ShowPinForProperties, Struct);
+			OptionalPinManager.CreateVisiblePins(ShowPinForProperties, Struct, EGPD_Input, this);
+
+			for (UEdGraphPin* Pin : Pins)
+			{
+				if (UE::StateTreeEditor::Private::IsPropertyPin(Pin))
+				{
+					// Force the property to be linked until we have the enabled/disabled on the default value.
+					Pin->bDefaultValueIsIgnored = true;
+				}
+			}
+		}
+	}
+}
+
+UK2Node::ERedirectType UK2Node_MakeStateTreeReference::DoPinsMatchForReconstruction(const UEdGraphPin* NewPin, int32 NewPinIndex, const UEdGraphPin* OldPin, int32 OldPinIndex) const
+{
+	ERedirectType Result = Super::DoPinsMatchForReconstruction(NewPin, NewPinIndex, OldPin, OldPinIndex);
+
+	if (Result == ERedirectType_None
+		&& UE::StateTreeEditor::Private::DoRenamedPinsMatch(StateTree, NewPin, OldPin))
+	{
+		Result = ERedirectType_Name;
+	}
+	return Result;
+}
+
+void UK2Node_MakeStateTreeReference::PinDefaultValueChanged(UEdGraphPin* Pin)
+{
+	Super::PinDefaultValueChanged(Pin);
+
+	if (Pin == FindPinChecked(UE::StateTreeEditor::Private::StateTreePinName))
+	{
+		SetStateTree(GetStateTreeDefaultValue());
+	}
+}
+
+namespace UE::StateTreeEditor::Private
+{
+	class SMakeStateTreeRefeerenceNode : public SGraphNodeK2Default
+	{
+		TSharedPtr<SGraphPin> CreatePinWidget(UEdGraphPin* Pin) const
+		{
+			TSharedPtr<SGraphPin> Result = SGraphNodeK2Default::CreatePinWidget(Pin);
+			if (Pin->PinName == UE::StateTreeEditor::Private::StateTreePinName)
+			{
+				Result->GetPinImageWidget()->SetVisibility(EVisibility::Hidden);
+			}
+			return Result;
+		}
+	};
+}
+
+TSharedPtr<SGraphNode> UK2Node_MakeStateTreeReference::CreateVisualWidget()
+{
+	return SNew(UE::StateTreeEditor::Private::SMakeStateTreeRefeerenceNode, this);
+}
+
+void UK2Node_MakeStateTreeReference::PreloadRequiredAssets()
+{
+	PreloadObject(FStateTreeReference::StaticStruct());
+	if (StateTree)
+	{
+		PreloadObject(StateTree);
+	}
+
+	Super::PreloadRequiredAssets();
+}
+
+void UK2Node_MakeStateTreeReference::ValidateNodeDuringCompilation(class FCompilerResultsLog& MessageLog) const
+{
+	Super::ValidateNodeDuringCompilation(MessageLog);
+
+	bool bTestExistingProperties = true;
+	if (StateTree)
+	{
+		// Tests if the property pins are valid
+		const FInstancedPropertyBag& Parameters = StateTree->GetDefaultParameters();
+		if (Parameters.IsValid())
+		{
+			bTestExistingProperties = false;
+			const UScriptStruct* Struct = Parameters.GetValue().GetScriptStruct();
+			for (UEdGraphPin* Pin : Pins)
+			{
+				if (UE::StateTreeEditor::Private::IsPropertyPin(Pin)
+					&& UE::StateTreeEditor::Private::CanUseProperty(Pin))
+				{
+					const FProperty* Property = Struct->FindPropertyByName(Pin->PinName);
+					const FPropertyBagPropertyDesc* PropertyDesc = Parameters.FindPropertyDescByName(Pin->PinName);
+					if (Property == nullptr || PropertyDesc == nullptr)
+					{
+						FFormatNamedArguments Args;
+						Args.Add(TEXT("PropertyName"), FText::FromName(Pin->PinName));
+						const FText Error = FText::Format(LOCTEXT("CanNotFindProeprty_Error", "Can't find the property {PropertyName} in @@"), Args);
+						MessageLog.Error(*Error.ToString(), this);
+					}
+				}
+			}
+		}
+
+		// Tests if the cached value matches the value of the pin. It should match unless it was set manually by code.
+		{
+			UEdGraphPin* ThisStateTreePin = FindPinChecked(UE::StateTreeEditor::Private::StateTreePinName);
+			if (StateTree != ThisStateTreePin->DefaultObject)
+			{
+				MessageLog.Error(*LOCTEXT("StateTreeMatchingError", "The State Tree asset does not match with the pin @@. Clear and set the State Tree pin.").ToString(), this);
+			}
+		}
+	}
+	
+	// Tests if we expect a state tree (it is valid to construct an empty struct)
+	if (bTestExistingProperties)
+	{
+		const bool bHasProperty = Pins.ContainsByPredicate([](const UEdGraphPin* Pin)
+			{
+				return UE::StateTreeEditor::Private::IsPropertyPin(Pin)
+					&& UE::StateTreeEditor::Private::CanUseProperty(Pin);
+			});
+		if (bHasProperty)
+		{
+			MessageLog.Error(*LOCTEXT("NoStateTree_Error", "No State Tree in @@").ToString(), this);
+		}
+	}
+}
+
+void UK2Node_MakeStateTreeReference::GetMenuActions(FBlueprintActionDatabaseRegistrar& ActionRegistrar) const
+{
+	UClass* Action = GetClass();
+	if (ActionRegistrar.IsOpenForRegistration(Action))
+	{
+		UBlueprintNodeSpawner* MakeNodeSpawner = UBlueprintNodeSpawner::Create(GetClass());
+		check(MakeNodeSpawner != nullptr);
+		ActionRegistrar.AddBlueprintAction(Action, MakeNodeSpawner);
+	}
+}
+
+void UK2Node_MakeStateTreeReference::SetStateTree(UStateTree* InStateTree)
+{
+	StateTree = InStateTree;
+	FBlueprintEditorUtils::MarkBlueprintAsModified(GetBlueprint());
+	ReconstructNode();
+}
+
+UStateTree* UK2Node_MakeStateTreeReference::GetStateTreeDefaultValue() const
+{
+	return Cast<UStateTree>(FindPinChecked(UE::StateTreeEditor::Private::StateTreePinName)->DefaultObject);
+}
+
+void UK2Node_MakeStateTreeReference::HandleStateTreeCompiled(const UStateTree& InStateTree)
+{
+	if (&InStateTree == StateTree)
+	{
+		SetStateTree(StateTree);
+	}
+}
+
+void UK2Node_MakeStateTreeReference::ExpandNode(FKismetCompilerContext& CompilerContext, UEdGraph* SourceGraph)
+{
+	Super::ExpandNode(CompilerContext, SourceGraph);
+
+	if (CompilerContext.bIsFullCompile)
+	{
+		const UEdGraphSchema_K2* K2Schema = CompilerContext.GetSchema();
+
+		// Convert to
+		//local = MakeStateTreeReference(StateTree)
+		//for each properties
+		//  K2_SetParametersProperty(local, id, value)
+		UEdGraphPin* LastThen = nullptr;
+		UEdGraphPin* MakeStateTreeReferenceNodeResultPin = nullptr;
+		{
+			UK2Node_CallFunction* MakeStateTreeReferenceNode = CompilerContext.SpawnIntermediateNode<UK2Node_CallFunction>(this, SourceGraph);
+			UFunction* Function = UStateTreeFunctionLibrary::StaticClass()->FindFunctionByName(GET_FUNCTION_NAME_CHECKED(UStateTreeFunctionLibrary, MakeStateTreeReference));
+			MakeStateTreeReferenceNode->SetFromFunction(Function);
+			MakeStateTreeReferenceNode->AllocateDefaultPins();
+			CompilerContext.MessageLog.NotifyIntermediateObjectCreation(MakeStateTreeReferenceNode, SourceGraph);
+			{
+				UEdGraphPin* ThisStateTreePin = FindPinChecked(UE::StateTreeEditor::Private::StateTreePinName);
+				UEdGraphPin* NewStateTreePin = MakeStateTreeReferenceNode->FindPinChecked(FName("StateTree"));
+				CompilerContext.MovePinLinksToIntermediate(*ThisStateTreePin, *NewStateTreePin);
+			}
+			{
+				UEdGraphPin* ThisResultPin = FindPinChecked(UEdGraphSchema_K2::PN_ReturnValue);
+				MakeStateTreeReferenceNodeResultPin = MakeStateTreeReferenceNode->FindPinChecked(UEdGraphSchema_K2::PN_ReturnValue);
+				CompilerContext.MovePinLinksToIntermediate(*ThisResultPin, *MakeStateTreeReferenceNodeResultPin);
+			}
+			{
+				UEdGraphPin* ThisExecPin = GetExecPin();
+				UEdGraphPin* NewExecPin = MakeStateTreeReferenceNode->GetExecPin();
+				CompilerContext.MovePinLinksToIntermediate(*ThisExecPin, *NewExecPin);
+			}
+			{
+				UEdGraphPin* ThisThenPin = GetThenPin();
+				UEdGraphPin* NewThenPin = MakeStateTreeReferenceNode->GetThenPin();
+				CompilerContext.MovePinLinksToIntermediate(*ThisThenPin, *NewThenPin);
+				LastThen = NewThenPin;
+			}
+		}
+
+		if (StateTree != nullptr)
+		{
+			//for each pin call K2_SetParametersProperty
+			for (UEdGraphPin* Pin : Pins)
+			{
+				if (UE::StateTreeEditor::Private::IsPropertyPin(Pin)
+					&& UE::StateTreeEditor::Private::CanUseProperty(Pin))
+				{
+						UK2Node_CallFunction* SetParametersPropertyNode = CompilerContext.SpawnIntermediateNode<UK2Node_CallFunction>(this, SourceGraph);
+						UFunction* Function = UStateTreeFunctionLibrary::StaticClass()->FindFunctionByName(GET_FUNCTION_NAME_CHECKED(UStateTreeFunctionLibrary, K2_SetParametersProperty));
+						SetParametersPropertyNode->SetFromFunction(Function);
+						SetParametersPropertyNode->AllocateDefaultPins();
+						CompilerContext.MessageLog.NotifyIntermediateObjectCreation(SetParametersPropertyNode, SourceGraph);
+						
+						{
+							UEdGraphPin* NewValuePin = SetParametersPropertyNode->FindPinChecked(FName("Reference"), EGPD_Input);
+							ensure(K2Schema->TryCreateConnection(MakeStateTreeReferenceNodeResultPin, NewValuePin));
+						}
+						{
+							UEdGraphPin* NewValuePin = SetParametersPropertyNode->FindPinChecked(FName("PropertyID"), EGPD_Input);
+							
+							const FPropertyBagPropertyDesc* PropertyDesc = StateTree->GetDefaultParameters().FindPropertyDescByName(Pin->PinName);
+							check(PropertyDesc);
+
+							FGuid Default;
+							FGuid TempValue = PropertyDesc->ID;
+							TBaseStructure<FGuid>::Get()->ExportText(NewValuePin->DefaultValue, &TempValue, &Default, nullptr, PPF_None, nullptr);
+						}
+						{
+							UEdGraphPin* NewValuePin = SetParametersPropertyNode->FindPinChecked(FName("NewValue"), EGPD_Input);
+							NewValuePin->PinType = Pin->PinType;
+							CompilerContext.MovePinLinksToIntermediate(*Pin, *NewValuePin);
+						}
+						// move last Then to new Then and link the last Then to new exec
+						{
+							UEdGraphPin* NewThenPin = SetParametersPropertyNode->GetThenPin();
+							CompilerContext.MovePinLinksToIntermediate(*LastThen, *NewThenPin);
+						}
+						{
+							UEdGraphPin* NewExecPin = SetParametersPropertyNode->GetExecPin();
+							ensure(K2Schema->TryCreateConnection(LastThen, NewExecPin));
+						}
+						LastThen = SetParametersPropertyNode->GetThenPin();
+				}
+			}
+		}
+	}
+	BreakAllNodeLinks();
+}
+
+#undef LOCTEXT_NAMESPACE
\ No newline at end of file
diff --git a/Engine/Plugins/Runtime/StateTree/Source/StateTreeEditorModule/Private/K2Node_StateTreeReference.h b/Engine/Plugins/Runtime/StateTree/Source/StateTreeEditorModule/Private/K2Node_StateTreeReference.h
new file mode 100644
--- /dev/null
+++ b/Engine/Plugins/Runtime/StateTree/Source/StateTreeEditorModule/Private/K2Node_StateTreeReference.h
@@ -0,0 +1,73 @@
+// Copyright Epic Games, Inc. All Rights Reserved.
+
+#pragma once
+
+#include "K2Node_CallFunction.h"
+#include "K2Node_StateTreeReference.generated.h"
+
+class FKismetCompilerContext;
+class SGraphNode;
+class UStateTree;
+
+UCLASS()
+class UK2Node_MakeStateTreeReference : public UK2Node
+{
+	GENERATED_BODY()
+
+public:
+	UK2Node_MakeStateTreeReference();
+	virtual void BeginDestroy() override;
+
+	virtual FText GetNodeTitle(ENodeTitleType::Type TitleType) const override;
+	virtual FText GetTooltipText() const override;
+	virtual FText GetMenuCategory() const override;
+	virtual FSlateIcon GetIconAndTint(FLinearColor& OutColor) const override;
+	virtual FLinearColor GetNodeTitleColor() const;
+
+	virtual void AllocateDefaultPins() override;
+	virtual ERedirectType DoPinsMatchForReconstruction(const UEdGraphPin* NewPin, int32 NewPinIndex, const UEdGraphPin* OldPin, int32 OldPinIndex) const override;
+	virtual void PinDefaultValueChanged(UEdGraphPin* Pin) override;
+	virtual TSharedPtr<SGraphNode> CreateVisualWidget() override;
+
+	virtual void PreloadRequiredAssets() override;
+	virtual void ValidateNodeDuringCompilation(FCompilerResultsLog& MessageLog) const override;
+	virtual void GetMenuActions(FBlueprintActionDatabaseRegistrar& ActionRegistrar) const override;
+
+	virtual void ExpandNode(FKismetCompilerContext& CompilerContext, UEdGraph* SourceGraph) override;
+
+	virtual bool NodeCausesStructuralBlueprintChange() const override
+	{
+		return false;
+	}
+
+	virtual bool IsNodePure() const override
+	{
+		return false;
+	}
+
+	virtual bool DrawNodeAsVariable() const override
+	{
+		return false;
+	}
+
+	virtual bool CanSplitPin(const UEdGraphPin* Pin) const override
+	{
+		return false;
+	}
+
+private:
+	void CreatePropertyPins();
+	UStateTree* GetStateTreeDefaultValue() const;
+	void SetStateTree(UStateTree* InStateTree);
+	void HandleStateTreeCompiled(const UStateTree& StateTree);
+
+private:
+	/** Created pins from the state tree properties */
+	UPROPERTY()
+	TArray<FOptionalPinFromProperty> ShowPinForProperties;
+	/** State tree asset set in the pin and saved here to rebuild the property pins.*/
+	UPROPERTY()
+	TObjectPtr<UStateTree> StateTree;
+	
+	FDelegateHandle ParametersChangedHandle;
+};
\ No newline at end of file
diff --git a/Engine/Plugins/Runtime/StateTree/Source/StateTreeModule/Private/StateTreeFunctionLibrary.cpp b/Engine/Plugins/Runtime/StateTree/Source/StateTreeModule/Private/StateTreeFunctionLibrary.cpp
new file mode 100644
--- /dev/null
+++ b/Engine/Plugins/Runtime/StateTree/Source/StateTreeModule/Private/StateTreeFunctionLibrary.cpp
@@ -0,0 +1,130 @@
+// Copyright Epic Games, Inc. All Rights Reserved.
+
+#include "StateTreeFunctionLibrary.h"
+
+#include "Blueprint/BlueprintExceptionInfo.h"
+#include "StateTree.h"
+#include "StateTreeReference.h"
+#include "StructUtils/PropertyBag.h"
+#include "UObject/Script.h"
+
+#define LOCTEXT_NAMESPACE "StateTreeFunctionLibrary"
+
+void UStateTreeFunctionLibrary::SetStateTree(FStateTreeReference& Reference, UStateTree* NewStateTree)
+{
+	Reference.SetStateTree(NewStateTree);
+}
+
+FStateTreeReference UStateTreeFunctionLibrary::MakeStateTreeReference(UStateTree* NewStateTree)
+{
+	FStateTreeReference Result;
+	Result.SetStateTree(NewStateTree);
+	return Result;
+}
+
+void UStateTreeFunctionLibrary::K2_SetParametersProperty(FStateTreeReference&, FGuid, const int32&)
+{
+	checkNoEntry();
+}
+
+void UStateTreeFunctionLibrary::K2_GetParametersProperty(const FStateTreeReference&, FGuid, int32&)
+{
+	checkNoEntry();
+}
+
+DEFINE_FUNCTION(UStateTreeFunctionLibrary::execK2_SetParametersProperty)
+{
+	// Read wildcard Value input.
+	Stack.MostRecentPropertyAddress = nullptr;
+	Stack.MostRecentPropertyContainer = nullptr;
+
+	P_GET_STRUCT_REF(FStateTreeReference, StateTreeReference);
+	P_GET_STRUCT(FGuid, PropertyID);
+
+	Stack.StepCompiledIn<FProperty>(nullptr);
+	const FProperty* SourceProperty = Stack.MostRecentProperty;
+	const uint8* SourcePtr = Stack.MostRecentPropertyAddress;
+	P_FINISH;
+
+	if (SourceProperty == nullptr|| SourcePtr == nullptr)
+	{
+		FBlueprintExceptionInfo ExceptionInfo(
+			EBlueprintExceptionType::AbortExecution,
+			LOCTEXT("GetParametersProperty_InvalidValueWarning", "Failed to resolve the Value for SetParametersProperty")
+		);
+
+		FBlueprintCoreDelegates::ThrowScriptException(P_THIS, Stack, ExceptionInfo);
+	}
+	else
+	{
+		P_NATIVE_BEGIN;
+		FInstancedPropertyBag& InstancedPropertyBag = StateTreeReference.GetMutableParameters();
+		FStructView PropertyBagView = InstancedPropertyBag.GetMutableValue();
+		const UPropertyBag* PropertyBag = InstancedPropertyBag.GetPropertyBagStruct();
+		if (PropertyBagView.IsValid() && PropertyBag)
+		{
+			if (const FPropertyBagPropertyDesc* PropertyBagDesc = PropertyBag->FindPropertyDescByID(PropertyID))
+			{
+				if (const FProperty* TargetProperty = PropertyBag->FindPropertyByName(PropertyBagDesc->Name))
+				{
+					if (SourceProperty->SameType(TargetProperty))
+					{
+						void* TargetPtr = TargetProperty->ContainerPtrToValuePtr<void>(PropertyBagView.GetMemory());
+						TargetProperty->CopyCompleteValue(TargetPtr, SourcePtr);
+						StateTreeReference.SetPropertyOverridden(PropertyID, true);
+					}
+				}
+			}
+		}
+		P_NATIVE_END;
+	}
+}
+
+DEFINE_FUNCTION(UStateTreeFunctionLibrary::execK2_GetParametersProperty)
+{
+	P_GET_STRUCT_REF(FStateTreeReference, StateTreeReference);
+	P_GET_STRUCT(FGuid, PropertyID);
+
+	// Read wildcard Value input.
+	Stack.MostRecentPropertyAddress = nullptr;
+	Stack.MostRecentPropertyContainer = nullptr;
+	Stack.StepCompiledIn<FProperty>(nullptr);
+
+	const FProperty* TargetProperty = Stack.MostRecentProperty;
+	void* TargetPtr = Stack.MostRecentPropertyAddress;
+
+	P_FINISH;
+
+	if (TargetProperty == nullptr || TargetPtr == nullptr)
+	{
+		FBlueprintExceptionInfo ExceptionInfo(
+			EBlueprintExceptionType::AbortExecution,
+			LOCTEXT("GetParametersProperty_InvalidValueWarning", "Failed to resolve the Value for GetParametersProperty")
+		);
+		FBlueprintCoreDelegates::ThrowScriptException(P_THIS, Stack, ExceptionInfo);
+	}
+	else
+	{
+		P_NATIVE_BEGIN;
+		const FInstancedPropertyBag& InstancedPropertyBag = StateTreeReference.GetParameters();
+		const FConstStructView PropertyBagView = InstancedPropertyBag.GetValue();
+		const UPropertyBag* PropertyBag = InstancedPropertyBag.GetPropertyBagStruct();
+		if (PropertyBagView.IsValid() && PropertyBag)
+		{
+			if (const FPropertyBagPropertyDesc* PropertyBagDesc = PropertyBag->FindPropertyDescByID(PropertyID))
+			{
+				if (const FProperty* SourceProperty = PropertyBag->FindPropertyByName(PropertyBagDesc->Name))
+				{
+					if (SourceProperty->SameType(TargetProperty))
+					{
+						const void* SourcePtr = SourceProperty->ContainerPtrToValuePtr<void>(PropertyBagView.GetMemory());
+						TargetProperty->CopyCompleteValue(TargetPtr, SourcePtr);
+					}
+				}
+			}
+		}
+		P_NATIVE_END;
+	}
+}
+
+#undef LOCTEXT_NAMESPACE
diff --git a/Engine/Plugins/Runtime/StateTree/Source/StateTreeModule/Public/StateTreeFunctionLibrary.h b/Engine/Plugins/Runtime/StateTree/Source/StateTreeModule/Public/StateTreeFunctionLibrary.h
new file mode 100644
--- /dev/null
+++ b/Engine/Plugins/Runtime/StateTree/Source/StateTreeModule/Public/StateTreeFunctionLibrary.h
@@ -0,0 +1,34 @@
+// Copyright Epic Games, Inc. All Rights Reserved.
+#pragma once
+
+#include "Kismet/BlueprintFunctionLibrary.h"
+#include "StateTreeFunctionLibrary.generated.h"
+
+struct FStateTreeReference;
+class UStateTree;
+
+/**
+ * A collection of blueprint functions for state tree.
+ */
+UCLASS(Abstract)
+class STATETREEMODULE_API UStateTreeFunctionLibrary : public UBlueprintFunctionLibrary
+{
+	GENERATED_BODY()
+
+public:
+	UFUNCTION(BlueprintCallable, Category = "StateTree|Reference")
+	static void SetStateTree(UPARAM(ref)FStateTreeReference& Reference, UStateTree* StateTree);
+
+	UFUNCTION(BlueprintCallable, DisplayName = "Make State Tree Reference", Category = "StateTree|Reference", meta = (Keywords = "construct build", NativeMakeFunc, BlueprintInternalUseOnly = "true"))
+	static FStateTreeReference MakeStateTreeReference(UStateTree* StateTree);
+
+	UFUNCTION(BlueprintCallable, CustomThunk, DisplayName = "Set Parameter Property", Category = "StateTree|Reference", meta = (BluerpintInternalUseOnly = "true", CustomStructureParam = "NewValue"))
+	static void K2_SetParametersProperty(UPARAM(ref)FStateTreeReference& Reference, FGuid PropertyID, UPARAM(ref) const int32& NewValue);
+
+	UFUNCTION(BlueprintCallable, CustomThunk, DisplayName = "Get Parameter Property", Category = "StateTree|Reference", meta = (BluerpintInternalUseOnly = "true", CustomStructureParam = "ReturnValue"))
+	static void K2_GetParametersProperty(UPARAM(ref)const FStateTreeReference& Reference, FGuid PropertyID, int32& ReturnValue);
+
+private:
+	DECLARE_FUNCTION(execK2_SetParametersProperty);
+	DECLARE_FUNCTION(execK2_GetParametersProperty);
+};
diff --git a/Engine/Plugins/Runtime/StateTree/Source/StateTreeModule/Public/StateTreeReference.h b/Engine/Plugins/Runtime/StateTree/Source/StateTreeModule/Public/StateTreeReference.h
--- a/Engine/Plugins/Runtime/StateTree/Source/StateTreeModule/Public/StateTreeReference.h
+++ b/Engine/Plugins/Runtime/StateTree/Source/StateTreeModule/Public/StateTreeReference.h
@@ -11,7 +11,7 @@ class UStateTree;
 /**
  * Struct to hold reference to a StateTree asset along with values to parameterized it.
  */
-USTRUCT(BlueprintType)
+USTRUCT(BlueprintType, meta = (DisableSplitPin, HasNativeMake = "/Script/StateTreeModule.StateTreeFunctionLibrary.MakeStateTreeReference"))
 struct STATETREEMODULE_API FStateTreeReference
 {
 	GENERATED_BODY()
