StateTree: Fix compilation issue with shadow variable. Introduced with CL 39893788. Change the EStateTreeNodeFormatting for a class enum to add a scope to the variable Text.

diff --git a/Engine/Plugins/Runtime/StateTree/Source/StateTreeEditorModule/Private/SStateTreeViewRow.h b/Engine/Plugins/Runtime/StateTree/Source/StateTreeEditorModule/Private/SStateTreeViewRow.h
--- a/Engine/Plugins/Runtime/StateTree/Source/StateTreeEditorModule/Private/SStateTreeViewRow.h
+++ b/Engine/Plugins/Runtime/StateTree/Source/StateTreeEditorModule/Private/SStateTreeViewRow.h
@@ -8,7 +8,7 @@
 class STableViewBase;
 class UStateTreeState;
 enum class EStateTreeTransitionTrigger : uint8;
-enum EStateTreeNodeFormatting : uint8;
+enum class EStateTreeNodeFormatting : uint8;
 struct FStateTreeStateLink;
 struct FStateTreeEditorNode;
 
diff --git a/Engine/Plugins/Runtime/StateTree/Source/StateTreeEditorModule/Private/StateTreeDescriptionHelpers.h b/Engine/Plugins/Runtime/StateTree/Source/StateTreeEditorModule/Private/StateTreeDescriptionHelpers.h
--- a/Engine/Plugins/Runtime/StateTree/Source/StateTreeEditorModule/Private/StateTreeDescriptionHelpers.h
+++ b/Engine/Plugins/Runtime/StateTree/Source/StateTreeEditorModule/Private/StateTreeDescriptionHelpers.h
@@ -1,4 +1,4 @@
-﻿// Copyright Epic Games, Inc. All Rights Reserved.
+// Copyright Epic Games, Inc. All Rights Reserved.
 
 #pragma once
 
@@ -7,7 +7,7 @@
 class UStateTreeEditorData;
 struct FStateTreeStateLink;
 struct FStateTreeTransition;
-enum EStateTreeNodeFormatting : uint8;
+enum class EStateTreeNodeFormatting : uint8;
 struct FSlateBrush;
 struct FSlateColor;
 class FText;
diff --git a/Engine/Plugins/Runtime/StateTree/Source/StateTreeEditorModule/Public/StateTreeEditorPropertyBindings.h b/Engine/Plugins/Runtime/StateTree/Source/StateTreeEditorModule/Public/StateTreeEditorPropertyBindings.h
--- a/Engine/Plugins/Runtime/StateTree/Source/StateTreeEditorModule/Public/StateTreeEditorPropertyBindings.h
+++ b/Engine/Plugins/Runtime/StateTree/Source/StateTreeEditorModule/Public/StateTreeEditorPropertyBindings.h
@@ -9,7 +9,7 @@
 #include "StateTreeEditorNode.h"
 #include "StateTreeEditorPropertyBindings.generated.h"
 
-enum EStateTreeNodeFormatting : uint8;
+enum class EStateTreeNodeFormatting : uint8;
 class IStateTreeEditorPropertyBindingsOwner;
 enum class EStateTreeVisitor : uint8;
 
diff --git a/Engine/Plugins/Runtime/StateTree/Source/StateTreeModule/Public/StateTreeNodeBase.h b/Engine/Plugins/Runtime/StateTree/Source/StateTreeModule/Public/StateTreeNodeBase.h
--- a/Engine/Plugins/Runtime/StateTree/Source/StateTreeModule/Public/StateTreeNodeBase.h
+++ b/Engine/Plugins/Runtime/StateTree/Source/StateTreeModule/Public/StateTreeNodeBase.h
@@ -18,7 +18,7 @@ struct IStateTreeBindingLookup;
  * - Subdued should be generally used for secondary/structural information, like "{Left} <s>equals</> {Right}".
  */
 UENUM(BlueprintType)
-enum EStateTreeNodeFormatting : uint8
+enum class EStateTreeNodeFormatting : uint8
 {
 	/**
 	 * The returned text can contain following right text formatting (no nesting)
diff --git a/Engine/Plugins/Runtime/StateTree/Source/StateTreeModule/Public/StateTreeNodeDescriptionHelpers.h b/Engine/Plugins/Runtime/StateTree/Source/StateTreeModule/Public/StateTreeNodeDescriptionHelpers.h
--- a/Engine/Plugins/Runtime/StateTree/Source/StateTreeModule/Public/StateTreeNodeDescriptionHelpers.h
+++ b/Engine/Plugins/Runtime/StateTree/Source/StateTreeModule/Public/StateTreeNodeDescriptionHelpers.h
@@ -7,7 +7,7 @@
 #include "StateTreeTypes.h"
 #include "StateTreePropertyBindings.h"
 
-enum EStateTreeNodeFormatting : uint8;
+enum class EStateTreeNodeFormatting : uint8;
 struct FGameplayTagContainer;
 struct FGameplayTagQuery;
 
diff --git a/Engine/Plugins/Runtime/StateTree/Source/StateTreeModule/Public/StateTreePropertyBindings.h b/Engine/Plugins/Runtime/StateTree/Source/StateTreeModule/Public/StateTreePropertyBindings.h
--- a/Engine/Plugins/Runtime/StateTree/Source/StateTreeModule/Public/StateTreePropertyBindings.h
+++ b/Engine/Plugins/Runtime/StateTree/Source/StateTreeModule/Public/StateTreePropertyBindings.h
@@ -17,7 +17,7 @@ class FProperty;
 struct FStateTreePropertyBindingCompiler;
 struct FStateTreePropertyRef;
 class UStateTree;
-enum EStateTreeNodeFormatting : uint8;
+enum class EStateTreeNodeFormatting : uint8;
 
 UENUM()
 enum class EStateTreeBindableStructSource : uint8
