[State Tree Editor] fixed nullptr access from skipping the header initialization for Property Function Nodes

diff --git a/Engine/Plugins/Runtime/StateTree/Source/StateTreeEditorModule/Private/Customizations/StateTreeEditorNodeDetails.cpp b/Engine/Plugins/Runtime/StateTree/Source/StateTreeEditorModule/Private/Customizations/StateTreeEditorNodeDetails.cpp
--- a/Engine/Plugins/Runtime/StateTree/Source/StateTreeEditorModule/Private/Customizations/StateTreeEditorNodeDetails.cpp
+++ b/Engine/Plugins/Runtime/StateTree/Source/StateTreeEditorModule/Private/Customizations/StateTreeEditorNodeDetails.cpp
@@ -697,6 +697,11 @@ void FStateTreeEditorNodeDetails::CustomizeHeader(TSharedRef<class IPropertyHand
 
 void FStateTreeEditorNodeDetails::MakeFlagsWidget()
 {
+	if (!FlagsContainer.IsValid())
+	{
+		return;
+	}
+
 	FlagsContainer->SetPadding(FMargin(4.0f));
 	FlagsContainer->SetContent(SNullWidget::NullWidget);
 
