StateTree: Rename FDelayedFinishedTask to FFinishedTask to be more descriptive. Review feedback for CL 39565986

diff --git a/Engine/Plugins/Runtime/StateTree/Source/StateTreeModule/Private/StateTreeExecutionContext.cpp b/Engine/Plugins/Runtime/StateTree/Source/StateTreeModule/Private/StateTreeExecutionContext.cpp
--- a/Engine/Plugins/Runtime/StateTree/Source/StateTreeModule/Private/StateTreeExecutionContext.cpp
+++ b/Engine/Plugins/Runtime/StateTree/Source/StateTreeModule/Private/StateTreeExecutionContext.cpp
@@ -69,7 +69,7 @@ namespace UE::StateTree::ExecutionContext::Private
 		return false;
 	}
 
-	EStateTreeRunStatus FinishTaskTypeToRunStatus(EStateTreeFinishTaskType FinishTask)
+	EStateTreeRunStatus GetRunStatusFromFinishTaskType(EStateTreeFinishTaskType FinishTask)
 	{
 		static_assert((int32)EStateTreeRunStatus::Succeeded == (int32)EStateTreeFinishTaskType::Succeeded, "Succeeded do not match in both types.");
 		static_assert((int32)EStateTreeRunStatus::Failed == (int32)EStateTreeFinishTaskType::Failed, "Failed do not match in both types.");
@@ -230,7 +230,7 @@ FStateTreeScheduledTick FStateTreeMinimalExecutionContext::GetNextScheduledTick(
 	}
 
 	// Completed task. For EnterState or for user that only called TickTasks and not TickTransitions.
-	if (Exec.DelayedFinishedTasks.Num() > 0)
+	if (Exec.FinishedTasks.Num() > 0)
 	{
 		return FStateTreeScheduledTick::MakeNextFrame();
 	}
@@ -1307,18 +1307,18 @@ void FStateTreeExecutionContext::FinishTask(const FStateTreeTaskBase& Task, ESta
 		const int32 ActiveStateIndex = CurrentlyProcessedFrame->ActiveStates.IndexOfReverse(CurrentlyProcessedState);
 		check(ActiveStateIndex != INDEX_NONE);
 		const FActiveStateID ActiveStateID = CurrentlyProcessedFrame->ActiveStates.StateIDs[ActiveStateIndex];
-		FinishTaskInternal(FDelayedFinishedTask::EReasonType::StateTask, ActiveFrameID, ActiveStateID, CurrentNodeIndex, FinishTaskTypeToRunStatus(FinishType));
+		FinishTaskInternal(FFinishedTask::EReasonType::StateTask, ActiveFrameID, ActiveStateID, CurrentNodeIndex, GetRunStatusFromFinishTaskType(FinishType));
 	}
 	else
 	{
-		FinishTaskInternal(FDelayedFinishedTask::EReasonType::GlobalTask, ActiveFrameID, FActiveStateID::Invalid, CurrentNodeIndex, FinishTaskTypeToRunStatus(FinishType));
+		FinishTaskInternal(FFinishedTask::EReasonType::GlobalTask, ActiveFrameID, FActiveStateID::Invalid, CurrentNodeIndex, GetRunStatusFromFinishTaskType(FinishType));
 	}
 
-	// The FStateTreeTaskBase::Tick or OnEnter returns the value. Set it here and test it later.
-	CompletedStateRunStatus = FinishTaskTypeToRunStatus(FinishType);
+	// The FStateTreeTaskBase::Tick or FStateTreeTaskBase::OnEnter can returns a value. Set it here and test it later.
+	CompletedStateRunStatus = GetRunStatusFromFinishTaskType(FinishType);
 }
 
-void FStateTreeExecutionContext::FinishTask(const UE::StateTree::FDelayedFinishedTask& Task, EStateTreeFinishTaskType FinishType)
+void FStateTreeExecutionContext::FinishTask(const UE::StateTree::FFinishedTask& Task, EStateTreeFinishTaskType FinishType)
 {
 	if (!IsValid())
 	{
@@ -1332,36 +1332,36 @@ void FStateTreeExecutionContext::FinishTask(const UE::StateTree::FDelayedFinishe
 		// We are currently ticking that task. We are either in Tick, EnterState, ExitState, Transition.
 		//Note: Completed tasks are not considered during ExitState or Transition.
 		//The node will return the value and it will be processed once returned. Set it here and test it later.
-		CompletedStateRunStatus = UE::StateTree::ExecutionContext::Private:: FinishTaskTypeToRunStatus(FinishType);
+		CompletedStateRunStatus = UE::StateTree::ExecutionContext::Private:: GetRunStatusFromFinishTaskType(FinishType);
 	}
-	else if (IsDelayedFinishedTaskValid(Task))
+	else if (IsFinishedTaskValid(Task))
 	{
 		FStateTreeExecutionState& Exec = GetExecState();
-		const UE::StateTree::FDelayedFinishedTask& DelayedCompletedTask = Exec.DelayedFinishedTasks.Add_GetRef(Task);
-		AddUniqueCompletedStateList(DelayedCompletedTask);
+		const UE::StateTree::FFinishedTask& CompletedTask = Exec.FinishedTasks.Add_GetRef(Task);
+		MarkStateCompleted(CompletedTask);
 		ScheduleNextTick();
 	}
 }
 
-void FStateTreeExecutionContext::FinishTaskInternal(UE::StateTree::FDelayedFinishedTask::EReasonType Reason, UE::StateTree::FActiveFrameID ActiveFrameID, UE::StateTree::FActiveStateID ActiveStateID, int32 TaskIndex, EStateTreeRunStatus RunStatus)
+void FStateTreeExecutionContext::FinishTaskInternal(UE::StateTree::FFinishedTask::EReasonType Reason, UE::StateTree::FActiveFrameID ActiveFrameID, UE::StateTree::FActiveStateID ActiveStateID, int32 TaskIndex, EStateTreeRunStatus RunStatus)
 {
 	using namespace UE::StateTree;
 
 	FStateTreeExecutionState& Exec = GetExecState();
 
-	// Add it for later. In case the user called TickTasks and will call TickTransitions later.
 	const FStateTreeIndex16 TaskIndex16 = TaskIndex >= 0 ? FStateTreeIndex16(TaskIndex) : FStateTreeIndex16::Invalid;
 	const EStateTreeFinishTaskType FinishType = RunStatus == EStateTreeRunStatus::Succeeded ? EStateTreeFinishTaskType::Succeeded : EStateTreeFinishTaskType::Failed;
 
-	const FDelayedFinishedTask Task {.FrameID = ActiveFrameID, .StateID = ActiveStateID, .TaskIndex = TaskIndex16, .Type = FinishType, .Reason = Reason };
-	if (IsDelayedFinishedTaskValid(Task))
+	const FFinishedTask Task {.FrameID = ActiveFrameID, .StateID = ActiveStateID, .TaskIndex = TaskIndex16, .Type = FinishType, .Reason = Reason };
+	if (IsFinishedTaskValid(Task))
 	{
-		Exec.DelayedFinishedTasks.Add(Task);
-		AddUniqueCompletedStateList(Task);
+		// Add it for later. In case the user called TickTasks and will call TickTransitions later.
+		Exec.FinishedTasks.Add(Task);
+		MarkStateCompleted(Task);
 	}
 }
 
-bool FStateTreeExecutionContext::IsDelayedFinishedTaskValid(const UE::StateTree::FDelayedFinishedTask& Task) const
+bool FStateTreeExecutionContext::IsFinishedTaskValid(const UE::StateTree::FFinishedTask& Task) const
 {
 	using namespace UE::StateTree;
 
@@ -1377,7 +1377,7 @@ bool FStateTreeExecutionContext::IsDelayedFinishedTaskValid(const UE::StateTree:
 	const UStateTree* StateTree = ExecutionFrame->StateTree;
 	check(StateTree);
 
-	if (Task.Reason == FDelayedFinishedTask::EReasonType::StateTask)
+	if (Task.Reason == FFinishedTask::EReasonType::StateTask)
 	{
 		if (!ensure(Task.StateID.IsValid() && Task.TaskIndex.IsValid()))
 		{
@@ -1396,10 +1396,10 @@ bool FStateTreeExecutionContext::IsDelayedFinishedTaskValid(const UE::StateTree:
 		const int32 BeginTaskIndex = State.TasksBegin;
 		const int32 EndTaskIndex = State.TasksBegin + State.TasksNum;
 		const bool bIsValidTask = Task.TaskIndex.AsInt32() >= BeginTaskIndex && Task.TaskIndex.AsInt32() < EndTaskIndex;
-		ensureMsgf(bIsValidTask, TEXT("The delayed finished task has a task index that doesn't match the state. Have you recompiled the asset runtime?"));
+		ensureMsgf(bIsValidTask, TEXT("The finished task has a task index that doesn't match the state. Have you recompiled the asset runtime?"));
 		return bIsValidTask;
 	}
-	else if (Task.Reason == FDelayedFinishedTask::EReasonType::GlobalTask)
+	else if (Task.Reason == FFinishedTask::EReasonType::GlobalTask)
 	{
 		if (!ensure(ExecutionFrame->bIsGlobalFrame && Task.TaskIndex.IsValid()))
 		{
@@ -1409,10 +1409,10 @@ bool FStateTreeExecutionContext::IsDelayedFinishedTaskValid(const UE::StateTree:
 		const int32 BeginTaskIndex = StateTree->GlobalTasksBegin;
 		const int32 EndTaskIndex = StateTree->GlobalTasksBegin + StateTree->GlobalTasksNum;
 		const bool bIsValidTask = Task.TaskIndex.AsInt32() >= BeginTaskIndex && Task.TaskIndex.AsInt32() < EndTaskIndex;
-		ensureMsgf(bIsValidTask, TEXT("The delayed finished task has a task index that is not global. Have you recompiled the asset runtime?"));
+		ensureMsgf(bIsValidTask, TEXT("The finished task has a task index that is not global. Have you recompiled the asset runtime?"));
 		return bIsValidTask;
 	}
-	else if (Task.Reason == FDelayedFinishedTask::EReasonType::InternalTransition)
+	else if (Task.Reason == FFinishedTask::EReasonType::InternalTransition)
 	{
 		if (!ensure(Task.StateID.IsValid()))
 		{
@@ -1433,21 +1433,21 @@ void FStateTreeExecutionContext::UpdateCompletedStateList()
 	CompletedStates.Empty();
 
 	FStateTreeExecutionState& Exec = GetExecState();
-	for (int32 Index = Exec.DelayedFinishedTasks.Num() - 1; Index >= 0; --Index)
+	for (int32 Index = Exec.FinishedTasks.Num() - 1; Index >= 0; --Index)
 	{
-		const UE::StateTree::FDelayedFinishedTask& Task = Exec.DelayedFinishedTasks[Index];
-		if (IsDelayedFinishedTaskValid(Task))
+		const UE::StateTree::FFinishedTask& Task = Exec.FinishedTasks[Index];
+		if (IsFinishedTaskValid(Task))
 		{
-			AddUniqueCompletedStateList(Task);
+			MarkStateCompleted(Task);
 		}
 		else
 		{
-			Exec.DelayedFinishedTasks.RemoveAtSwap(Index);
+			Exec.FinishedTasks.RemoveAtSwap(Index);
 		}
 	}
 }
 
-void FStateTreeExecutionContext::AddUniqueCompletedStateList(const UE::StateTree::FDelayedFinishedTask& NewFinishedTask)
+void FStateTreeExecutionContext::MarkStateCompleted(const UE::StateTree::FFinishedTask& NewFinishedTask)
 {
 	bool bAdd = true;
 	// The priority is global tasks and failure. Only keep 1 FinishedTask per Frame per State.
@@ -1462,9 +1462,9 @@ void FStateTreeExecutionContext::AddUniqueCompletedStateList(const UE::StateTree
 			// If the existing one is state and the new one is global, discard the existing and add the new.
 			// If the existing one is state and match the new one, only update the finish type.
 			
-			if (ExistingCompletedState.Reason == UE::StateTree::FDelayedFinishedTask::EReasonType::GlobalTask)
+			if (ExistingCompletedState.Reason == UE::StateTree::FFinishedTask::EReasonType::GlobalTask)
 			{
-				if (NewFinishedTask.Reason == UE::StateTree::FDelayedFinishedTask::EReasonType::GlobalTask
+				if (NewFinishedTask.Reason == UE::StateTree::FFinishedTask::EReasonType::GlobalTask
 					&& NewFinishedTask.Type == EStateTreeFinishTaskType::Failed)
 				{
 					// Set failure, otherwise keep as it is
@@ -1475,9 +1475,9 @@ void FStateTreeExecutionContext::AddUniqueCompletedStateList(const UE::StateTree
 			}
 			else
 			{
-				check(ExistingCompletedState.Reason == UE::StateTree::FDelayedFinishedTask::EReasonType::StateTask
-					|| ExistingCompletedState.Reason == UE::StateTree::FDelayedFinishedTask::EReasonType::InternalTransition);
-				if (NewFinishedTask.Reason == UE::StateTree::FDelayedFinishedTask::EReasonType::GlobalTask)
+				check(ExistingCompletedState.Reason == UE::StateTree::FFinishedTask::EReasonType::StateTask
+					|| ExistingCompletedState.Reason == UE::StateTree::FFinishedTask::EReasonType::InternalTransition);
+				if (NewFinishedTask.Reason == UE::StateTree::FFinishedTask::EReasonType::GlobalTask)
 				{
 					CompletedStates.RemoveAtSwap(CompletedStateIndex);
 					break;
@@ -1498,8 +1498,8 @@ void FStateTreeExecutionContext::AddUniqueCompletedStateList(const UE::StateTree
 
 	if (bAdd)
 	{
-		FCompletedState& CompletedState = CompletedStates.Add_GetRef(FCompletedState(NewFinishedTask));
-		CompletedState.RunStatus = UE::StateTree::ExecutionContext::Private::FinishTaskTypeToRunStatus(NewFinishedTask.Type);
+		const EStateTreeRunStatus RunStatus = UE::StateTree::ExecutionContext::Private::GetRunStatusFromFinishTaskType(NewFinishedTask.Type);
+		CompletedStates.Emplace(NewFinishedTask, RunStatus);
 	}
 }
 
@@ -1508,7 +1508,7 @@ EStateTreeRunStatus FStateTreeExecutionContext::GetGlobalTasksCompletedStatesSta
 	EStateTreeRunStatus Result = EStateTreeRunStatus::Running;
 	for (const FCompletedState& CompletedState : CompletedStates)
 	{
-		if (CompletedState.Reason == UE::StateTree::FDelayedFinishedTask::EReasonType::GlobalTask)
+		if (CompletedState.Reason == UE::StateTree::FFinishedTask::EReasonType::GlobalTask)
 		{
 			// return failure in priority
 			if (CompletedState.RunStatus == EStateTreeRunStatus::Failed)
@@ -2350,7 +2350,7 @@ EStateTreeRunStatus FStateTreeExecutionContext::EnterState(FStateTreeTransitionR
 	CaptureNewStateEvents(Exec.ActiveFrames, Transition.NextActiveFrames, Transition.NextActiveFrameEvents);
 
 	Exec.StateChangeCount++;
-	Exec.DelayedFinishedTasks.Empty();
+	Exec.FinishedTasks.Empty();
 	Exec.EnterStateFailedFrameIndex = FStateTreeIndex16::Invalid; // This will make all tasks to be accepted.
 	Exec.EnterStateFailedTaskIndex = FStateTreeIndex16::Invalid; // This will make all tasks to be accepted.
 	
@@ -2540,7 +2540,7 @@ EStateTreeRunStatus FStateTreeExecutionContext::EnterState(FStateTreeTransitionR
 					if (Status != EStateTreeRunStatus::Running)
 					{
 						Result = Status;
-						FinishTaskInternal(UE::StateTree::FDelayedFinishedTask::EReasonType::StateTask, CurrentFrame.FrameID, CurrentFrame.ActiveStates.StateIDs[StateIndex], TaskIndex, Result);
+						FinishTaskInternal(UE::StateTree::FFinishedTask::EReasonType::StateTask, CurrentFrame.FrameID, CurrentFrame.ActiveStates.StateIDs[StateIndex], TaskIndex, Result);
 						if (Status == EStateTreeRunStatus::Failed)
 						{
 							Exec.EnterStateFailedFrameIndex = FStateTreeIndex16(FrameIndex);
@@ -3013,7 +3013,7 @@ EStateTreeRunStatus FStateTreeExecutionContext::TickEvaluatorsAndGlobalTasks(con
 
 			const FCompletedState* FrameCompletedState = CompletedStates.FindByPredicate([FrameID = CurrentFrame.FrameID](FCompletedState& Other)
 				{
-					return Other.FrameID == FrameID && Other.Reason == UE::StateTree::FDelayedFinishedTask::EReasonType::GlobalTask;
+					return Other.FrameID == FrameID && Other.Reason == UE::StateTree::FFinishedTask::EReasonType::GlobalTask;
 				});
 			if (FrameCompletedState)
 			{
@@ -3579,7 +3579,7 @@ EStateTreeRunStatus FStateTreeExecutionContext::TickTasks(const float DeltaTime)
 					if (TaskResult != EStateTreeRunStatus::Running)
 					{
 						Result = TaskResult;
-						FinishTaskInternal(UE::StateTree::FDelayedFinishedTask::EReasonType::StateTask, CurrentFrame.FrameID, CurrentFrame.ActiveStates.StateIDs[StateIndex], TaskIndex, TaskResult);
+						FinishTaskInternal(UE::StateTree::FFinishedTask::EReasonType::StateTask, CurrentFrame.FrameID, CurrentFrame.ActiveStates.StateIDs[StateIndex], TaskIndex, TaskResult);
 						if (TaskResult == EStateTreeRunStatus::Failed)
 						{
 							bShouldTickTasks = false;
@@ -3594,7 +3594,7 @@ EStateTreeRunStatus FStateTreeExecutionContext::TickTasks(const float DeltaTime)
 		}
 	}
 
-	// Delayed completed task but no task were ticked to return the value.
+	// Delayed finished task but no task were ticked to return the value.
 	if (Result == EStateTreeRunStatus::Running && CompletedStates.Num() > 0)
 	{
 		const FCompletedState* FirstFailureCompletedState = CompletedStates.FindByPredicate([](const FCompletedState& Other) { return Other.RunStatus != EStateTreeRunStatus::Succeeded; });
@@ -3608,7 +3608,7 @@ EStateTreeRunStatus FStateTreeExecutionContext::TickTasks(const float DeltaTime)
 			const FStateTreeExecutionFrame& LastFrame = Exec.ActiveFrames.Last();
 			const UE::StateTree::FActiveFrameID CompletedFrameID = LastFrame.FrameID;
 			const UE::StateTree::FActiveStateID CompletedStateID = LastFrame.ActiveStates.StateIDs[LastFrame.ActiveStates.Num() - 1];
-			FinishTaskInternal(UE::StateTree::FDelayedFinishedTask::EReasonType::InternalTransition, CompletedFrameID, CompletedStateID, INDEX_NONE, EStateTreeRunStatus::Succeeded);
+			FinishTaskInternal(UE::StateTree::FFinishedTask::EReasonType::InternalTransition, CompletedFrameID, CompletedStateID, INDEX_NONE, EStateTreeRunStatus::Succeeded);
 		}
 		else
 		{
@@ -4386,7 +4386,7 @@ PRAGMA_ENABLE_DEPRECATION_WARNINGS
 		for (const FCompletedState& CompletedState : CompletedStates)
 		{
 			// Global task are not triggered in the transition loop.
-			if (CompletedState.Reason != UE::StateTree::FDelayedFinishedTask::EReasonType::GlobalTask)
+			if (CompletedState.Reason != UE::StateTree::FFinishedTask::EReasonType::GlobalTask)
 			{
 				const int32 NewFrameIndex = Exec.IndexOfActiveFrame(CompletedState.FrameID);
 				check(NewFrameIndex != INDEX_NONE); // the array is supposed to be clear every loop
@@ -4402,7 +4402,7 @@ PRAGMA_ENABLE_DEPRECATION_WARNINGS
 			}
 		}
 
-		Exec.DelayedFinishedTasks.Empty();
+		Exec.FinishedTasks.Empty();
 		CompletedStates.Empty();
 
 		// Start from the last completed state if specified, else only check the last state from the last frame.
@@ -4516,7 +4516,7 @@ PRAGMA_ENABLE_DEPRECATION_WARNINGS
 	}
 	else
 	{
-		Exec.DelayedFinishedTasks.Empty();
+		Exec.FinishedTasks.Empty();
 		CompletedStates.Empty();
 	}
 
@@ -4544,7 +4544,7 @@ PRAGMA_ENABLE_DEPRECATION_WARNINGS
 				// Set the parent linked state as last completed state, and update tick status to the status from the transition.
 				const int32 TaskIndex = INDEX_NONE;
 				const UE::StateTree::FActiveStateID ParentLastActiveStateID = ParentFrame.ActiveStates.StateIDs[ParentFrame.ActiveStates.Num() - 1];
-				FinishTaskInternal(UE::StateTree::FDelayedFinishedTask::EReasonType::InternalTransition, ParentFrame.FrameID, ParentLastActiveStateID, TaskIndex, RunStatus);
+				FinishTaskInternal(UE::StateTree::FFinishedTask::EReasonType::InternalTransition, ParentFrame.FrameID, ParentLastActiveStateID, TaskIndex, RunStatus);
 				Exec.LastTickStatus = RunStatus;
 
 				// Clear the transition and return that no transition took place.
@@ -5955,8 +5955,8 @@ bool FStateTreeWeakExecutionContext::FinishTask(const FStateTreeWeakTaskRef& Tas
 					const int32 GlobalTaskBeginIndex = ExecutionFrame->StateTree->GlobalTasksBegin;
 					const int32 GlobalTaskEndIndex = ExecutionFrame->StateTree->GlobalTasksBegin + ExecutionFrame->StateTree->GlobalTasksNum;
 					const bool bIsGlobalTask = TaskPin.GetTaskIndex().AsInt32() >= GlobalTaskBeginIndex && TaskPin.GetTaskIndex().AsInt32() < GlobalTaskEndIndex;
-					const UE::StateTree::FDelayedFinishedTask::EReasonType Reason = bIsGlobalTask ? UE::StateTree::FDelayedFinishedTask::EReasonType::GlobalTask : UE::StateTree::FDelayedFinishedTask::EReasonType::StateTask;
-					UE::StateTree::FDelayedFinishedTask CompletedTask{
+					const UE::StateTree::FFinishedTask::EReasonType Reason = bIsGlobalTask ? UE::StateTree::FFinishedTask::EReasonType::GlobalTask : UE::StateTree::FFinishedTask::EReasonType::StateTask;
+					UE::StateTree::FFinishedTask CompletedTask{
 						.FrameID = FrameID,
 						.StateID = StateID,
 						.TaskIndex = TaskPin.GetTaskIndex(),
@@ -5965,7 +5965,7 @@ bool FStateTreeWeakExecutionContext::FinishTask(const FStateTreeWeakTaskRef& Tas
 					};
 
 					// Confirm that the task if from the correct state.
-					Exec.DelayedFinishedTasks.Add(CompletedTask);
+					Exec.FinishedTasks.Add(CompletedTask);
 					ScheduleNextTick(OwnerPinned.Get(), StateTreePinned.Get(), *StoragePinned);
 					bCompleted = true;
 				}
diff --git a/Engine/Plugins/Runtime/StateTree/Source/StateTreeModule/Public/StateTreeExecutionContext.h b/Engine/Plugins/Runtime/StateTree/Source/StateTreeModule/Public/StateTreeExecutionContext.h
--- a/Engine/Plugins/Runtime/StateTree/Source/StateTreeModule/Public/StateTreeExecutionContext.h
+++ b/Engine/Plugins/Runtime/StateTree/Source/StateTreeModule/Public/StateTreeExecutionContext.h
@@ -585,12 +585,13 @@ public:
 	 */
 	void FinishTask(const FStateTreeTaskBase& Task, EStateTreeFinishTaskType FinishType);
 	
+	UE_DEPRECATED(5.6, "Use the weak context to finish a task async or FinishTask(FStateTreeTaskBase, EStateTreeFinishTaskType) to finish the current task.")
 	/**
 	 * Finishes a task.
 	 * If called during tick processing, then the state completes immediately.
 	 * If called outside of the tick processing, then the request is buffered and handled on the next tick.
 	 */
-	void FinishTask(const UE::StateTree::FDelayedFinishedTask& Task, EStateTreeFinishTaskType FinishType);
+	void FinishTask(const UE::StateTree::FFinishedTask& Task, EStateTreeFinishTaskType FinishType);
 
 	/** @return data view of the specified handle relative to given frame. */
 	static FStateTreeDataView GetDataViewFromInstanceStorage(FStateTreeInstanceStorage& InstanceDataStorage, FStateTreeInstanceStorage* CurrentlyProcessedSharedInstanceStorage, const FStateTreeExecutionFrame* ParentFrame, const FStateTreeExecutionFrame& CurrentFrame, const FStateTreeDataHandle Handle);
@@ -798,11 +799,11 @@ protected:
 	 */
 	bool TriggerTransitions();
 
-	/** Adds a task to the finished and delayed list. */
-	void FinishTaskInternal(UE::StateTree::FDelayedFinishedTask::EReasonType Reason, UE::StateTree::FActiveFrameID ActiveFrameID, UE::StateTree::FActiveStateID ActiveStateID, int32 TaskIndex, EStateTreeRunStatus RunStatus);
+	/** Adds the state to the completed list and adds task to the finished list. */
+	void FinishTaskInternal(UE::StateTree::FFinishedTask::EReasonType Reason, UE::StateTree::FActiveFrameID ActiveFrameID, UE::StateTree::FActiveStateID ActiveStateID, int32 TaskIndex, EStateTreeRunStatus RunStatus);
 
 	/** Confirms that the frame and state ID are valid and the task index is correct. */
-	bool IsDelayedFinishedTaskValid(const UE::StateTree::FDelayedFinishedTask& Task) const;
+	bool IsFinishedTaskValid(const UE::StateTree::FFinishedTask& Task) const;
 
 	/**
 	 * Removes stale entries and fills the completed states list.
@@ -810,11 +811,8 @@ protected:
 	 */
 	void UpdateCompletedStateList();
 
-	/**
-	 * Adds uniquely the task to the finished.
-	 * Adds the task to the delayed list. The delayed list may contains duplicate. They will be filtered with UpdateCompletedStateList when ticked.
-	 */
-	void AddUniqueCompletedStateList(const UE::StateTree::FDelayedFinishedTask& DelayedFinisedTask);
+	/** Adds the state to the completed list from a completed task. */
+	void MarkStateCompleted(const UE::StateTree::FFinishedTask& FinisedTask);
 
 	/** In the list of completed states, returns the run status when it's completed by a global task. */
 	EStateTreeRunStatus GetGlobalTasksCompletedStatesStatus() const;
@@ -960,11 +958,12 @@ protected:
 	TOptional<EStateTreeRunStatus> CompletedStateRunStatus;
 
 	/** Helper struct to store the resolved FStateTreeExecutionState::FCompletedState */
-	struct FCompletedState : UE::StateTree::FDelayedFinishedTask
+	struct FCompletedState : UE::StateTree::FFinishedTask
 	{
 		FCompletedState() = default;
-		FCompletedState(const UE::StateTree::FDelayedFinishedTask& Other)
-			: UE::StateTree::FDelayedFinishedTask(Other)
+		FCompletedState(const UE::StateTree::FFinishedTask& Other, EStateTreeRunStatus InRunStatus)
+			: UE::StateTree::FFinishedTask(Other)
+			, RunStatus(InRunStatus)
 		{}
 		EStateTreeRunStatus RunStatus = EStateTreeRunStatus::Unset;
 	};
diff --git a/Engine/Plugins/Runtime/StateTree/Source/StateTreeModule/Public/StateTreeExecutionTypes.h b/Engine/Plugins/Runtime/StateTree/Source/StateTreeModule/Public/StateTreeExecutionTypes.h
--- a/Engine/Plugins/Runtime/StateTree/Source/StateTreeModule/Public/StateTreeExecutionTypes.h
+++ b/Engine/Plugins/Runtime/StateTree/Source/StateTreeModule/Public/StateTreeExecutionTypes.h
@@ -678,7 +678,7 @@ PRAGMA_ENABLE_DEPRECATION_WARNINGS
 namespace UE::StateTree
 {
 /** Describes a finished task waiting to be processed by an execution context. */
-struct FDelayedFinishedTask
+struct FFinishedTask
 {
 	enum class EReasonType : uint8
 	{
@@ -971,7 +971,7 @@ public:
 	TArray<FStateTreeTransitionDelayedState> DelayedTransitions;
 
 	/** Pending finished tasks that need processing. */
-	TArray<UE::StateTree::FDelayedFinishedTask> DelayedFinishedTasks;
+	TArray<UE::StateTree::FFinishedTask> FinishedTasks;
 
 	/** Used by state tree random-based operations. */
 	UPROPERTY()
