StateTree: Fix the localization issue introduced with CL 39902219.

#rnx

diff --git a/Engine/Plugins/Runtime/StateTree/Source/StateTreeEditorModule/Private/K2Node_StateTreeReference.cpp b/Engine/Plugins/Runtime/StateTree/Source/StateTreeEditorModule/Private/K2Node_StateTreeReference.cpp
--- a/Engine/Plugins/Runtime/StateTree/Source/StateTreeEditorModule/Private/K2Node_StateTreeReference.cpp
+++ b/Engine/Plugins/Runtime/StateTree/Source/StateTreeEditorModule/Private/K2Node_StateTreeReference.cpp
@@ -107,7 +107,7 @@ FText UK2Node_MakeStateTreeReference::GetTooltipText() const
 		{
 			FFormatNamedArguments Args;
 			Args.Add(TEXT("StructName"), FStateTreeReference::StaticStruct()->GetDisplayNameText());
-			return FText::Format(LOCTEXT("MakeNodeTitle", "Adds a node that create a  {StructName} from its members"), Args);
+			return FText::Format(LOCTEXT("MakeNodeTooltip", "Adds a node that create a  {StructName} from its members"), Args);
 		};
 	static FText LocalCachedTooltip = MakeNodeTooltip();
 	return LocalCachedTooltip;
diff --git a/Engine/Plugins/Runtime/StateTree/Source/StateTreeModule/Private/StateTreeFunctionLibrary.cpp b/Engine/Plugins/Runtime/StateTree/Source/StateTreeModule/Private/StateTreeFunctionLibrary.cpp
--- a/Engine/Plugins/Runtime/StateTree/Source/StateTreeModule/Private/StateTreeFunctionLibrary.cpp
+++ b/Engine/Plugins/Runtime/StateTree/Source/StateTreeModule/Private/StateTreeFunctionLibrary.cpp
@@ -50,7 +50,7 @@ DEFINE_FUNCTION(UStateTreeFunctionLibrary::execK2_SetParametersProperty)
 	{
 		FBlueprintExceptionInfo ExceptionInfo(
 			EBlueprintExceptionType::AbortExecution,
-			LOCTEXT("GetParametersProperty_InvalidValueWarning", "Failed to resolve the Value for SetParametersProperty")
+			LOCTEXT("SetParametersProperty_InvalidValueWarning", "Failed to resolve the Value for SetParametersProperty")
 		);
 
 		FBlueprintCoreDelegates::ThrowScriptException(P_THIS, Stack, ExceptionInfo);
