StateTree: Add a priority to task results.

#rnx

diff --git a/Engine/Plugins/Runtime/StateTree/Source/StateTreeModule/Private/StateTreeExecutionContext.cpp b/Engine/Plugins/Runtime/StateTree/Source/StateTreeModule/Private/StateTreeExecutionContext.cpp
--- a/Engine/Plugins/Runtime/StateTree/Source/StateTreeModule/Private/StateTreeExecutionContext.cpp
+++ b/Engine/Plugins/Runtime/StateTree/Source/StateTreeModule/Private/StateTreeExecutionContext.cpp
@@ -75,6 +75,28 @@ namespace UE::StateTree::ExecutionContext::Private
 		static_assert((int32)EStateTreeRunStatus::Failed == (int32)EStateTreeFinishTaskType::Failed, "Failed do not match in both types.");
 		return (EStateTreeRunStatus)FinishTask;
 	}
+
+	/** @return in order {Running, Failed, Succeeded, Stopped, Unset} */
+	EStateTreeRunStatus GetPriorityRunStatus(EStateTreeRunStatus A, EStateTreeRunStatus B)
+	{
+		constexpr auto NumEStateTreeRunStatus = []()
+		{
+#ifdef UE_STATETREE_ESTATETREERUNSTATUS_PRIVATE
+#error UE_STATETREE_ESTATETREERUNSTATUS_PRIVATE_ALREADY_DEFINED
+#endif
+
+#define UE_STATETREE_ESTATETREERUNSTATUS_PRIVATE(EnumValue) ++Number;
+
+			int32 Number = 0;
+			FOREACH_ENUM_ESTATETREERUNSTATUS(UE_STATETREE_ESTATETREERUNSTATUS_PRIVATE)
+
+#undef UE_STATETREE_ESTATETREERUNSTATUS_PRIVATE
+			return Number;
+		};
+		constexpr int32 PriorityMatrix[] = { 1, 4, 3, 2, 0 };
+		static_assert(NumEStateTreeRunStatus() == UE_ARRAY_COUNT(PriorityMatrix), "The number of entries in EStateTreeRunStatus changed. Update GetPriorityRunStatus.");
+		return PriorityMatrix[(uint8)A] > PriorityMatrix[(uint8)B] ? A : B;
+	}
 }
 
 FStateTreeMinimalExecutionContext::FStateTreeMinimalExecutionContext(UObject& InOwner, const UStateTree& InStateTree, FStateTreeInstanceData& InInstanceData)
@@ -2527,15 +2549,14 @@ EStateTreeRunStatus FStateTreeExecutionContext::EnterState(FStateTreeTransitionR
 						Status = Task.EnterState(*this, CurrentTransition);
 					}
 
-					UE_STATETREE_DEBUG_TASK_EVENT(this, TaskIndex, TaskInstanceView, EStateTreeTraceEventType::OnEntered, Status);
-
-					// Delayed from CompleteTask(), override the Tick result unless it's a failure.
-					if (CompletedStateRunStatus.IsSet()
-						&& (Status == EStateTreeRunStatus::Succeeded || Status == EStateTreeRunStatus::Running))
+					// Delayed from FinishTask(), use the best priority from what the task returns and FinishTask().
+					if (CompletedStateRunStatus.IsSet())
 					{
-						Status = CompletedStateRunStatus.GetValue();
+						Status = UE::StateTree::ExecutionContext::Private::GetPriorityRunStatus(Status, CompletedStateRunStatus.GetValue());
 						CompletedStateRunStatus.Reset();
 					}
+					UE_STATETREE_ENSURE_COMPLETED_STATE_RUN_STATUS();
+					UE_STATETREE_DEBUG_TASK_EVENT(this, TaskIndex, TaskInstanceView, EStateTreeTraceEventType::OnEntered, Status);
 
 					if (Status != EStateTreeRunStatus::Running)
 					{
@@ -2548,8 +2569,6 @@ EStateTreeRunStatus FStateTreeExecutionContext::EnterState(FStateTreeTransitionR
 							break;
 						}
 					}
-
-					UE_STATETREE_ENSURE_COMPLETED_STATE_RUN_STATUS();
 				}
 			}
 
@@ -3067,14 +3086,13 @@ EStateTreeRunStatus FStateTreeExecutionContext::TickEvaluatorsAndGlobalTasks(con
 							TaskResult = Task.Tick(*this, DeltaTime);
 						}
 
-						// Delayed from CompleteTask(), override the Tick result unless it's a failure.
-						if (CompletedStateRunStatus.IsSet()
-							&& (TaskResult == EStateTreeRunStatus::Succeeded || TaskResult == EStateTreeRunStatus::Running))
+						// Delayed from FinishTask(), use the best priority from what the task returns and FinishTask().
+						if (CompletedStateRunStatus.IsSet())
 						{
-							TaskResult = CompletedStateRunStatus.GetValue();
+							TaskResult = UE::StateTree::ExecutionContext::Private::GetPriorityRunStatus(TaskResult, CompletedStateRunStatus.GetValue());
 							CompletedStateRunStatus.Reset();
 						}
-
+						UE_STATETREE_ENSURE_COMPLETED_STATE_RUN_STATUS();
 						UE_STATETREE_DEBUG_TASK_EVENT(this, TaskIndex, TaskInstanceView,
 							TaskResult != EStateTreeRunStatus::Running ? EStateTreeTraceEventType::OnTaskCompleted : EStateTreeTraceEventType::OnTicked,
 							TaskResult);
@@ -3090,8 +3108,6 @@ EStateTreeRunStatus FStateTreeExecutionContext::TickEvaluatorsAndGlobalTasks(con
 								bShouldTickTasks = false;
 							}
 						}
-
-						UE_STATETREE_ENSURE_COMPLETED_STATE_RUN_STATUS();
 					}
 				}
 			}
@@ -3148,7 +3164,8 @@ EStateTreeRunStatus FStateTreeExecutionContext::StartEvaluatorsAndGlobalTasks(FS
 			// Start Global tasks
 			// Even if we call Enter/ExitState() on global tasks, they do not enter any specific state.
 			const FStateTreeTransitionResult Transition = {}; // Empty transition
-		
+			CompletedStateRunStatus.Reset();
+
 			for (int32 TaskIndex = CurrentStateTree->GlobalTasksBegin; TaskIndex < (CurrentStateTree->GlobalTasksBegin + CurrentStateTree->GlobalTasksNum); TaskIndex++)
 			{
 				const FStateTreeTaskBase& Task =  CurrentStateTree->Nodes[TaskIndex].Get<const FStateTreeTaskBase>();
@@ -3169,11 +3186,19 @@ EStateTreeRunStatus FStateTreeExecutionContext::StartEvaluatorsAndGlobalTasks(FS
 				}
 
 				STATETREE_LOG(Verbose, TEXT("  Start: '%s'"), *Task.Name.ToString());
+				UE_STATETREE_ENSURE_COMPLETED_STATE_RUN_STATUS();
 				UE_STATETREE_DEBUG_TASK_ENTER_STATE(this, CurrentStateTree, FStateTreeIndex16(TaskIndex));
 				{
 					QUICK_SCOPE_CYCLE_COUNTER(StateTree_Task_TreeStart);
-					const EStateTreeRunStatus TaskStatus = Task.EnterState(*this, Transition);
+					EStateTreeRunStatus TaskStatus = Task.EnterState(*this, Transition);
 
+					// Delayed from FinishTask(), use the best priority from what the task returns and FinishTask().
+					if (CompletedStateRunStatus.IsSet())
+					{
+						TaskStatus = UE::StateTree::ExecutionContext::Private::GetPriorityRunStatus(TaskStatus, CompletedStateRunStatus.GetValue());
+						CompletedStateRunStatus.Reset();
+					}
+					UE_STATETREE_ENSURE_COMPLETED_STATE_RUN_STATUS();
 					UE_STATETREE_DEBUG_TASK_EVENT(this, TaskIndex, TaskInstanceView, EStateTreeTraceEventType::OnEntered, TaskStatus);
 
 					if (TaskStatus != EStateTreeRunStatus::Running)
@@ -3386,10 +3411,19 @@ EStateTreeRunStatus FStateTreeExecutionContext::StartTemporaryEvaluatorsAndGloba
 		}
 
 		STATETREE_LOG(Verbose, TEXT("  Start: '%s'"), *Task.Name.ToString());
+		UE_STATETREE_ENSURE_COMPLETED_STATE_RUN_STATUS();
 		UE_STATETREE_DEBUG_TASK_ENTER_STATE(this, CurrentStateTree, FStateTreeIndex16(TaskIndex));
 		QUICK_SCOPE_CYCLE_COUNTER(StateTree_Task_TreeStart);
-		const EStateTreeRunStatus TaskStatus = Task.EnterState(*this, Transition);
 
+		EStateTreeRunStatus TaskStatus = Task.EnterState(*this, Transition);
+
+		// Delayed from FinishTask(), use the best priority from what the task returns and FinishTask().
+		if (CompletedStateRunStatus.IsSet())
+		{
+			TaskStatus = UE::StateTree::ExecutionContext::Private::GetPriorityRunStatus(TaskStatus, CompletedStateRunStatus.GetValue());
+			CompletedStateRunStatus.Reset();
+		}
+		UE_STATETREE_ENSURE_COMPLETED_STATE_RUN_STATUS();
 		UE_STATETREE_DEBUG_TASK_EVENT(this, TaskIndex, TaskDataView, EStateTreeTraceEventType::OnEntered, TaskStatus);
 
 		if (TaskStatus != EStateTreeRunStatus::Running)
@@ -3564,14 +3598,13 @@ EStateTreeRunStatus FStateTreeExecutionContext::TickTasks(const float DeltaTime)
 						TaskResult = Task.Tick(*this, DeltaTime);
 					}
 
-					// Delayed from CompleteTask(), override the Tick result unless it's a failure.
-					if (CompletedStateRunStatus.IsSet()
-						&& (TaskResult == EStateTreeRunStatus::Succeeded || TaskResult == EStateTreeRunStatus::Running))
+					// Delayed from FinishTask(), use the best priority from what the task returns and the FinishTask().
+					if (CompletedStateRunStatus.IsSet())
 					{
-						TaskResult = CompletedStateRunStatus.GetValue();
+						Result = UE::StateTree::ExecutionContext::Private::GetPriorityRunStatus(Result, CompletedStateRunStatus.GetValue());
 						CompletedStateRunStatus.Reset();
 					}
-
+					UE_STATETREE_ENSURE_COMPLETED_STATE_RUN_STATUS();
 					UE_STATETREE_DEBUG_TASK_EVENT(this, TaskIndex, TaskInstanceView,
 						TaskResult != EStateTreeRunStatus::Running ? EStateTreeTraceEventType::OnTaskCompleted : EStateTreeTraceEventType::OnTicked,
 						TaskResult);
@@ -3585,8 +3618,6 @@ EStateTreeRunStatus FStateTreeExecutionContext::TickTasks(const float DeltaTime)
 							bShouldTickTasks = false;
 						}
 					}
-
-					UE_STATETREE_ENSURE_COMPLETED_STATE_RUN_STATUS();
 				}
 			}
 
