PR #12850: Use the correct tree to iterate on transitions when selecting a new state.

diff --git a/Engine/Plugins/Runtime/StateTree/Source/StateTreeModule/Private/StateTreeExecutionContext.cpp b/Engine/Plugins/Runtime/StateTree/Source/StateTreeModule/Private/StateTreeExecutionContext.cpp
--- a/Engine/Plugins/Runtime/StateTree/Source/StateTreeModule/Private/StateTreeExecutionContext.cpp
+++ b/Engine/Plugins/Runtime/StateTree/Source/StateTreeModule/Private/StateTreeExecutionContext.cpp
@@ -4671,7 +4671,7 @@ bool FStateTreeExecutionContext::SelectState(const FStateTreeExecutionFrame& Cur
 	else
 	{
 		STATETREE_LOG(Error, TEXT("%hs: Encountered unrecognized state %s during state selection from '%s'.  '%s' using StateTree '%s'."),
-			__FUNCTION__, *GetNameSafe(NextStateTree), *GetStateStatusString(Exec), *GetNameSafe(&Owner), *GetFullNameSafe(&RootStateTree));
+			__FUNCTION__, *GetNameSafe(NextStateTree), *GetStateStatusString(Exec), *GetNameSafe(&Owner), *GetFullNameSafe(NextStateTree));
 		return false;
 	}
 	
@@ -5142,7 +5142,7 @@ bool FStateTreeExecutionContext::SelectStateInternal(
 			for (uint8 i = 0; i < NextState.TransitionsNum; i++)
 			{
 				const int16 TransitionIndex = NextState.TransitionsBegin + i;
-				const FCompactStateTransition& Transition = RootStateTree.Transitions[TransitionIndex];
+				const FCompactStateTransition& Transition = CurrentStateTree->Transitions[TransitionIndex];
 
 				// Skip disabled transitions
 				if (Transition.bTransitionEnabled == false)
