StateTree: Add cvar StateTree.ForceCopyBoundPropertiesOnNonTickedTask. Code comments described a behavior that was not what the code was actually doing. Add a cvar to toggle between the 2 behaviors.

diff --git a/Engine/Plugins/Runtime/StateTree/Source/StateTreeModule/Private/StateTreeExecutionContext.cpp b/Engine/Plugins/Runtime/StateTree/Source/StateTreeModule/Private/StateTreeExecutionContext.cpp
--- a/Engine/Plugins/Runtime/StateTree/Source/StateTreeModule/Private/StateTreeExecutionContext.cpp
+++ b/Engine/Plugins/Runtime/StateTree/Source/StateTreeModule/Private/StateTreeExecutionContext.cpp
@@ -36,6 +36,13 @@ namespace UE::StateTree::Debug
 
 namespace UE::StateTree::ExecutionContext::Private
 {
+	bool bForceCopyBoundPropertiesOnNonTickedTask = false;
+	FAutoConsoleVariableRef CVarLogCompiledStateTree(
+		TEXT("StateTree.ForceCopyBoundPropertiesOnNonTickedTask"),
+		bForceCopyBoundPropertiesOnNonTickedTask,
+		TEXT("When ticking the tasks, copy the bindings when the tasks are not ticked. A task can not ticked because of another task failure or because it disabled ticking.")
+	);
+
 	bool MarkDelegateAsBroadcasted(FStateTreeDelegateDispatcher Dispatcher, const FStateTreeExecutionFrame& CurrentFrame, FStateTreeInstanceStorage& Storage)
 	{
 		const UStateTree* StateTree = CurrentFrame.StateTree;
@@ -2985,6 +2992,11 @@ void FStateTreeExecutionContext::StateCompleted()
 
 EStateTreeRunStatus FStateTreeExecutionContext::TickEvaluatorsAndGlobalTasks(const float DeltaTime, const bool bTickGlobalTasks)
 {
+	// When a global task is completed it completes the tree execution.
+	// A global task can complete async. See CompletedStates.
+	// CompletedStateRunStatus is used to detect that a global task is completed with FinishTask instead of a global task returning a completed result.
+	// When a global task fails, stop ticking the following tasks. If bForceCopyBoundPropertiesOnNonTickedTask, update their bindings.
+
 	CSV_SCOPED_TIMING_STAT_EXCLUSIVE(StateTree_TickEvaluators);
 	UE_STATETREE_DEBUG_SCOPED_PHASE(this, EStateTreeUpdatePhase::TickingGlobalTasks);
 
@@ -2995,6 +3007,8 @@ EStateTreeRunStatus FStateTreeExecutionContext::TickEvaluatorsAndGlobalTasks(con
 	EStateTreeRunStatus Result = EStateTreeRunStatus::Running;
 	CompletedStateRunStatus.Reset();
 
+	const bool bForceCopyBoundPropertiesOnNonTickedTask = UE::StateTree::ExecutionContext::Private::bForceCopyBoundPropertiesOnNonTickedTask;
+
 	for (int32 FrameIndex = 0; FrameIndex < Exec.ActiveFrames.Num(); FrameIndex++)
 	{
 		FStateTreeExecutionFrame* CurrentParentFrame = FrameIndex > 0 ? &Exec.ActiveFrames[FrameIndex - 1] : nullptr;
@@ -3027,7 +3041,7 @@ EStateTreeRunStatus FStateTreeExecutionContext::TickEvaluatorsAndGlobalTasks(con
 				}
 			}
 
-			// Used to stop ticking tasks after one fails, but we still want to keep updating the data views so that property binding works properly.
+			// Used to stop ticking tasks after one fails, but if bForceCopyBoundPropertiesOnNonTickedTask, we still want to keep updating the data views so that property binding works properly.
 			bool bShouldTickTasks = true;
 
 			const FCompletedState* FrameCompletedState = CompletedStates.FindByPredicate([FrameID = CurrentFrame.FrameID](FCompletedState& Other)
@@ -3037,15 +3051,14 @@ EStateTreeRunStatus FStateTreeExecutionContext::TickEvaluatorsAndGlobalTasks(con
 			if (FrameCompletedState)
 			{
 				bShouldTickTasks = false;
-				//@TODO: use a priority system? What if another frame did return fail and this frame success.
-				Result = FrameCompletedState->RunStatus;
+				Result = UE::StateTree::ExecutionContext::Private::GetPriorityRunStatus(Result, FrameCompletedState->RunStatus);
 			}
 
 			if (bTickGlobalTasks)
 			{
 				const bool bHasEvents = EventQueue && EventQueue->HasEvents();
 
-				if (CurrentStateTree->ShouldTickGlobalTasks(bHasEvents))
+				if (bForceCopyBoundPropertiesOnNonTickedTask || CurrentStateTree->ShouldTickGlobalTasks(bHasEvents))
 				{
 					for (int32 TaskIndex = CurrentStateTree->GlobalTasksBegin; TaskIndex < (CurrentStateTree->GlobalTasksBegin + CurrentStateTree->GlobalTasksNum); TaskIndex++)
 					{
@@ -3062,18 +3075,22 @@ EStateTreeRunStatus FStateTreeExecutionContext::TickEvaluatorsAndGlobalTasks(con
 
 						const bool bNeedsTick = bShouldTickTasks && (Task.bShouldCallTick || (bHasEvents && Task.bShouldCallTickOnlyOnEvents));
 						STATETREE_LOG(VeryVerbose, TEXT("  Tick: '%s' %s"), *Task.Name.ToString(), !bNeedsTick ? TEXT("[not ticked]") : TEXT(""));
-						if (!bNeedsTick)
-						{
-							continue;
-						}
 
 						// Copy bound properties.
 						// Only copy properties when the task is actually ticked, and copy properties at tick is requested.
-						if (Task.BindingsBatch.IsValid() && Task.bShouldCopyBoundPropertiesOnTick)
+						const bool bCopyBatch = (bForceCopyBoundPropertiesOnNonTickedTask || bNeedsTick)
+							&& Task.BindingsBatch.IsValid()
+							&& Task.bShouldCopyBoundPropertiesOnTick;
+						if (bCopyBatch)
 						{
 							CopyBatchOnActiveInstances(CurrentParentFrame, CurrentFrame, TaskInstanceView, Task.BindingsBatch);
 						}
 
+						if (!bNeedsTick)
+						{
+							continue;
+						}
+
 						//UE_STATETREE_DEBUG_TASK_EVENT(this, TaskIndex, TaskDataView, EStateTreeTraceEventType::OnTickingTask, EStateTreeRunStatus::Running);
 						UE_STATETREE_ENSURE_COMPLETED_STATE_RUN_STATUS();
 						UE_STATETREE_DEBUG_TASK_TICK(this, CurrentStateTree, TaskIndex);
@@ -3100,12 +3117,15 @@ EStateTreeRunStatus FStateTreeExecutionContext::TickEvaluatorsAndGlobalTasks(con
 						// If a global task succeeds or fails, it will stop the whole tree.
 						if (TaskResult != EStateTreeRunStatus::Running)
 						{
-							//@TODO: use a priority system? What if another frame did return fail and this one success.
-							Result = TaskResult;
-
+							Result = UE::StateTree::ExecutionContext::Private::GetPriorityRunStatus(Result, TaskResult);
+							FinishTaskInternal(UE::StateTree::FFinishedTask::EReasonType::GlobalTask, CurrentFrame.FrameID, UE::StateTree::FActiveStateID::Invalid, TaskIndex, TaskResult);
 							if (TaskResult == EStateTreeRunStatus::Failed)
 							{
 								bShouldTickTasks = false;
+								if (!bForceCopyBoundPropertiesOnNonTickedTask)
+								{
+									break; // go to next frame
+								}
 							}
 						}
 					}
@@ -3504,6 +3524,12 @@ void FStateTreeExecutionContext::StopTemporaryEvaluatorsAndGlobalTasks(const FSt
 
 EStateTreeRunStatus FStateTreeExecutionContext::TickTasks(const float DeltaTime)
 {
+	// When a task is completed it also completes the state and triggers the completion transition (because LastTickStatus is set).
+	// A task can complete async. See CompletedStates.
+	// CompletedStateRunStatus is used to detect that a task is completed with FinishTask instead of the task returning a completed result.
+	// When a task fails, stop ticking the following tasks. If bForceCopyBoundPropertiesOnNonTickedTask, update their bindings.
+	// When no task ticks, then the leaf completes.
+
 	CSV_SCOPED_TIMING_STAT_EXCLUSIVE(StateTree_TickTasks);
 	UE_STATETREE_DEBUG_SCOPED_PHASE(this, EStateTreeUpdatePhase::TickingTasks);
 
@@ -3519,7 +3545,7 @@ EStateTreeRunStatus FStateTreeExecutionContext::TickTasks(const float DeltaTime)
 
 	const bool bHasEvents = EventQueue && EventQueue->HasEvents();
 	
-	// Used to stop ticking tasks after one fails, but we still want to keep updating the data views so that property binding works properly.
+	const bool bForceCopyBoundPropertiesOnNonTickedTask = UE::StateTree::ExecutionContext::Private::bForceCopyBoundPropertiesOnNonTickedTask;
 	bool bShouldTickTasks = CompletedStates.Num() == 0;
 	CompletedStateRunStatus.Reset();
 
@@ -3553,7 +3579,7 @@ EStateTreeRunStatus FStateTreeExecutionContext::TickTasks(const float DeltaTime)
 				}
 			}
 
-			if (State.ShouldTickTasks(bHasEvents))
+			if (bForceCopyBoundPropertiesOnNonTickedTask || State.ShouldTickTasks(bHasEvents))
 			{
 				// Update Tasks data and tick if possible (ie. if no task has yet failed and so bShouldTickTasks is true)
 				for (int32 TaskIndex = State.TasksBegin; TaskIndex < (State.TasksBegin + State.TasksNum); TaskIndex++)
@@ -3574,18 +3600,21 @@ EStateTreeRunStatus FStateTreeExecutionContext::TickTasks(const float DeltaTime)
 						, *Task.Name.ToString()
 						, !bNeedsTick ? TEXT("[not ticked]") : TEXT(""));
 
-					if (!bNeedsTick)
-					{
-						continue;
-					}
-				
 					// Copy bound properties.
 					// Only copy properties when the task is actually ticked, and copy properties at tick is requested.
-					if (Task.BindingsBatch.IsValid() && Task.bShouldCopyBoundPropertiesOnTick)
+					const bool bCopyBatch = (bForceCopyBoundPropertiesOnNonTickedTask || bNeedsTick)
+						&& Task.BindingsBatch.IsValid()
+						&& Task.bShouldCopyBoundPropertiesOnTick;
+					if (bCopyBatch)
 					{
 						CopyBatchOnActiveInstances(CurrentParentFrame, CurrentFrame, TaskInstanceView, Task.BindingsBatch);
 					}
 
+					if (!bNeedsTick)
+					{
+						continue;
+					}
+
 					//UE_STATETREE_DEBUG_TASK_EVENT(this, TaskIndex, TaskDataView, EStateTreeTraceEventType::OnTickingTask, EStateTreeRunStatus::Running);
 					UE_STATETREE_ENSURE_COMPLETED_STATE_RUN_STATUS();
 					UE_STATETREE_DEBUG_TASK_TICK(this, CurrentStateTree, TaskIndex);
@@ -3613,9 +3642,13 @@ EStateTreeRunStatus FStateTreeExecutionContext::TickTasks(const float DeltaTime)
 					{
 						Result = TaskResult;
 						FinishTaskInternal(UE::StateTree::FFinishedTask::EReasonType::StateTask, CurrentFrame.FrameID, CurrentFrame.ActiveStates.StateIDs[StateIndex], TaskIndex, TaskResult);
-						if (TaskResult == EStateTreeRunStatus::Failed)
+						if (Result == EStateTreeRunStatus::Failed)
 						{
 							bShouldTickTasks = false;
+							if (!bForceCopyBoundPropertiesOnNonTickedTask)
+							{
+								break; // go to next frame
+							}
 						}
 					}
 				}
