StateTree: Save the run instead of the finished type. A sync FinishTask can request a Stopped in a parallel task and using the finished type would lose that information.

diff --git a/Engine/Plugins/Runtime/StateTree/Source/StateTreeModule/Private/StateTreeExecutionContext.cpp b/Engine/Plugins/Runtime/StateTree/Source/StateTreeModule/Private/StateTreeExecutionContext.cpp
--- a/Engine/Plugins/Runtime/StateTree/Source/StateTreeModule/Private/StateTreeExecutionContext.cpp
+++ b/Engine/Plugins/Runtime/StateTree/Source/StateTreeModule/Private/StateTreeExecutionContext.cpp
@@ -1379,9 +1379,8 @@ void FStateTreeExecutionContext::FinishTaskInternal(UE::StateTree::FFinishedTask
 	FStateTreeExecutionState& Exec = GetExecState();
 
 	const FStateTreeIndex16 TaskIndex16 = TaskIndex >= 0 ? FStateTreeIndex16(TaskIndex) : FStateTreeIndex16::Invalid;
-	const EStateTreeFinishTaskType FinishType = RunStatus == EStateTreeRunStatus::Succeeded ? EStateTreeFinishTaskType::Succeeded : EStateTreeFinishTaskType::Failed;
 
-	const FFinishedTask Task {.FrameID = ActiveFrameID, .StateID = ActiveStateID, .TaskIndex = TaskIndex16, .Type = FinishType, .Reason = Reason };
+	const FFinishedTask Task {.FrameID = ActiveFrameID, .StateID = ActiveStateID, .TaskIndex = TaskIndex16, .RunStatus = RunStatus, .Reason = Reason };
 	if (IsFinishedTaskValid(Task))
 	{
 		// Add it for later. In case the user called TickTasks and will call TickTransitions later.
@@ -1493,11 +1492,9 @@ void FStateTreeExecutionContext::MarkStateCompleted(const UE::StateTree::FFinish
 			
 			if (ExistingCompletedState.Reason == UE::StateTree::FFinishedTask::EReasonType::GlobalTask)
 			{
-				if (NewFinishedTask.Reason == UE::StateTree::FFinishedTask::EReasonType::GlobalTask
-					&& NewFinishedTask.Type == EStateTreeFinishTaskType::Failed)
+				if (NewFinishedTask.Reason == UE::StateTree::FFinishedTask::EReasonType::GlobalTask)
 				{
-					// Set failure, otherwise keep as it is
-					ExistingCompletedState.RunStatus = EStateTreeRunStatus::Failed;
+					ExistingCompletedState.RunStatus = UE::StateTree::ExecutionContext::Private::GetPriorityRunStatus(ExistingCompletedState.RunStatus, NewFinishedTask.RunStatus);
 				}
 				bAdd = false;
 				break;
@@ -1513,11 +1510,7 @@ void FStateTreeExecutionContext::MarkStateCompleted(const UE::StateTree::FFinish
 				}
 				else if (ExistingCompletedState.StateID == NewFinishedTask.StateID)
 				{
-					if (NewFinishedTask.Type == EStateTreeFinishTaskType::Failed)
-					{
-						// Set failure, otherwise keep as it is
-						ExistingCompletedState.RunStatus = EStateTreeRunStatus::Failed;
-					}
+					ExistingCompletedState.RunStatus = UE::StateTree::ExecutionContext::Private::GetPriorityRunStatus(ExistingCompletedState.RunStatus, NewFinishedTask.RunStatus);
 					bAdd = false;
 					break;
 				}
@@ -1527,8 +1520,7 @@ void FStateTreeExecutionContext::MarkStateCompleted(const UE::StateTree::FFinish
 
 	if (bAdd)
 	{
-		const EStateTreeRunStatus RunStatus = UE::StateTree::ExecutionContext::Private::GetRunStatusFromFinishTaskType(NewFinishedTask.Type);
-		CompletedStates.Emplace(NewFinishedTask, RunStatus);
+		CompletedStates.Emplace(NewFinishedTask);
 	}
 }
 
@@ -2548,28 +2540,28 @@ EStateTreeRunStatus FStateTreeExecutionContext::EnterState(FStateTreeTransitionR
 					UE_STATETREE_ENSURE_COMPLETED_STATE_RUN_STATUS();
 					UE_STATETREE_DEBUG_TASK_ENTER_STATE(this, NextFrame.StateTree, FStateTreeIndex16(TaskIndex));
 
-					EStateTreeRunStatus Status = EStateTreeRunStatus::Unset;
+					EStateTreeRunStatus TaskResult = EStateTreeRunStatus::Unset;
 					{
 						QUICK_SCOPE_CYCLE_COUNTER(StateTree_Task_EnterState);
 						CSV_SCOPED_TIMING_STAT_EXCLUSIVE(StateTree_Task_EnterState);
 					
-						Status = Task.EnterState(*this, CurrentTransition);
+						TaskResult = Task.EnterState(*this, CurrentTransition);
 					}
 
 					// Delayed from FinishTask(), use the best priority from what the task returns and FinishTask().
 					if (CompletedStateRunStatus.IsSet())
 					{
-						Status = UE::StateTree::ExecutionContext::Private::GetPriorityRunStatus(Status, CompletedStateRunStatus.GetValue());
+						TaskResult = UE::StateTree::ExecutionContext::Private::GetPriorityRunStatus(TaskResult, CompletedStateRunStatus.GetValue());
 						CompletedStateRunStatus.Reset();
 					}
 					UE_STATETREE_ENSURE_COMPLETED_STATE_RUN_STATUS();
-					UE_STATETREE_DEBUG_TASK_EVENT(this, TaskIndex, TaskInstanceView, EStateTreeTraceEventType::OnEntered, Status);
+					UE_STATETREE_DEBUG_TASK_EVENT(this, TaskIndex, TaskInstanceView, EStateTreeTraceEventType::OnEntered, TaskResult);
 
-					if (Status != EStateTreeRunStatus::Running)
+					if (TaskResult != EStateTreeRunStatus::Running)
 					{
-						Result = Status;
-						FinishTaskInternal(UE::StateTree::FFinishedTask::EReasonType::StateTask, CurrentFrame.FrameID, CurrentFrame.ActiveStates.StateIDs[StateIndex], TaskIndex, Result);
-						if (Status == EStateTreeRunStatus::Failed)
+						Result = UE::StateTree::ExecutionContext::Private::GetPriorityRunStatus(Result, TaskResult);
+						FinishTaskInternal(UE::StateTree::FFinishedTask::EReasonType::StateTask, CurrentFrame.FrameID, CurrentFrame.ActiveStates.StateIDs[StateIndex], TaskIndex, TaskResult);
+						if (TaskResult == EStateTreeRunStatus::Failed)
 						{
 							Exec.EnterStateFailedFrameIndex = FStateTreeIndex16(FrameIndex);
 							Exec.EnterStateFailedTaskIndex = FStateTreeIndex16(TaskIndex);
@@ -3640,7 +3632,7 @@ EStateTreeRunStatus FStateTreeExecutionContext::TickTasks(const float DeltaTime)
 
 					if (TaskResult != EStateTreeRunStatus::Running)
 					{
-						Result = TaskResult;
+						Result = UE::StateTree::ExecutionContext::Private::GetPriorityRunStatus(Result, TaskResult);
 						FinishTaskInternal(UE::StateTree::FFinishedTask::EReasonType::StateTask, CurrentFrame.FrameID, CurrentFrame.ActiveStates.StateIDs[StateIndex], TaskIndex, TaskResult);
 						if (Result == EStateTreeRunStatus::Failed)
 						{
@@ -6020,15 +6012,15 @@ bool FStateTreeWeakExecutionContext::FinishTask(const FStateTreeWeakTaskRef& Tas
 					const int32 GlobalTaskEndIndex = ExecutionFrame->StateTree->GlobalTasksBegin + ExecutionFrame->StateTree->GlobalTasksNum;
 					const bool bIsGlobalTask = TaskPin.GetTaskIndex().AsInt32() >= GlobalTaskBeginIndex && TaskPin.GetTaskIndex().AsInt32() < GlobalTaskEndIndex;
 					const UE::StateTree::FFinishedTask::EReasonType Reason = bIsGlobalTask ? UE::StateTree::FFinishedTask::EReasonType::GlobalTask : UE::StateTree::FFinishedTask::EReasonType::StateTask;
+					const EStateTreeRunStatus RunStatus = UE::StateTree::ExecutionContext::Private::GetRunStatusFromFinishTaskType(FinishType);
 					UE::StateTree::FFinishedTask CompletedTask{
 						.FrameID = FrameID,
 						.StateID = StateID,
 						.TaskIndex = TaskPin.GetTaskIndex(),
-						.Type = FinishType,
+						.RunStatus = RunStatus,
 						.Reason = Reason
 					};
 
-					// Confirm that the task if from the correct state.
 					Exec.FinishedTasks.Add(CompletedTask);
 					ScheduleNextTick(OwnerPinned.Get(), StateTreePinned.Get(), *StoragePinned);
 					bCompleted = true;
diff --git a/Engine/Plugins/Runtime/StateTree/Source/StateTreeModule/Public/StateTreeExecutionContext.h b/Engine/Plugins/Runtime/StateTree/Source/StateTreeModule/Public/StateTreeExecutionContext.h
--- a/Engine/Plugins/Runtime/StateTree/Source/StateTreeModule/Public/StateTreeExecutionContext.h
+++ b/Engine/Plugins/Runtime/StateTree/Source/StateTreeModule/Public/StateTreeExecutionContext.h
@@ -957,15 +957,9 @@ protected:
 	/** Run status set in FinishTask and that used after the task ticks. */
 	TOptional<EStateTreeRunStatus> CompletedStateRunStatus;
 
-	/** Helper struct to store the resolved FStateTreeExecutionState::FCompletedState */
+	/** State that was completed in this frame. */
 	struct FCompletedState : UE::StateTree::FFinishedTask
 	{
-		FCompletedState() = default;
-		FCompletedState(const UE::StateTree::FFinishedTask& Other, EStateTreeRunStatus InRunStatus)
-			: UE::StateTree::FFinishedTask(Other)
-			, RunStatus(InRunStatus)
-		{}
-		EStateTreeRunStatus RunStatus = EStateTreeRunStatus::Unset;
 	};
 
 	/**
diff --git a/Engine/Plugins/Runtime/StateTree/Source/StateTreeModule/Public/StateTreeExecutionTypes.h b/Engine/Plugins/Runtime/StateTree/Source/StateTreeModule/Public/StateTreeExecutionTypes.h
--- a/Engine/Plugins/Runtime/StateTree/Source/StateTreeModule/Public/StateTreeExecutionTypes.h
+++ b/Engine/Plugins/Runtime/StateTree/Source/StateTreeModule/Public/StateTreeExecutionTypes.h
@@ -696,7 +696,7 @@ struct FFinishedTask
 	/** Task that is finished and needs to be processed. */
 	FStateTreeIndex16 TaskIndex = FStateTreeIndex16::Invalid;
 	/** The result of the finished task. */
-	EStateTreeFinishTaskType Type = EStateTreeFinishTaskType::Failed;
+	EStateTreeRunStatus RunStatus = EStateTreeRunStatus::Failed;
 	/** The reason of the finished task. */
 	EReasonType Reason = EReasonType::GlobalTask;
 };
