Fix compilation issue introduced with CL 40031258

diff --git a/Engine/Plugins/Runtime/StateTree/Source/StateTreeModule/Private/StateTreeExecutionContext.cpp b/Engine/Plugins/Runtime/StateTree/Source/StateTreeModule/Private/StateTreeExecutionContext.cpp
--- a/Engine/Plugins/Runtime/StateTree/Source/StateTreeModule/Private/StateTreeExecutionContext.cpp
+++ b/Engine/Plugins/Runtime/StateTree/Source/StateTreeModule/Private/StateTreeExecutionContext.cpp
@@ -1520,7 +1520,7 @@ void FStateTreeExecutionContext::MarkStateCompleted(const UE::StateTree::FFinish
 
 	if (bAdd)
 	{
-		CompletedStates.Emplace(NewFinishedTask);
+		CompletedStates.Add(FCompletedState(NewFinishedTask));
 	}
 }
 
