Fix compilation issue introduced with CL 40031258

diff --git a/Engine/Plugins/Runtime/StateTree/Source/StateTreeModule/Public/StateTreeExecutionContext.h b/Engine/Plugins/Runtime/StateTree/Source/StateTreeModule/Public/StateTreeExecutionContext.h
--- a/Engine/Plugins/Runtime/StateTree/Source/StateTreeModule/Public/StateTreeExecutionContext.h
+++ b/Engine/Plugins/Runtime/StateTree/Source/StateTreeModule/Public/StateTreeExecutionContext.h
@@ -960,6 +960,11 @@ protected:
 	/** State that was completed in this frame. */
 	struct FCompletedState : UE::StateTree::FFinishedTask
 	{
+		FCompletedState() = default;
+		FCompletedState(UE::StateTree::FFinishedTask Other)
+			: UE::StateTree::FFinishedTask(Other)
+		{
+		};
 	};
 
 	/**
