[State Tree Editor] Fixed typos in State and Node Height tooltips in a State View settings

diff --git a/Engine/Plugins/Runtime/StateTree/Source/StateTreeEditorModule/Private/StateTreeEditorUserSettings.h b/Engine/Plugins/Runtime/StateTree/Source/StateTreeEditorModule/Private/StateTreeEditorUserSettings.h
--- a/Engine/Plugins/Runtime/StateTree/Source/StateTreeEditorModule/Private/StateTreeEditorUserSettings.h
+++ b/Engine/Plugins/Runtime/StateTree/Source/StateTreeEditorModule/Private/StateTreeEditorUserSettings.h
@@ -52,11 +52,11 @@ private:
 	UPROPERTY(config, EditAnywhere, Category = "State View", meta = (DisplayName = "Display Node"))
 	EStateTreeEditorUserSettingsNodeType StatesViewDisplayNodeType = EStateTreeEditorUserSettingsNodeType::All;
 	
-	/** Height of the a state in the Sates View. */
+	/** Height of a state in the States View. */
 	UPROPERTY(config, EditAnywhere, Category = "State View", meta = (DisplayName = "State Height", ClampMin = 8.0f))
 	float StatesViewStateRowHeight = 32.0f;
 	
-	/** Height of the a node in the Sates View. */
+	/** Height of a node in the States View. */
 	UPROPERTY(config, EditAnywhere, Category = "State View", meta = (DisplayName = "Node Height", ClampMin = 8.0f))
 	float StatesViewNodeRowHeight = 16.0f;
 };
