[StateTreeDebugger] added SET_NODE_CUSTOM_TRACE_TEXT to append or override a debug string to the node instance data. Updated some existing conditions.

diff --git a/Engine/Plugins/Runtime/StateTree/Source/StateTreeEditorModule/Private/Debugger/SStateTreeDebuggerView.cpp b/Engine/Plugins/Runtime/StateTree/Source/StateTreeEditorModule/Private/Debugger/SStateTreeDebuggerView.cpp
--- a/Engine/Plugins/Runtime/StateTree/Source/StateTreeEditorModule/Private/Debugger/SStateTreeDebuggerView.cpp
+++ b/Engine/Plugins/Runtime/StateTree/Source/StateTreeEditorModule/Private/Debugger/SStateTreeDebuggerView.cpp
@@ -979,33 +979,39 @@ void SStateTreeDebuggerView::GenerateElementsForProperties(const FStateTreeTrace
 {
 	FString TypePath;
 	FString InstanceDataAsString;
+	FString DebugText;
 
-	Visit([&TypePath, &InstanceDataAsString](auto& TypedEvent)
+	Visit([&TypePath, &InstanceDataAsString, &DebugText](auto& TypedEvent)
 		{
 			TypePath = TypedEvent.GetDataTypePath();
 			InstanceDataAsString = TypedEvent.GetDataAsText();
+			DebugText = TypedEvent.GetDebugText();
 		}, Event);
 
+	auto CreateTreeElement = [ParentElement]<typename EventType>(const FString& Line)
+		{
+			// Create new event
+			EventType Event(/*RecordingWorldTime*/0, ELogVerbosity::Verbose, *Line);
+
+			// Create Tree element to hold the event
+			const TSharedPtr<FStateTreeDebuggerEventTreeElement> NewChildElement = MakeShareable(new FStateTreeDebuggerEventTreeElement(
+				ParentElement->Frame,
+				FStateTreeTraceEventVariantType(TInPlaceType<EventType>(), Event),
+				ParentElement->WeakStateTree.Get()));
+
+			ParentElement->Children.Add(NewChildElement);
+		};
+
 	if (!InstanceDataAsString.IsEmpty())
 	{
-		auto CreatePropertyElement = [ParentElement](const FStringView Line, const int32 NestedCount = 0)
+		auto CreatePropertyElement = [ParentElement, &CreateTreeElement](const FStringView Line, const int32 NestedCount = 0)
 			{
 				constexpr int32 Indent = 4;
-				FString ConvertedString = FString::Printf(TEXT("%*s"), NestedCount*Indent, TEXT(""));
+				FString ConvertedString = FString::Printf(TEXT("%*s"), NestedCount * Indent, TEXT(""));
 				ConvertedString.Append(Line);
 				ConvertedString.ReplaceInline(TEXT("="), TEXT(" = "));
 				ConvertedString.ReplaceInline(TEXT("\""), TEXT(""));
-
-				// Create new property event
-				FStateTreeTracePropertyEvent PropertyEvent(/*RecordingWorldTime*/0, *ConvertedString);
-
-				// Create Tree element to hold the event
-				const TSharedPtr<FStateTreeDebuggerEventTreeElement> NewChildElement = MakeShareable(new FStateTreeDebuggerEventTreeElement(
-					ParentElement->Frame,
-					FStateTreeTraceEventVariantType(TInPlaceType<FStateTreeTracePropertyEvent>(), PropertyEvent),
-					ParentElement->WeakStateTree.Get()));
-
-				ParentElement->Children.Add(NewChildElement);
+				CreateTreeElement.template operator()<FStateTreeTracePropertyEvent>(ConvertedString);
 			};
 
 		// Try to parse Struct for which properties are exported between '(' and ')' 
@@ -1062,6 +1068,11 @@ void SStateTreeDebuggerView::GenerateElementsForProperties(const FStateTreeTrace
 			}
 		}
 	}
+
+	if (!DebugText.IsEmpty())
+	{
+		CreateTreeElement.template operator()<FStateTreeTraceLogEvent>(DebugText);
+	}
 }
 
 void SStateTreeDebuggerView::ExpandAll(const TArray<TSharedPtr<FStateTreeDebuggerEventTreeElement>>& Items)
@@ -1071,11 +1082,18 @@ void SStateTreeDebuggerView::ExpandAll(const TArray<TSharedPtr<FStateTreeDebugge
 		bool bExpand = true;
 		if (Item->Children.Num() > 0)
 		{
-			const FStateTreeDebuggerEventTreeElement* FirstChild = Item->Children[0].Get();
-			if (FirstChild && FirstChild->Event.IsType<FStateTreeTracePropertyEvent>())
+			if (const FStateTreeDebuggerEventTreeElement* FirstChild = Item->Children[0].Get())
 			{
-				bExpand = false;
-			}	
+				if (const FStateTreeTraceLogEvent* LogEvent = FirstChild->Event.TryGet<FStateTreeTraceLogEvent>())
+				{
+					// Do not auto expand verbose logs
+					bExpand = (LogEvent->Verbosity < ELogVerbosity::Verbose);
+				}
+				else if (FirstChild->Event.IsType<FStateTreeTracePropertyEvent>())
+				{
+					bExpand = false;
+				}
+			}
 		}
 
 		if (bExpand)
diff --git a/Engine/Plugins/Runtime/StateTree/Source/StateTreeEditorModule/Private/Debugger/SStateTreeDebuggerViewRow.cpp b/Engine/Plugins/Runtime/StateTree/Source/StateTreeEditorModule/Private/Debugger/SStateTreeDebuggerViewRow.cpp
--- a/Engine/Plugins/Runtime/StateTree/Source/StateTreeEditorModule/Private/Debugger/SStateTreeDebuggerViewRow.cpp
+++ b/Engine/Plugins/Runtime/StateTree/Source/StateTreeEditorModule/Private/Debugger/SStateTreeDebuggerViewRow.cpp
@@ -171,11 +171,20 @@ const FTextBlockStyle& SStateTreeDebuggerViewRow::GetEventTextStyle() const
 		return StyleSet.GetWidgetStyle<FTextBlockStyle>("StateTreeDebugger.Element.Bold");
 	}
 
-	 if (Item->Event.IsType<FStateTreeTracePropertyEvent>())
+	if (Item->Event.IsType<FStateTreeTracePropertyEvent>())
 	{
 		return StyleSet.GetWidgetStyle<FTextBlockStyle>("StateTreeDebugger.Element.Subdued");
 	}
 
+	if (const FStateTreeTraceLogEvent* LogEvent = Item->Event.TryGet<FStateTreeTraceLogEvent>())
+	{
+		// Make verbose logs more subtle
+		if (LogEvent->Verbosity >= ELogVerbosity::Verbose)
+		{
+			return StyleSet.GetWidgetStyle<FTextBlockStyle>("StateTreeDebugger.Element.Subdued");
+		}
+	}
+
 	return StyleSet.GetWidgetStyle<FTextBlockStyle>("StateTreeDebugger.Element.Normal");
 }
 
diff --git a/Engine/Plugins/Runtime/StateTree/Source/StateTreeModule/Private/Conditions/StateTreeCommonConditions.cpp b/Engine/Plugins/Runtime/StateTree/Source/StateTreeModule/Private/Conditions/StateTreeCommonConditions.cpp
--- a/Engine/Plugins/Runtime/StateTree/Source/StateTreeModule/Private/Conditions/StateTreeCommonConditions.cpp
+++ b/Engine/Plugins/Runtime/StateTree/Source/StateTreeModule/Private/Conditions/StateTreeCommonConditions.cpp
@@ -54,6 +54,13 @@ bool FStateTreeCompareIntCondition::TestCondition(FStateTreeExecutionContext& Co
 	const FInstanceDataType& InstanceData = Context.GetInstanceData(*this);
 
 	const bool bResult = UE::StateTree::Conditions::CompareNumbers<int32>(InstanceData.Left, InstanceData.Right, Operator);
+
+	SET_NODE_CUSTOM_TRACE_TEXT(Context, Override, TEXT("%s%s %s %s")
+		, *UE::StateTree::DescHelpers::GetInvertText(bInvert, EStateTreeNodeFormatting::Text).ToString()
+		, *LexToString(InstanceData.Left)
+		, *UE::StateTree::DescHelpers::GetOperatorText(Operator, EStateTreeNodeFormatting::Text).ToString()
+		, *LexToString(InstanceData.Right));
+
 	return bResult ^ bInvert;
 }
 
@@ -95,6 +102,13 @@ bool FStateTreeCompareFloatCondition::TestCondition(FStateTreeExecutionContext&
 	const FInstanceDataType& InstanceData = Context.GetInstanceData(*this);
 
 	const bool bResult = UE::StateTree::Conditions::CompareNumbers<double>(InstanceData.Left, InstanceData.Right, Operator);
+
+	SET_NODE_CUSTOM_TRACE_TEXT(Context, Override, TEXT("%s%s %s %s")
+		, *UE::StateTree::DescHelpers::GetInvertText(bInvert, EStateTreeNodeFormatting::Text).ToString()
+		, *LexToString(InstanceData.Left)
+		, *UE::StateTree::DescHelpers::GetOperatorText(Operator, EStateTreeNodeFormatting::Text).ToString()
+		, *LexToString(InstanceData.Right));
+
 	return bResult ^ bInvert;
 }
 
@@ -138,7 +152,12 @@ FText FStateTreeCompareFloatCondition::GetDescription(const FGuid& ID, FStateTre
 bool FStateTreeCompareBoolCondition::TestCondition(FStateTreeExecutionContext& Context) const
 {
 	const FInstanceDataType& InstanceData = Context.GetInstanceData(*this);
-	
+
+	SET_NODE_CUSTOM_TRACE_TEXT(Context, Override, TEXT("%s%s is %s")
+		, *UE::StateTree::DescHelpers::GetInvertText(bInvert, EStateTreeNodeFormatting::Text).ToString()
+		, *UE::StateTree::DescHelpers::GetBoolText(InstanceData.bLeft, EStateTreeNodeFormatting::Text).ToString()
+		, *UE::StateTree::DescHelpers::GetBoolText(InstanceData.bRight, EStateTreeNodeFormatting::Text).ToString());
+
 	return (InstanceData.bLeft == InstanceData.bRight) ^ bInvert;
 }
 
@@ -181,6 +200,11 @@ bool FStateTreeCompareEnumCondition::TestCondition(FStateTreeExecutionContext& C
 {
 	const FInstanceDataType& InstanceData = Context.GetInstanceData(*this);
 
+	SET_NODE_CUSTOM_TRACE_TEXT(Context, Override, TEXT("%s%s is %s")
+		, *UE::StateTree::DescHelpers::GetInvertText(bInvert, EStateTreeNodeFormatting::Text).ToString()
+		, *InstanceData.Left.Enum->GetNameStringByValue(InstanceData.Left.Value)
+		, *InstanceData.Right.Enum->GetNameStringByValue(InstanceData.Right.Value));
+
 	return (InstanceData.Left == InstanceData.Right) ^ bInvert;
 }
 
@@ -285,6 +309,15 @@ bool FStateTreeCompareDistanceCondition::TestCondition(FStateTreeExecutionContex
 	const FVector::FReal Left = FVector::DistSquared(InstanceData.Source, InstanceData.Target);
 	const FVector::FReal Right = FMath::Square(InstanceData.Distance);
 	const bool bResult = UE::StateTree::Conditions::CompareNumbers<FVector::FReal>(Left, Right, Operator);
+
+	SET_NODE_CUSTOM_TRACE_TEXT(Context, Override, TEXT("%sDistance %s %s %s (from [%s] to [%s])")
+		, *UE::StateTree::DescHelpers::GetInvertText(bInvert, EStateTreeNodeFormatting::Text).ToString()
+		, *LexToString(FMath::Sqrt(Left))
+		, *UE::StateTree::DescHelpers::GetOperatorText(Operator, EStateTreeNodeFormatting::Text).ToString()
+		, *LexToString(InstanceData.Distance)
+		, *InstanceData.Source.ToString()
+		, *InstanceData.Target.ToString());
+
 	return bResult ^ bInvert;
 }
 
@@ -339,8 +372,13 @@ FText FStateTreeCompareDistanceCondition::GetDescription(const FGuid& ID, FState
 bool FStateTreeRandomCondition::TestCondition(FStateTreeExecutionContext& Context) const
 {
 	const FInstanceDataType& InstanceData = Context.GetInstanceData(*this);
+	const float RandomValue = FMath::FRandRange(0.0f, 1.0f);
+
+	SET_NODE_CUSTOM_TRACE_TEXT(Context, Override, TEXT("Random value %f < %s")
+		, RandomValue
+		, *LexToString(InstanceData.Threshold));
 
-	return FMath::FRandRange(0.0f, 1.0f) < InstanceData.Threshold;
+	return RandomValue < InstanceData.Threshold;
 }
 
 #if WITH_EDITOR
diff --git a/Engine/Plugins/Runtime/StateTree/Source/StateTreeModule/Private/Conditions/StateTreeGameplayTagConditions.cpp b/Engine/Plugins/Runtime/StateTree/Source/StateTreeModule/Private/Conditions/StateTreeGameplayTagConditions.cpp
--- a/Engine/Plugins/Runtime/StateTree/Source/StateTreeModule/Private/Conditions/StateTreeGameplayTagConditions.cpp
+++ b/Engine/Plugins/Runtime/StateTree/Source/StateTreeModule/Private/Conditions/StateTreeGameplayTagConditions.cpp
@@ -25,7 +25,13 @@ namespace UE::StateTree::Conditions
 bool FGameplayTagMatchCondition::TestCondition(FStateTreeExecutionContext& Context) const
 {
 	const FInstanceDataType& InstanceData = Context.GetInstanceData(*this);
-	
+
+	SET_NODE_CUSTOM_TRACE_TEXT(Context, Override, TEXT("%s'%s' contains '%s%s'")
+		, *UE::StateTree::DescHelpers::GetInvertText(bInvert, EStateTreeNodeFormatting::Text).ToString()
+		, *UE::StateTree::DescHelpers::GetGameplayTagContainerAsText(InstanceData.TagContainer).ToString()
+		, *UE::StateTree::DescHelpers::GetExactMatchText(bExactMatch, EStateTreeNodeFormatting::Text).ToString()
+		, *InstanceData.Tag.ToString());
+
 	return (bExactMatch ? InstanceData.TagContainer.HasTagExact(InstanceData.Tag) : InstanceData.TagContainer.HasTag(InstanceData.Tag)) ^ bInvert;
 }
 
@@ -82,7 +88,14 @@ bool FGameplayTagContainerMatchCondition::TestCondition(FStateTreeExecutionConte
 	default:
 		ensureMsgf(false, TEXT("Unhandled match type %s."), *UEnum::GetValueAsString(MatchType));
 	}
-	
+
+	SET_NODE_CUSTOM_TRACE_TEXT(Context, Override, TEXT("%s'%s' contains '%s %s%s'")
+		, *UE::StateTree::DescHelpers::GetInvertText(bInvert, EStateTreeNodeFormatting::Text).ToString()
+		, *UE::StateTree::DescHelpers::GetGameplayTagContainerAsText(InstanceData.TagContainer).ToString()
+		, *UEnum::GetValueAsString(MatchType)
+		, *UE::StateTree::DescHelpers::GetExactMatchText(bExactMatch, EStateTreeNodeFormatting::Text).ToString()
+		, *UE::StateTree::DescHelpers::GetGameplayTagContainerAsText(InstanceData.OtherContainer).ToString());
+
 	return bResult ^ bInvert;
 }
 
@@ -129,6 +142,11 @@ bool FGameplayTagQueryCondition::TestCondition(FStateTreeExecutionContext& Conte
 {
 	const FInstanceDataType& InstanceData = Context.GetInstanceData(*this);
 
+	SET_NODE_CUSTOM_TRACE_TEXT(Context, Override, TEXT("%s'%s' matches %s")
+		, *UE::StateTree::DescHelpers::GetInvertText(bInvert, EStateTreeNodeFormatting::Text).ToString()
+		, *UE::StateTree::DescHelpers::GetGameplayTagContainerAsText(InstanceData.TagContainer).ToString()
+		, *UE::StateTree::DescHelpers::GetGameplayTagQueryAsText(TagQuery).ToString());
+
 	return TagQuery.Matches(InstanceData.TagContainer) ^ bInvert;
 }
 
diff --git a/Engine/Plugins/Runtime/StateTree/Source/StateTreeModule/Private/Conditions/StateTreeObjectConditions.cpp b/Engine/Plugins/Runtime/StateTree/Source/StateTreeModule/Private/Conditions/StateTreeObjectConditions.cpp
--- a/Engine/Plugins/Runtime/StateTree/Source/StateTreeModule/Private/Conditions/StateTreeObjectConditions.cpp
+++ b/Engine/Plugins/Runtime/StateTree/Source/StateTreeModule/Private/Conditions/StateTreeObjectConditions.cpp
@@ -64,6 +64,13 @@ bool FStateTreeObjectIsChildOfClassCondition::TestCondition(FStateTreeExecutionC
 	const FInstanceDataType& InstanceData = Context.GetInstanceData(*this);
 	const bool bResult = InstanceData.Object && InstanceData.Class
 						&& InstanceData.Object->GetClass()->IsChildOf(InstanceData.Class);
+
+	SET_NODE_CUSTOM_TRACE_TEXT(Context, Override, TEXT("%s of type '%s' is%s child of '%s'")
+		, *GetNameSafe(InstanceData.Object)
+		, *GetNameSafe(InstanceData.Object ? InstanceData.Object->GetClass() : nullptr)
+		, bResult ? TEXT("") : TEXT(" not")
+		, *GetNameSafe(InstanceData.Class));
+
 	return bResult ^ bInvert;
 }
 
diff --git a/Engine/Plugins/Runtime/StateTree/Source/StateTreeModule/Private/Debugger/StateTreeTrace.cpp b/Engine/Plugins/Runtime/StateTree/Source/StateTreeModule/Private/Debugger/StateTreeTrace.cpp
--- a/Engine/Plugins/Runtime/StateTree/Source/StateTreeModule/Private/Debugger/StateTreeTrace.cpp
+++ b/Engine/Plugins/Runtime/StateTree/Source/StateTreeModule/Private/Debugger/StateTreeTrace.cpp
@@ -610,7 +610,7 @@ void TraceBufferedEvents(const FStateTreeInstanceDebugId InstanceId)
 	GBufferedEvents.Flush_AnyThread(InstanceId);
 }
 
-void SerializeDataViewToArchive(FBufferArchive& Ar, const FStateTreeDataView DataView)
+void SerializeDebugDataToArchive(FBufferArchive& Ar, FNodeCustomDebugData&& CustomDebugData, const FStateTreeDataView DataView)
 {
 	constexpr uint32 PortFlags = 
 		PPF_PropertyWindow // limit to properties visible in Editor 
@@ -620,31 +620,45 @@ void SerializeDataViewToArchive(FBufferArchive& Ar, const FStateTreeDataView Dat
 		| PPF_SimpleObjectText // object property values should be exported without the package or class information
 		| PPF_ForDiff; // do not emit object path
 
+	FString TypePath;
+	FString InstanceDataAsText;
+	FString DebugText;
+
 	if (const UScriptStruct* ScriptStruct = Cast<const UScriptStruct>(DataView.GetStruct()))
 	{
-		TRACE_CPUPROFILER_EVENT_SCOPE(UE::StateTree::ExportStructAsText)
-		FString StructPath = ScriptStruct->GetPathName();
-		FString TextValue;
-
-		ScriptStruct->ExportText(TextValue, DataView.GetMemory(), DataView.GetMemory(), /*OwnerObject*/nullptr, PortFlags | PPF_SeparateDefine, /*ExportRootScope*/nullptr);
-
-		Ar << StructPath;
-		Ar << TextValue;
+		TRACE_CPUPROFILER_EVENT_SCOPE(UE::StateTree::ExportStructAsText);
+		TypePath = ScriptStruct->GetPathName();
+		
+		if (!CustomDebugData.ShouldOverrideDataView())
+		{
+			ScriptStruct->ExportText(InstanceDataAsText, DataView.GetMemory(), DataView.GetMemory(), /*OwnerObject*/nullptr, PortFlags | PPF_SeparateDefine, /*ExportRootScope*/nullptr);
+		}
 	}
 	else if (const UClass* Class = Cast<const UClass>(DataView.GetStruct()))
 	{
-		TRACE_CPUPROFILER_EVENT_SCOPE(UE::StateTree::ExportObjectAsText)
-		FString StructPath = Class->GetPathName();
-		FStringOutputDevice OutputDevice;
-		UObject* Object = DataView.GetMutablePtr<UObject>();
+		TRACE_CPUPROFILER_EVENT_SCOPE(UE::StateTree::ExportObjectAsText);
+		TypePath = Class->GetPathName();
+
+		if (!CustomDebugData.ShouldOverrideDataView())
+		{
+			FStringOutputDevice OutputDevice;
+			UObject* Object = DataView.GetMutablePtr<UObject>();
 
-		// Not using on scope FExportObjectInnerContext since it is very costly to build.
-		// Passing a null context will make the export use an already built thread local context.
-		UExporter::ExportToOutputDevice(nullptr, Object, /*Exporter*/nullptr, OutputDevice, TEXT("copy"), 0, PortFlags, false, Object->GetOuter());
+			// Not using on scope FExportObjectInnerContext since it is very costly to build.
+			// Passing a null context will make the export use an already built thread local context.
+			UExporter::ExportToOutputDevice(nullptr, Object, /*Exporter*/nullptr, OutputDevice, TEXT("copy"), 0, PortFlags, false, Object->GetOuter());
+			InstanceDataAsText = OutputDevice;
+		}
+	}
 
-		Ar << StructPath;
-		Ar << OutputDevice;
+	if (CustomDebugData.IsSet())
+	{
+		DebugText = CustomDebugData.GetTraceDebuggerString();
 	}
+
+	Ar << TypePath;
+	Ar << InstanceDataAsText;
+	Ar << DebugText;
 }
 
 void RegisterGlobalDelegates()
@@ -807,14 +821,12 @@ void OutputStateEventTrace(
 
 void OutputTaskEventTrace(
 	const FStateTreeInstanceDebugId InstanceId,
-	const FStateTreeIndex16 TaskIdx,
-	const FStateTreeDataView DataView,
-	const EStateTreeTraceEventType EventType,
-	const EStateTreeRunStatus Status
-	)
+	FNodeCustomDebugData&& CustomDebugData,
+	FStateTreeIndex16 TaskIdx, FStateTreeDataView DataView,
+	EStateTreeTraceEventType EventType, EStateTreeRunStatus Status)
 {
 	FBufferArchive Archive;
-	SerializeDataViewToArchive(Archive, DataView);
+	SerializeDebugDataToArchive(Archive, MoveTemp(CustomDebugData), DataView);
 
 	TraceBufferedEvents(InstanceId);
 
@@ -830,13 +842,13 @@ void OutputTaskEventTrace(
 
 void OutputEvaluatorEventTrace(
 	const FStateTreeInstanceDebugId InstanceId,
-	const FStateTreeIndex16 EvaluatorIdx,
-	const FStateTreeDataView DataView,
-	const EStateTreeTraceEventType EventType
-	)
+	FNodeCustomDebugData&& CustomDebugData,
+	FStateTreeIndex16 EvaluatorIdx,
+	FStateTreeDataView DataView,
+	EStateTreeTraceEventType EventType)
 {
 	FBufferArchive Archive;
-	SerializeDataViewToArchive(Archive, DataView);
+	SerializeDebugDataToArchive(Archive, MoveTemp(CustomDebugData), DataView);
 
 	TraceBufferedEvents(InstanceId);
 
@@ -879,13 +891,13 @@ void OutputTransitionEventTrace(
 
 void OutputConditionEventTrace(
 	const FStateTreeInstanceDebugId InstanceId,
-	const FStateTreeIndex16 ConditionIdx,
-	const FStateTreeDataView DataView,
-	const EStateTreeTraceEventType EventType
-	)
+	FNodeCustomDebugData&& CustomDebugData,
+	FStateTreeIndex16 ConditionIdx,
+	FStateTreeDataView DataView,
+	EStateTreeTraceEventType EventType)
 {
 	FBufferArchive Archive;
-	SerializeDataViewToArchive(Archive, DataView);
+	SerializeDebugDataToArchive(Archive, MoveTemp(CustomDebugData), DataView);
 
 	TraceBufferedEvents(InstanceId);
 
diff --git a/Engine/Plugins/Runtime/StateTree/Source/StateTreeModule/Private/Debugger/StateTreeTraceAnalyzer.cpp b/Engine/Plugins/Runtime/StateTree/Source/StateTreeModule/Private/Debugger/StateTreeTraceAnalyzer.cpp
--- a/Engine/Plugins/Runtime/StateTree/Source/StateTreeModule/Private/Debugger/StateTreeTraceAnalyzer.cpp
+++ b/Engine/Plugins/Runtime/StateTree/Source/StateTreeModule/Private/Debugger/StateTreeTraceAnalyzer.cpp
@@ -147,16 +147,17 @@ bool FStateTreeTraceAnalyzer::OnEvent(const uint16 RouteId, EStyle Style, const
 		}
 	case RouteId_Task:
 		{
-			FString TypePath, DataAsText;
+			FString TypePath, DataAsText, DebugText;
 			FMemoryReaderView Archive(EventData.GetArrayView<uint8>("DataView"));
 			Archive << TypePath;
 			Archive << DataAsText;
+			Archive << DebugText;
 
 			const FStateTreeTraceTaskEvent Event(WorldTime,
 				FStateTreeIndex16(EventData.GetValue<uint16>("NodeIndex")),
 				EventData.GetValue<EStateTreeTraceEventType>("EventType"),
 				EventData.GetValue<EStateTreeRunStatus>("Status"),
-				TypePath, DataAsText);
+				TypePath, DataAsText, DebugText);
 
 			Provider.AppendEvent(FStateTreeInstanceDebugId(EventData.GetValue<uint32>("InstanceId"), EventData.GetValue<uint32>("InstanceSerial")),
 				Context.EventTime.AsSeconds(EventData.GetValue<uint64>("Cycle")),
@@ -165,15 +166,16 @@ bool FStateTreeTraceAnalyzer::OnEvent(const uint16 RouteId, EStyle Style, const
 		}
 	case RouteId_Evaluator:
 		{
-			FString TypePath, DataAsText;
+			FString TypePath, DataAsText, DebugText;
 			FMemoryReaderView Archive(EventData.GetArrayView<uint8>("DataView"));
 			Archive << TypePath;
 			Archive << DataAsText;
+			Archive << DebugText;
 
 			const FStateTreeTraceEvaluatorEvent Event(WorldTime,
 				FStateTreeIndex16(EventData.GetValue<uint16>("NodeIndex")),
 				EventData.GetValue<EStateTreeTraceEventType>("EventType"),
-				TypePath, DataAsText);
+				TypePath, DataAsText, DebugText);
 
 			Provider.AppendEvent(FStateTreeInstanceDebugId(EventData.GetValue<uint32>("InstanceId"), EventData.GetValue<uint32>("InstanceSerial")),
 				Context.EventTime.AsSeconds(EventData.GetValue<uint64>("Cycle")),
@@ -182,15 +184,16 @@ bool FStateTreeTraceAnalyzer::OnEvent(const uint16 RouteId, EStyle Style, const
 		}
 	case RouteId_Condition:
 		{
-			FString TypePath, DataAsText;
+			FString TypePath, DataAsText, DebugText;
 			FMemoryReaderView Archive(EventData.GetArrayView<uint8>("DataView"));
 			Archive << TypePath;
 			Archive << DataAsText;
+			Archive << DebugText;
 
 			const FStateTreeTraceConditionEvent Event(WorldTime,
 				FStateTreeIndex16(EventData.GetValue<uint16>("NodeIndex")),
 				EventData.GetValue<EStateTreeTraceEventType>("EventType"),
-				TypePath, DataAsText);
+				TypePath, DataAsText, DebugText);
 			
 			Provider.AppendEvent(FStateTreeInstanceDebugId(EventData.GetValue<uint32>("InstanceId"), EventData.GetValue<uint32>("InstanceSerial")),
 				Context.EventTime.AsSeconds(EventData.GetValue<uint64>("Cycle")),
diff --git a/Engine/Plugins/Runtime/StateTree/Source/StateTreeModule/Private/StateTreeNodeDescriptionHelpers.cpp b/Engine/Plugins/Runtime/StateTree/Source/StateTreeModule/Private/StateTreeNodeDescriptionHelpers.cpp
--- a/Engine/Plugins/Runtime/StateTree/Source/StateTreeModule/Private/StateTreeNodeDescriptionHelpers.cpp
+++ b/Engine/Plugins/Runtime/StateTree/Source/StateTreeModule/Private/StateTreeNodeDescriptionHelpers.cpp
@@ -4,7 +4,7 @@
 #include "GameplayTagContainer.h"
 #include "StateTreeNodeBase.h"
 
-#if WITH_EDITOR
+#if WITH_EDITOR || WITH_STATETREE_TRACE
 #define LOCTEXT_NAMESPACE "StateTree"
 
 namespace UE::StateTree::DescHelpers
@@ -41,13 +41,13 @@ FText GetOperatorText(const EGenericAICheck Operator, EStateTreeNodeFormatting F
 		case EGenericAICheck::NotEqual:
 			return FText::FromString(TEXT("!="));
 		case EGenericAICheck::Less:
-			return FText::FromString(TEXT("&lt;"));
+			return FText::FromString(TEXT("<"));
 		case EGenericAICheck::LessOrEqual:
-			return FText::FromString(TEXT("&lt;="));
+			return FText::FromString(TEXT("<="));
 		case EGenericAICheck::Greater:
-			return FText::FromString(TEXT("&gt;"));
+			return FText::FromString(TEXT(">"));
 		case EGenericAICheck::GreaterOrEqual:
-			return FText::FromString(TEXT("&gt;="));
+			return FText::FromString(TEXT(">="));
 		default:
 			break;
 		}
@@ -220,4 +220,4 @@ FText GetSingleParamFunctionText(const FText& FunctionText, const FText& ParamTe
 } // UE::StateTree::Helpers
 
 #undef LOCTEXT_NAMESPACE
-#endif // WITH_EDITOR
\ No newline at end of file
+#endif // WITH_EDITOR || WITH_STATETREE_TRACE
\ No newline at end of file
diff --git a/Engine/Plugins/Runtime/StateTree/Source/StateTreeModule/Public/Debugger/StateTreeDebug.h b/Engine/Plugins/Runtime/StateTree/Source/StateTreeModule/Public/Debugger/StateTreeDebug.h
--- a/Engine/Plugins/Runtime/StateTree/Source/StateTreeModule/Public/Debugger/StateTreeDebug.h
+++ b/Engine/Plugins/Runtime/StateTree/Source/StateTreeModule/Public/Debugger/StateTreeDebug.h
@@ -44,13 +44,12 @@ enum class EStateTreeTraceEventType : uint8;
 		TRACE_STATETREE_TRANSITION_EVENT((InExecutionContextPtr)->GetInstanceDebugId(), (InTransitionSource), (InEventType))
 
 #define UE_STATETREE_DEBUG_LOG_EVENT(InExecutionContextPtr, InLogVerbosity, InFormat, ...) TRACE_STATETREE_LOG_EVENT((InExecutionContextPtr)->GetInstanceDebugId(), InLogVerbosity, InFormat, ##__VA_ARGS__)
-#define UE_STATETREE_DEBUG_CONDITION_EVENT(InExecutionContextPtr, InIndex, InDataView, InEventType) TRACE_STATETREE_CONDITION_EVENT((InExecutionContextPtr)->GetInstanceDebugId(), FStateTreeIndex16(InIndex), (InDataView), (InEventType));
+#define UE_STATETREE_DEBUG_CONDITION_EVENT(InExecutionContextPtr, InIndex, InDataView, InEventType) TRACE_STATETREE_CONDITION_EVENT((InExecutionContextPtr)->GetInstanceDebugId(), (InExecutionContextPtr)->StealNodeCustomDebugTraceData(), FStateTreeIndex16(InIndex), (InDataView), (InEventType));
 #define UE_STATETREE_DEBUG_INSTANCE_EVENT(InExecutionContextPtr, InEventType) TRACE_STATETREE_INSTANCE_EVENT((InExecutionContextPtr)->GetInstanceDebugId(), (InExecutionContextPtr)->GetStateTree(), *((InExecutionContextPtr)->GetInstanceDescriptionInternal()), (InEventType));
 #define UE_STATETREE_DEBUG_INSTANCE_FRAME_EVENT(InExecutionContextPtr, InFrame) TRACE_STATETREE_INSTANCE_FRAME_EVENT((InExecutionContextPtr)->GetInstanceDebugId(), (InFrame));
 #define UE_STATETREE_DEBUG_ACTIVE_STATES_EVENT(InExecutionContextPtr, InActiveFrames) TRACE_STATETREE_ACTIVE_STATES_EVENT((InExecutionContextPtr)->GetInstanceDebugId(), InActiveFrames);
-#define UE_STATETREE_DEBUG_TASK_EVENT(InExecutionContextPtr, InIndex, InDataView, InEventType, InStatus) TRACE_STATETREE_TASK_EVENT((InExecutionContextPtr)->GetInstanceDebugId(), FStateTreeIndex16(InIndex), (InDataView), (InEventType), (InStatus));
-#define UE_STATETREE_DEBUG_EVALUATOR_EVENT(InExecutionContextPtr, InIndex, InDataView, InEventType) TRACE_STATETREE_EVALUATOR_EVENT((InExecutionContextPtr)->GetInstanceDebugId(), FStateTreeIndex16(InIndex), (InDataView), (InEventType));
-#define UE_STATETREE_DEBUG_CONDITION_EVENT(InExecutionContextPtr, InIndex, InDataView, InEventType) TRACE_STATETREE_CONDITION_EVENT((InExecutionContextPtr)->GetInstanceDebugId(), FStateTreeIndex16(InIndex), (InDataView), (InEventType));
+#define UE_STATETREE_DEBUG_TASK_EVENT(InExecutionContextPtr, InIndex, InDataView, InEventType, InStatus) TRACE_STATETREE_TASK_EVENT((InExecutionContextPtr)->GetInstanceDebugId(), (InExecutionContextPtr)->StealNodeCustomDebugTraceData(), FStateTreeIndex16(InIndex), (InDataView), (InEventType), (InStatus));
+#define UE_STATETREE_DEBUG_EVALUATOR_EVENT(InExecutionContextPtr, InIndex, InDataView, InEventType) TRACE_STATETREE_EVALUATOR_EVENT((InExecutionContextPtr)->GetInstanceDebugId(), (InExecutionContextPtr)->StealNodeCustomDebugTraceData(), FStateTreeIndex16(InIndex), (InDataView), (InEventType));
 
 #define UE_STATETREE_ID_NAME PREPROCESSOR_JOIN(InstanceId,__LINE__) \
 
diff --git a/Engine/Plugins/Runtime/StateTree/Source/StateTreeModule/Public/Debugger/StateTreeTrace.h b/Engine/Plugins/Runtime/StateTree/Source/StateTreeModule/Public/Debugger/StateTreeTrace.h
--- a/Engine/Plugins/Runtime/StateTree/Source/StateTreeModule/Public/Debugger/StateTreeTrace.h
+++ b/Engine/Plugins/Runtime/StateTree/Source/StateTreeModule/Public/Debugger/StateTreeTrace.h
@@ -6,6 +6,9 @@
 
 #include "Trace/Trace.h"
 #include "Containers/ContainersFwd.h"
+#include "Containers/UnrealString.h"
+#include "Containers/StringView.h"
+
 
 struct FStateTreeExecutionFrame;
 class UStateTree;
@@ -25,23 +28,91 @@ UE_TRACE_CHANNEL_EXTERN(StateTreeDebugChannel, STATETREEMODULE_API)
 
 namespace UE::StateTreeTrace
 {
-	void RegisterGlobalDelegates();
-	void UnregisterGlobalDelegates();
-	UE_DEPRECATED(5.6, "This method will no longer be exposed publicly.")
-	FStateTreeIndex16 FindOrAddDebugIdForAsset(const UStateTree* StateTree);
-	void OutputPhaseScopeEvent(FStateTreeInstanceDebugId InstanceId, EStateTreeUpdatePhase Phase, EStateTreeTraceEventType EventType, FStateTreeStateHandle StateHandle);
-	void OutputAssetDebugIdEvent(const UStateTree* StateTree, FStateTreeIndex16 AssetDebugId);
-	void OutputInstanceLifetimeEvent(FStateTreeInstanceDebugId InstanceId, const UStateTree* StateTree, const TCHAR* InstanceName, EStateTreeTraceEventType EventType);
-	void OutputInstanceAssetEvent(FStateTreeInstanceDebugId InstanceId, const UStateTree* StateTree);
-	void OutputInstanceFrameEvent(FStateTreeInstanceDebugId InstanceId, const FStateTreeExecutionFrame* Frame);
-	void OutputLogEventTrace(FStateTreeInstanceDebugId InstanceId, ELogVerbosity::Type Verbosity, const TCHAR* Fmt, ...);
-	void OutputStateEventTrace(FStateTreeInstanceDebugId InstanceId, FStateTreeStateHandle StateHandle, EStateTreeTraceEventType EventType);
-	void OutputTaskEventTrace(FStateTreeInstanceDebugId InstanceId, FStateTreeIndex16 TaskIdx, FStateTreeDataView DataView, EStateTreeTraceEventType EventType, EStateTreeRunStatus Status);
-	void OutputEvaluatorEventTrace(FStateTreeInstanceDebugId InstanceId, FStateTreeIndex16 EvaluatorIdx, FStateTreeDataView DataView, EStateTreeTraceEventType EventType);
-	void OutputConditionEventTrace(FStateTreeInstanceDebugId InstanceId, FStateTreeIndex16 ConditionIdx, FStateTreeDataView DataView, EStateTreeTraceEventType EventType);
-	void OutputTransitionEventTrace(FStateTreeInstanceDebugId InstanceId, FStateTreeTransitionSource TransitionSource, EStateTreeTraceEventType EventType);
-	void OutputActiveStatesEventTrace(FStateTreeInstanceDebugId InstanceId, TConstArrayView<FStateTreeExecutionFrame> ActiveFrames);
-}
+
+/** Struct allowing StateTree nodes to add (or replace) a custom string to the trace event */
+struct FNodeCustomDebugData
+{
+	enum class EMergePolicy: uint8
+	{
+		Unset,
+		Append,
+		Override
+	};
+
+	FNodeCustomDebugData() = default;
+
+	FNodeCustomDebugData(FStringView TraceDebuggerString, const EMergePolicy MergePolicy): TraceDebuggerString(TraceDebuggerString)
+	, MergePolicy(MergePolicy)
+	{
+	}
+
+	FNodeCustomDebugData(FNodeCustomDebugData&& Other)
+		: TraceDebuggerString(MoveTemp(Other.TraceDebuggerString))
+		, MergePolicy(Other.MergePolicy)
+	{
+		Other.Reset();
+	}
+
+	FNodeCustomDebugData& operator=(FNodeCustomDebugData&& Other)
+	{
+		if (this == &Other)
+		{
+			return *this;
+		}
+		TraceDebuggerString = MoveTemp(Other.TraceDebuggerString);
+		MergePolicy = Other.MergePolicy;
+		Other.Reset();
+		return *this;
+	}
+
+	bool IsSet() const
+	{
+		return MergePolicy != EMergePolicy::Unset && !TraceDebuggerString.IsEmpty();
+	}
+
+	void Reset()
+	{
+		TraceDebuggerString.Reset();
+		MergePolicy = EMergePolicy::Unset;
+	}
+
+	bool ShouldOverrideDataView() const
+	{
+		return MergePolicy == EMergePolicy::Override;
+	}
+	bool ShouldAppendToDataView() const
+	{
+		return MergePolicy == EMergePolicy::Append;
+	}
+
+	FStringView GetTraceDebuggerString() const
+	{
+		return TraceDebuggerString;
+	}
+
+private:
+	FString TraceDebuggerString;
+	EMergePolicy MergePolicy = EMergePolicy::Unset;
+};
+
+void RegisterGlobalDelegates();
+void UnregisterGlobalDelegates();
+UE_DEPRECATED(5.6, "This method will no longer be exposed publicly.")
+FStateTreeIndex16 FindOrAddDebugIdForAsset(const UStateTree* StateTree);
+void OutputPhaseScopeEvent(FStateTreeInstanceDebugId InstanceId, EStateTreeUpdatePhase Phase, EStateTreeTraceEventType EventType, FStateTreeStateHandle StateHandle);
+void OutputAssetDebugIdEvent(const UStateTree* StateTree, FStateTreeIndex16 AssetDebugId);
+void OutputInstanceLifetimeEvent(FStateTreeInstanceDebugId InstanceId, const UStateTree* StateTree, const TCHAR* InstanceName, EStateTreeTraceEventType EventType);
+void OutputInstanceAssetEvent(FStateTreeInstanceDebugId InstanceId, const UStateTree* StateTree);
+void OutputInstanceFrameEvent(FStateTreeInstanceDebugId InstanceId, const FStateTreeExecutionFrame* Frame);
+void OutputLogEventTrace(FStateTreeInstanceDebugId InstanceId, ELogVerbosity::Type Verbosity, const TCHAR* Fmt, ...);
+void OutputStateEventTrace(FStateTreeInstanceDebugId InstanceId, FStateTreeStateHandle StateHandle, EStateTreeTraceEventType EventType);
+void OutputTaskEventTrace(FStateTreeInstanceDebugId InstanceId, FNodeCustomDebugData&& DebugData, FStateTreeIndex16 TaskIdx, FStateTreeDataView DataView, EStateTreeTraceEventType EventType, EStateTreeRunStatus Status);
+void OutputEvaluatorEventTrace(FStateTreeInstanceDebugId InstanceId, FNodeCustomDebugData&& DebugData, FStateTreeIndex16 EvaluatorIdx, FStateTreeDataView DataView, EStateTreeTraceEventType EventType);
+void OutputConditionEventTrace(FStateTreeInstanceDebugId InstanceId, FNodeCustomDebugData&& DebugData, FStateTreeIndex16 ConditionIdx, FStateTreeDataView DataView, EStateTreeTraceEventType EventType);
+void OutputTransitionEventTrace(FStateTreeInstanceDebugId InstanceId, FStateTreeTransitionSource TransitionSource, EStateTreeTraceEventType EventType);
+void OutputActiveStatesEventTrace(FStateTreeInstanceDebugId InstanceId, TConstArrayView<FStateTreeExecutionFrame> ActiveFrames);
+
+} // UE::StateTreeTrace
 
 #define TRACE_STATETREE_INSTANCE_EVENT(InstanceID, StateTree, InstanceName, EventType) \
 	UE::StateTreeTrace::OutputInstanceLifetimeEvent(InstanceID, StateTree, InstanceName, EventType);
@@ -64,22 +135,22 @@ namespace UE::StateTreeTrace
 		UE::StateTreeTrace::OutputStateEventTrace(InstanceId, StateHandle, EventType); \
 	}
 
-#define TRACE_STATETREE_TASK_EVENT(InstanceId, TaskIdx, DataView, EventType, Status) \
+#define TRACE_STATETREE_TASK_EVENT(InstanceId, DebugData, TaskIdx, DataView, EventType, Status) \
 	if (UE_TRACE_CHANNELEXPR_IS_ENABLED(StateTreeDebugChannel)) \
 	{ \
-		UE::StateTreeTrace::OutputTaskEventTrace(InstanceId, TaskIdx, DataView, EventType, Status); \
+		UE::StateTreeTrace::OutputTaskEventTrace(InstanceId, DebugData, TaskIdx, DataView, EventType, Status); \
 	}
 
-#define TRACE_STATETREE_EVALUATOR_EVENT(InstanceId, EvaluatorIdx, DataView, EventType) \
+#define TRACE_STATETREE_EVALUATOR_EVENT(InstanceId, DebugData, EvaluatorIdx, DataView, EventType) \
 	if (UE_TRACE_CHANNELEXPR_IS_ENABLED(StateTreeDebugChannel)) \
 	{ \
-		UE::StateTreeTrace::OutputEvaluatorEventTrace(InstanceId, EvaluatorIdx, DataView, EventType); \
+		UE::StateTreeTrace::OutputEvaluatorEventTrace(InstanceId, DebugData, EvaluatorIdx, DataView, EventType); \
 	}
 
-#define TRACE_STATETREE_CONDITION_EVENT(InstanceId, ConditionIdx, DataView, EventType) \
+#define TRACE_STATETREE_CONDITION_EVENT(InstanceId, DebugData, ConditionIdx, DataView, EventType) \
 	if (UE_TRACE_CHANNELEXPR_IS_ENABLED(StateTreeDebugChannel)) \
 	{ \
-		UE::StateTreeTrace::OutputConditionEventTrace(InstanceId, ConditionIdx, DataView, EventType); \
+		UE::StateTreeTrace::OutputConditionEventTrace(InstanceId, DebugData, ConditionIdx, DataView, EventType); \
 	}
 
 #define TRACE_STATETREE_TRANSITION_EVENT(InstanceId, TransitionIdx, EventType) \
diff --git a/Engine/Plugins/Runtime/StateTree/Source/StateTreeModule/Public/Debugger/StateTreeTraceTypes.h b/Engine/Plugins/Runtime/StateTree/Source/StateTreeModule/Public/Debugger/StateTreeTraceTypes.h
--- a/Engine/Plugins/Runtime/StateTree/Source/StateTreeModule/Public/Debugger/StateTreeTraceTypes.h
+++ b/Engine/Plugins/Runtime/StateTree/Source/StateTreeModule/Public/Debugger/StateTreeTraceTypes.h
@@ -57,8 +57,20 @@ struct FStateTreeTraceBaseEvent
 	{
 	}
 
-	static FString GetDataTypePath() { return TEXT(""); }
-	static FString GetDataAsText() { return TEXT(""); }
+	static FString GetDataTypePath()
+	{
+		return TEXT("");
+	}
+
+	static FString GetDataAsText()
+	{
+		return TEXT("");
+	}
+
+	static FString GetDebugText()
+	{
+		return TEXT("");
+	}
 
 	double RecordingWorldTime = 0;
 	EStateTreeTraceEventType EventType;
@@ -83,7 +95,7 @@ struct FStateTreeTracePhaseEvent : FStateTreeTraceBaseEvent
 
 struct FStateTreeTraceLogEvent : FStateTreeTraceBaseEvent
 {
-	explicit FStateTreeTraceLogEvent(const double RecordingWorldTime, ELogVerbosity::Type Verbosity, const FString& Message)
+	explicit FStateTreeTraceLogEvent(const double RecordingWorldTime, const ELogVerbosity::Type Verbosity, const FString& Message)
 		: FStateTreeTraceBaseEvent(RecordingWorldTime, EStateTreeTraceEventType::Unset)
 		, Verbosity(Verbosity)
 		, Message(Message)
@@ -105,6 +117,11 @@ struct FStateTreeTracePropertyEvent : FStateTreeTraceLogEvent
 	{
 	}
 
+	explicit FStateTreeTracePropertyEvent(const double RecordingWorldTime, const ELogVerbosity::Type Verbosity, const FString& Message)
+		: FStateTreeTraceLogEvent(RecordingWorldTime, Verbosity, Message)
+	{
+	}
+
 	STATETREEMODULE_API FString ToFullString(const UStateTree& StateTree) const;
 	STATETREEMODULE_API FString GetValueString(const UStateTree& StateTree) const;
 	STATETREEMODULE_API FString GetTypeString(const UStateTree& StateTree) const;
@@ -155,10 +172,17 @@ struct FStateTreeTraceStateEvent : FStateTreeTraceNodeEvent
 
 struct FStateTreeTraceTaskEvent : FStateTreeTraceNodeEvent
 {
-	explicit FStateTreeTraceTaskEvent(const double RecordingWorldTime, const FStateTreeIndex16 Index, const EStateTreeTraceEventType EventType, const EStateTreeRunStatus Status, const FString& TypePath, const FString& InstanceDataAsText)
+	explicit FStateTreeTraceTaskEvent(const double RecordingWorldTime
+									, const FStateTreeIndex16 Index
+									, const EStateTreeTraceEventType EventType
+									, const EStateTreeRunStatus Status
+									, const FString& TypePath
+									, const FString& InstanceDataAsText
+									, const FString& DebugText)
 		: FStateTreeTraceNodeEvent(RecordingWorldTime, Index, EventType)
 		, TypePath(TypePath)
 		, InstanceDataAsText(InstanceDataAsText)
+		, DebugText(DebugText)
 		, Status(Status)
 	{
 	}
@@ -167,20 +191,39 @@ struct FStateTreeTraceTaskEvent : FStateTreeTraceNodeEvent
 	STATETREEMODULE_API FString GetValueString(const UStateTree& StateTree) const;
 	STATETREEMODULE_API FString GetTypeString(const UStateTree& StateTree) const;
 
-	FString GetDataTypePath() const { return TypePath; }
-	FString GetDataAsText() const { return InstanceDataAsText; }
+	FString GetDataTypePath() const
+	{
+		return TypePath;
+	}
+
+	FString GetDataAsText() const
+	{
+		return InstanceDataAsText;
+	}
+
+	FString GetDebugText() const
+	{
+		return DebugText;
+	}
 
 	FString TypePath;
 	FString InstanceDataAsText;
+	FString DebugText;
 	EStateTreeRunStatus Status;
 };
 
 struct FStateTreeTraceEvaluatorEvent : FStateTreeTraceNodeEvent
 {
-	explicit FStateTreeTraceEvaluatorEvent(const double RecordingWorldTime, const FStateTreeIndex16 Index, const EStateTreeTraceEventType EventType, const FString& TypePath, const FString& InstanceDataAsText)
+	explicit FStateTreeTraceEvaluatorEvent(const double RecordingWorldTime
+										, const FStateTreeIndex16 Index
+										, const EStateTreeTraceEventType EventType
+										, const FString& TypePath
+										, const FString& InstanceDataAsText
+										, const FString& DebugText)
 		: FStateTreeTraceNodeEvent(RecordingWorldTime, Index, EventType)
 		, TypePath(TypePath)
 		, InstanceDataAsText(InstanceDataAsText)
+		, DebugText(DebugText)
 	{
 	}
 
@@ -188,19 +231,38 @@ struct FStateTreeTraceEvaluatorEvent : FStateTreeTraceNodeEvent
 	STATETREEMODULE_API FString GetValueString(const UStateTree& StateTree) const;
 	STATETREEMODULE_API FString GetTypeString(const UStateTree& StateTree) const;
 
-	FString GetDataTypePath() const { return TypePath; }
-	FString GetDataAsText() const { return InstanceDataAsText; }
+	FString GetDataTypePath() const
+	{
+		return TypePath;
+	}
+
+	FString GetDataAsText() const
+	{
+		return InstanceDataAsText;
+	}
+
+	FString GetDebugText() const
+	{
+		return DebugText;
+	}
 
 	FString TypePath;
 	FString InstanceDataAsText;
+	FString DebugText;
 };
 
 struct FStateTreeTraceConditionEvent : FStateTreeTraceNodeEvent
 {
-	explicit FStateTreeTraceConditionEvent(const double RecordingWorldTime, const FStateTreeIndex16 Index, const EStateTreeTraceEventType EventType, const FString& TypePath, const FString& InstanceDataAsText)
+	explicit FStateTreeTraceConditionEvent(const double RecordingWorldTime
+										, const FStateTreeIndex16 Index
+										, const EStateTreeTraceEventType EventType
+										, const FString& TypePath
+										, const FString& InstanceDataAsText
+										, const FString& DebugText)
 		: FStateTreeTraceNodeEvent(RecordingWorldTime, Index, EventType)
 		, TypePath(TypePath)
 		, InstanceDataAsText(InstanceDataAsText)
+		, DebugText(DebugText)
 	{
 	}
 
@@ -208,11 +270,24 @@ struct FStateTreeTraceConditionEvent : FStateTreeTraceNodeEvent
 	STATETREEMODULE_API FString GetValueString(const UStateTree& StateTree) const;
 	STATETREEMODULE_API FString GetTypeString(const UStateTree& StateTree) const;
 
-	FString GetDataTypePath() const { return TypePath; }
-	FString GetDataAsText() const { return InstanceDataAsText; }
+	FString GetDataTypePath() const
+	{
+		return TypePath;
+	}
+
+	FString GetDataAsText() const
+	{
+		return InstanceDataAsText;
+	}
+
+	FString GetDebugText() const
+	{
+		return DebugText;
+	}
 
 	FString TypePath;
 	FString InstanceDataAsText;
+	FString DebugText;
 };
 
 struct FStateTreeTraceActiveStates
diff --git a/Engine/Plugins/Runtime/StateTree/Source/StateTreeModule/Public/StateTreeExecutionContext.h b/Engine/Plugins/Runtime/StateTree/Source/StateTreeModule/Public/StateTreeExecutionContext.h
--- a/Engine/Plugins/Runtime/StateTree/Source/StateTreeModule/Public/StateTreeExecutionContext.h
+++ b/Engine/Plugins/Runtime/StateTree/Source/StateTreeModule/Public/StateTreeExecutionContext.h
@@ -8,6 +8,7 @@
 #include "StateTreeNodeBase.h"
 #include "StateTreeNodeRef.h"
 #include "StateTreeReference.h"
+#include "Debugger/StateTreeTrace.h"
 #include "Experimental/ConcurrentLinearAllocator.h"
 #include "Templates/IsInvocable.h"
 
@@ -105,13 +106,30 @@ public:
 	/** Sends event for the StateTree. */
 	void SendEvent(const FGameplayTag Tag, const FConstStructView Payload = FConstStructView(), const FName Origin = FName());
 
+#if WITH_STATETREE_TRACE
+	void SetNodeCustomDebugTraceData(UE::StateTreeTrace::FNodeCustomDebugData&& DebugData)
+	{
+		ensureMsgf(!NodeCustomDebugTraceData .IsSet()
+			, TEXT("CustomData is not expected to be already set."
+				" This might indicate nested calls to SetNodeCustomDebugTraceData without calls to a trace macro"));
+		NodeCustomDebugTraceData = MoveTemp(DebugData);
+	}
+#endif //WITH_STATETREE_TRACE
+
 protected:
 	/** @return Prefix that will be used by STATETREE_LOG and STATETREE_CLOG, Owner name by default. */
 	FString GetInstanceDescriptionInternal() const;
 
 #if WITH_STATETREE_TRACE
-	/** A unique Id used by debuging tools to identify the instance. */
+	/** A unique Id used by debugging tools to identify the instance. */
 	FStateTreeInstanceDebugId GetInstanceDebugId() const;
+
+	UE::StateTreeTrace::FNodeCustomDebugData StealNodeCustomDebugTraceData()
+	{
+		return MoveTemp(NodeCustomDebugTraceData);
+	}
+
+	UE::StateTreeTrace::FNodeCustomDebugData NodeCustomDebugTraceData;
 #endif // WITH_STATETREE_TRACE
 
 	/** Informs the owner when the instance of the tree must woke up from a scheduled tick sleep. */
diff --git a/Engine/Plugins/Runtime/StateTree/Source/StateTreeModule/Public/StateTreeNodeBase.h b/Engine/Plugins/Runtime/StateTree/Source/StateTreeModule/Public/StateTreeNodeBase.h
--- a/Engine/Plugins/Runtime/StateTree/Source/StateTreeModule/Public/StateTreeNodeBase.h
+++ b/Engine/Plugins/Runtime/StateTree/Source/StateTreeModule/Public/StateTreeNodeBase.h
@@ -10,6 +10,18 @@ struct FPropertyBindingPath;
 struct FStateTreePropertyPath;
 struct IStateTreeBindingLookup;
 
+#if WITH_STATETREE_TRACE
+#define SET_NODE_CUSTOM_TRACE_TEXT(Context, MergePolicy, Format, ...) \
+	if (UE_TRACE_CHANNELEXPR_IS_ENABLED(StateTreeDebugChannel)) \
+	{ \
+		Context.SetNodeCustomDebugTraceData( \
+			UE::StateTreeTrace::FNodeCustomDebugData(FString::Printf(Format, ##__VA_ARGS__) \
+													, ::UE::StateTreeTrace::FNodeCustomDebugData::EMergePolicy::MergePolicy)); \
+	}
+#else
+#define SET_NODE_CUSTOM_TRACE_TEXT(...)
+#endif //WITH_STATETREE_TRACE
+
 /**
  * Enum describing in what format a text is expected to be returned.
  *
diff --git a/Engine/Plugins/Runtime/StateTree/Source/StateTreeModule/Public/StateTreeNodeDescriptionHelpers.h b/Engine/Plugins/Runtime/StateTree/Source/StateTreeModule/Public/StateTreeNodeDescriptionHelpers.h
--- a/Engine/Plugins/Runtime/StateTree/Source/StateTreeModule/Public/StateTreeNodeDescriptionHelpers.h
+++ b/Engine/Plugins/Runtime/StateTree/Source/StateTreeModule/Public/StateTreeNodeDescriptionHelpers.h
@@ -13,7 +13,7 @@ struct FGameplayTagQuery;
 
 namespace UE::StateTree::DescHelpers
 {
-#if WITH_EDITOR
+#if WITH_EDITOR || WITH_STATETREE_TRACE
 
 /** @return description for a EGenericAICheck. */
 extern STATETREEMODULE_API FText GetOperatorText(const EGenericAICheck Operator, EStateTreeNodeFormatting Formatting);
@@ -56,6 +56,12 @@ extern STATETREEMODULE_API FText GetText(const UObject* Value, EStateTreeNodeFor
 
 extern STATETREEMODULE_API FText GetMathOperationText(const FText& OperationText, const FText& LeftText, const FText& RightText, EStateTreeNodeFormatting Formatting);
 
+extern STATETREEMODULE_API FText GetSingleParamFunctionText(const FText& FunctionText, const FText& ParamText, EStateTreeNodeFormatting Formatting);
+
+#endif // WITH_EDITOR || WITH_STATETREE_TRACE
+
+#if WITH_EDITOR
+
 /** @return description in the form of (Left OperationText Right).
 *	Expect TInstanceDataType to have a member Left and Right whose types have an overloaded UE::StateTree::DescHelpers::GetText function.
 */
@@ -79,8 +85,6 @@ FText GetDescriptionForMathOperation(FText OperationText, const FGuid& ID, FStat
 	return GetMathOperationText(OperationText, LeftValue, RightValue, Formatting);
 }
 
-extern STATETREEMODULE_API FText GetSingleParamFunctionText(const FText& FunctionText, const FText& ParamText, EStateTreeNodeFormatting Formatting);
-
 /** @return description in the form of OperationText(Input).
  *	Expect TInstanceDataType to have a member input whose type has an overloaded UE::StateTree::DescHelpers::GetText function.
  */
