StateTree: Add support for transition on PropertyRef K2Node.

diff --git a/Engine/Plugins/Runtime/StateTree/Source/StateTreeEditorModule/Private/K2Node_StateTreeBlueprintPropertyRef.cpp b/Engine/Plugins/Runtime/StateTree/Source/StateTreeEditorModule/Private/K2Node_StateTreeBlueprintPropertyRef.cpp
--- a/Engine/Plugins/Runtime/StateTree/Source/StateTreeEditorModule/Private/K2Node_StateTreeBlueprintPropertyRef.cpp
+++ b/Engine/Plugins/Runtime/StateTree/Source/StateTreeEditorModule/Private/K2Node_StateTreeBlueprintPropertyRef.cpp
@@ -72,6 +72,7 @@ void UK2Node_StateTreeBlueprintPropertyRef::ReconstructNode()
 void UK2Node_StateTreeBlueprintPropertyRef::ClearCachedBlueprintData(UBlueprint* Blueprint)
 {
 	Super::ClearCachedBlueprintData(Blueprint);
+	Modify(true);
 	UpdateOutputPin();
 }
 
