StateTree: Fix calls to PostEditInstanceDataChangeChainProperty on nodes with TArray property.

diff --git a/Engine/Plugins/Runtime/StateTree/Source/StateTreeEditorModule/Private/StateTreePropertyHelpers.cpp b/Engine/Plugins/Runtime/StateTree/Source/StateTreeEditorModule/Private/StateTreePropertyHelpers.cpp
--- a/Engine/Plugins/Runtime/StateTree/Source/StateTreeEditorModule/Private/StateTreePropertyHelpers.cpp
+++ b/Engine/Plugins/Runtime/StateTree/Source/StateTreeEditorModule/Private/StateTreePropertyHelpers.cpp
@@ -105,7 +105,8 @@ void DispatchPostEditToNodes(UObject& Owner, FPropertyChangedChainEvent& InPrope
 			const int32 Index = InPropertyChangedEvent.GetArrayIndex(ArrayProperty->GetName());
 			if (!Helper.IsValidIndex(Index))
 			{
-				return;
+				check(CurrentPropNode->GetNextNode() == nullptr);
+				break;
 			}
 
 			if (TargetPath.GetStructID().IsValid())
