[State Tree Editor] Cache Transition Widgets and update its content on demand instead of refreshing the whole list

diff --git a/Engine/Plugins/Runtime/StateTree/Source/StateTreeEditorModule/Private/SStateTreeView.cpp b/Engine/Plugins/Runtime/StateTree/Source/StateTreeEditorModule/Private/SStateTreeView.cpp
--- a/Engine/Plugins/Runtime/StateTree/Source/StateTreeEditorModule/Private/SStateTreeView.cpp
+++ b/Engine/Plugins/Runtime/StateTree/Source/StateTreeEditorModule/Private/SStateTreeView.cpp
@@ -448,6 +448,7 @@ void SStateTreeView::HandleModelStatesChanged(const TSet<UStateTreeState*>& Affe
 	// When the tasks or conditions array changed(this includes both normal array operations: Add, Remove. Clear, Move,
 	// and Paste or Duplicate an element in the array), The TreeView needs to be rebuilt because new elements came in or old elements have gone or both.
 	// This will not rebuild the list when we change an inner property in a condition or in a task node because of InstanceStruct wrapper
+	// @todo: change it to cache and re-set the content of the widget instead of rebuilding the whole list for perf
 	if (PropertyChangedEvent.MemberProperty->GetFName() == GET_MEMBER_NAME_CHECKED(UStateTreeState, Tasks)
 		|| PropertyChangedEvent.MemberProperty->GetFName() == GET_MEMBER_NAME_CHECKED(UStateTreeState, EnterConditions)
 		|| PropertyChangedEvent.MemberProperty->GetFName() == GET_MEMBER_NAME_CHECKED(UStateTreeState, bHasRequiredEventToEnter)
@@ -455,16 +456,6 @@ void SStateTreeView::HandleModelStatesChanged(const TSet<UStateTreeState*>& Affe
 	{
 		TreeView->RebuildList();
 	}
-	//@todo: find a better way to fetch/update Transition information so we don't have to hardcode the rules here
-	// We cache the information of Transitions when making SStateTreeViewRow and using cached info to display
-	// So we need to rebuild the list whenever the relevant information changes to update the view
-	// For State Transitions, this will cover any changes made on the Array and on each Transition
-	// For Task Transitions, LinkType is the only exposed one in the editor
-	else if (PropertyChangedEvent.MemberProperty->GetFName() == GET_MEMBER_NAME_CHECKED(UStateTreeState, Transitions)
-		|| PropertyChangedEvent.Property->GetFName() == GET_MEMBER_NAME_CHECKED(FStateTreeStateLink, LinkType))
-	{
-		TreeView->RebuildList();
-	}
 }
 
 void SStateTreeView::HandleModelSelectionChanged(const TArray<TWeakObjectPtr<UStateTreeState>>& SelectedStates)
diff --git a/Engine/Plugins/Runtime/StateTree/Source/StateTreeEditorModule/Private/SStateTreeViewRow.cpp b/Engine/Plugins/Runtime/StateTree/Source/StateTreeEditorModule/Private/SStateTreeViewRow.cpp
--- a/Engine/Plugins/Runtime/StateTree/Source/StateTreeEditorModule/Private/SStateTreeViewRow.cpp
+++ b/Engine/Plugins/Runtime/StateTree/Source/StateTreeEditorModule/Private/SStateTreeViewRow.cpp
@@ -40,6 +40,8 @@
 			static_cast<uint8>(FMath::RoundToInt(static_cast<float>(A.B) * (1.f - T) + static_cast<float>(B.B) * T)),
 			static_cast<uint8>(FMath::RoundToInt(static_cast<float>(A.A) * (1.f - T) + static_cast<float>(B.A) * T))));
 	}
+
+	static constexpr FLinearColor IconTint = FLinearColor(1, 1, 1, 0.5f);
 } // UE:StateTree::Editor
 
 void SStateTreeViewRow::Construct(const FArguments& InArgs, const TSharedRef<STableViewBase>& InOwnerTableView, TWeakObjectPtr<UStateTreeState> InState, const TSharedPtr<SScrollBox>& ViewBox, TSharedPtr<FStateTreeViewModel> InStateTreeViewModel)
@@ -50,6 +52,7 @@
 	WeakEditorData = State != nullptr ? State->GetTypedOuter<UStateTreeEditorData>() : nullptr;
 
 	AssetChangedHandle = StateTreeViewModel->GetOnAssetChanged().AddSP(this, &SStateTreeViewRow::HandleAssetChanged);
+	StatesChangedHandle = StateTreeViewModel->GetOnStatesChanged().AddSP(this, &SStateTreeViewRow::HandleStatesChanged);
 
 	ConstructInternal(STableRow::FArguments()
 		.OnDragDetected(this, &SStateTreeViewRow::HandleDragDetected)
@@ -59,84 +62,6 @@
 		.Style(&FStateTreeEditorStyle::Get().GetWidgetStyle<FTableRowStyle>("StateTree.Selection"))
 		, InOwnerTableView);
 
-	static constexpr FLinearColor IconTint = FLinearColor(1, 1, 1, 0.5f);
-	
-	auto MakeTransitionWidget = [this](const EStateTreeTransitionTrigger Trigger, const FSlateBrush* Icon)
-	{
-		FTransitionDescFilterOptions FilterOptions;
-		FilterOptions.bUseMask = EnumHasAnyFlags(Trigger, EStateTreeTransitionTrigger::OnTick | EStateTreeTransitionTrigger::OnEvent| EStateTreeTransitionTrigger::OnDelegate);
-		
-			return SNew(SBox)
-			.HeightOverride(GetDefault<UStateTreeEditorUserSettings>()->GetStatesViewStateRowHeight())
-			.Visibility(this, &SStateTreeViewRow::GetTransitionsVisibility, Trigger)
-			[
-				SNew(SHorizontalBox)
-				+ SHorizontalBox::Slot()
-				.VAlign(VAlign_Center)
-				.AutoWidth()
-				.Padding(FMargin(0.f, 0.f, 0.f, 0.f))
-				[
-					SNew(SHorizontalBox)
-					
-					+ SHorizontalBox::Slot()
-					.VAlign(VAlign_Center)
-					.AutoWidth()
-					[
-						SNew(SImage)
-						.Image(Icon)
-					]
-					
-					+ SHorizontalBox::Slot()
-					.VAlign(VAlign_Center)
-					.AutoWidth()
-					[
-						SNew(SImage)
-						.Image(this, &SStateTreeViewRow::GetTransitionsIcon, Trigger)
-						.ColorAndOpacity(IconTint)
-					]
-				]
-				+ SHorizontalBox::Slot()
-				.VAlign(VAlign_Center)
-				.AutoWidth()
-				.Padding(FMargin(4.f, 0.f, 12.f, 0.f))
-				[
-					SNew(SOverlay)
-					+ SOverlay::Slot()
-					[
-						MakeTransitionWidgets(Trigger, FilterOptions)
-					]
-					+ SOverlay::Slot()
-					[
-						// Breakpoint box
-						SNew(SHorizontalBox)
-						+ SHorizontalBox::Slot()
-						.VAlign(VAlign_Top)
-						.HAlign(HAlign_Left)
-						.AutoWidth()
-						[
-							SNew(SBox)
-							.Padding(FMargin(-4.f, -4.f, 0.f, 0.f))
-							[
-								SNew(SImage)
-								.DesiredSizeOverride(FVector2D(10.f, 10.f))
-								.Image(FStateTreeEditorStyle::Get().GetBrush(TEXT("StateTreeEditor.Debugger.Breakpoint.EnabledAndValid")))
-								.Visibility(this, &SStateTreeViewRow::GetTransitionsBreakpointVisibility, Trigger)
-								.ToolTipText_Lambda([this, Trigger, InFilterOptions = FilterOptions]
-									{
-										FTransitionDescFilterOptions FilterOptions = InFilterOptions;
-										FilterOptions.WithBreakpoint = ETransitionDescRequirement::RequiredTrue;
-
-										return FText::Format(LOCTEXT("TransitionBreakpointTooltip","Break when executing transition: {0}"),
-											GetTransitionsDesc(Trigger, FilterOptions));
-									})
-							]
-						]
-					]
-
-				]
-			];
-	};
-	
 	TSharedPtr<SVerticalBox> StateAndTasksVerticalBox;
 	TSharedPtr<SHorizontalBox> StateHorizontalBox;
 	TSharedPtr<SBorder> FlagBorder;
@@ -250,7 +175,7 @@
 												.Visibility(this, &SStateTreeViewRow::GetConditionVisibility)
 												[
 													SNew(SImage)
-													.ColorAndOpacity(IconTint)
+													.ColorAndOpacity(UE::StateTree::Editor::IconTint)
 													.Image(FStateTreeEditorStyle::Get().GetBrush("StateTreeEditor.StateConditions"))
 													.ToolTipText(LOCTEXT("StateHasEnterConditions", "State selection is guarded with enter conditions."))
 												]
@@ -266,7 +191,7 @@
 												[
 													SNew(SImage)
 													.Image(this, &SStateTreeViewRow::GetSelectorIcon)
-													.ColorAndOpacity(IconTint)
+													.ColorAndOpacity(UE::StateTree::Editor::IconTint)
 													.ToolTipText(this, &SStateTreeViewRow::GetSelectorTooltip)
 												]
 											]
@@ -313,7 +238,7 @@
 													SNew(SImage)
 													.Image(FAppStyle::Get().GetBrush(TEXT("Icons.Comment")))
 													.ColorAndOpacity(FStyleColors::Foreground)
-													.ColorAndOpacity(IconTint)
+													.ColorAndOpacity(UE::StateTree::Editor::IconTint)
 													.ToolTipText(this, &SStateTreeViewRow::GetStateDescription)
 												]
 											]
@@ -346,7 +271,7 @@
 													.Padding(FMargin(4.f, 0.f, 4.f, 0.f))
 													[
 														SNew(SImage)
-														.ColorAndOpacity(IconTint)
+														.ColorAndOpacity(UE::StateTree::Editor::IconTint)
 														.Image(FStateTreeEditorStyle::Get().GetBrush("StateTreeEditor.StateLinked"))
 													]
 
@@ -414,54 +339,7 @@
 		.HAlign(HAlign_Left)
 		[
 			// Transitions
-			SNew(SHorizontalBox)
-
-			+ SHorizontalBox::Slot()
-			.VAlign(VAlign_Top)
-			.AutoWidth()
-			[
-				SNew(SBox)
-				.HeightOverride(GetDefault<UStateTreeEditorUserSettings>()->GetStatesViewStateRowHeight())
-				.Visibility(this, &SStateTreeViewRow::GetTransitionDashVisibility)
-				.VAlign(VAlign_Center)
-				[
-					SNew(SImage)
-					.Image(FStateTreeEditorStyle::Get().GetBrush("StateTreeEditor.Transition.Dash"))
-					.ColorAndOpacity(IconTint)
-				]
-			]
-
-			// Completed transitions
-			+ SHorizontalBox::Slot()
-			.VAlign(VAlign_Top)
-			.AutoWidth()
-			[
-				MakeTransitionWidget(EStateTreeTransitionTrigger::OnStateCompleted, nullptr)
-			]
-
-			// Succeeded transitions
-			+ SHorizontalBox::Slot()
-			.VAlign(VAlign_Top)
-			.AutoWidth()
-			[
-				MakeTransitionWidget(EStateTreeTransitionTrigger::OnStateSucceeded, FStateTreeEditorStyle::Get().GetBrush("StateTreeEditor.Transition.Succeeded"))
-			]
-
-			// Failed transitions
-			+ SHorizontalBox::Slot()
-			.VAlign(VAlign_Top)
-			.AutoWidth()
-			[
-				MakeTransitionWidget(EStateTreeTransitionTrigger::OnStateFailed, FStateTreeEditorStyle::Get().GetBrush("StateTreeEditor.Transition.Failed"))
-			]
-
-			// Transitions
-			+ SHorizontalBox::Slot()
-			.VAlign(VAlign_Top)
-			.AutoWidth()
-			[
-				MakeTransitionWidget(EStateTreeTransitionTrigger::OnTick | EStateTreeTransitionTrigger::OnEvent | EStateTreeTransitionTrigger::OnDelegate, FStateTreeEditorStyle::Get().GetBrush("StateTreeEditor.Transition.Condition"))
-			]
+			SAssignNew(TransitionsContainer, SHorizontalBox)
 		];
 	}
 
@@ -485,7 +363,7 @@
 		];
 	}
 
-
+	MakeTransitionsWidget();
 	MakeFlagsWidget();
 }
 
@@ -495,6 +373,7 @@
 	if (StateTreeViewModel)
 	{
 		StateTreeViewModel->GetOnAssetChanged().Remove(AssetChangedHandle);
+		StateTreeViewModel->GetOnStatesChanged().Remove(StatesChangedHandle);
 	}
 }
 
@@ -961,4 +840,324 @@
 	return VerticalBox;
 }
 
+void SStateTreeViewRow::MakeTransitionsWidget()
+{
+	SHorizontalBox* TransitionsContainerPtr = TransitionsContainer.Get();
+	check(TransitionsContainerPtr != nullptr)
+
+	TransitionsContainerPtr->ClearChildren();
+
+	TransitionsContainerPtr->AddSlot()
+	.VAlign(VAlign_Top)
+	.AutoWidth()
+	[
+		SNew(SBox)
+		.HeightOverride(GetDefault<UStateTreeEditorUserSettings>()->GetStatesViewStateRowHeight())
+		.Visibility(this, &SStateTreeViewRow::GetTransitionDashVisibility)
+		.VAlign(VAlign_Center)
+		[
+			SNew(SImage)
+			.Image(FStateTreeEditorStyle::Get().GetBrush("StateTreeEditor.Transition.Dash"))
+			.ColorAndOpacity(UE::StateTree::Editor::IconTint)
+		]
+	];
+
+	// On State Completed
+	//We don't show any additional signs for On Completed transitions, just the dash
+	constexpr FSlateBrush* OnCompletedSlateIcon = nullptr;
+	TransitionsContainerPtr->AddSlot()
+	.VAlign(VAlign_Top)
+	.AutoWidth()
+	[
+		MakeTransitionWidget(EStateTreeTransitionTrigger::OnStateCompleted, OnCompletedSlateIcon)
+	];
+
+	// On State Succeeded
+	TransitionsContainerPtr->AddSlot()
+	.VAlign(VAlign_Top)
+	.AutoWidth()
+	[
+		MakeTransitionWidget(EStateTreeTransitionTrigger::OnStateSucceeded, FStateTreeEditorStyle::Get().GetBrush("StateTreeEditor.Transition.Succeeded"))
+	];
+
+	// On State Failed
+	TransitionsContainerPtr->AddSlot()
+	.VAlign(VAlign_Top)
+	.AutoWidth()
+	[
+		MakeTransitionWidget(EStateTreeTransitionTrigger::OnStateFailed, FStateTreeEditorStyle::Get().GetBrush("StateTreeEditor.Transition.Failed"))
+	];
+
+	// On Tick, Event, Delegate
+	TransitionsContainerPtr->AddSlot()
+	.VAlign(VAlign_Top)
+	.AutoWidth()
+	[
+		MakeTransitionWidget(EStateTreeTransitionTrigger::OnTick | EStateTreeTransitionTrigger::OnEvent | EStateTreeTransitionTrigger::OnDelegate, FStateTreeEditorStyle::Get().GetBrush("StateTreeEditor.Transition.Condition"))
+	];
+}
+
+TSharedRef<SWidget> SStateTreeViewRow::MakeTransitionWidget(const EStateTreeTransitionTrigger Trigger, const FSlateBrush* Icon)
+{
+		FTransitionDescFilterOptions FilterOptions;
+		FilterOptions.bUseMask = EnumHasAnyFlags(Trigger, EStateTreeTransitionTrigger::OnTick | EStateTreeTransitionTrigger::OnEvent | EStateTreeTransitionTrigger::OnDelegate);
+
+		return SNew(SBox)
+		.HeightOverride(GetDefault<UStateTreeEditorUserSettings>()->GetStatesViewStateRowHeight())
+		.Visibility(this, &SStateTreeViewRow::GetTransitionsVisibility, Trigger)
+		[
+			SNew(SHorizontalBox)
+			+ SHorizontalBox::Slot()
+			.VAlign(VAlign_Center)
+			.AutoWidth()
+			.Padding(FMargin(0.f, 0.f, 0.f, 0.f))
+			[
+				SNew(SHorizontalBox)
+
+				+ SHorizontalBox::Slot()
+				.VAlign(VAlign_Center)
+				.AutoWidth()
+				[
+					SNew(SImage)
+					.Image(Icon)
+				]
+
+				+ SHorizontalBox::Slot()
+				.VAlign(VAlign_Center)
+				.AutoWidth()
+				[
+					SNew(SImage)
+					.Image(this, &SStateTreeViewRow::GetTransitionsIcon, Trigger)
+					.ColorAndOpacity(UE::StateTree::Editor::IconTint)
+				]
+			]
+
+			+ SHorizontalBox::Slot()
+			.VAlign(VAlign_Center)
+			.AutoWidth()
+			.Padding(FMargin(4.f, 0.f, 12.f, 0.f))
+			[
+				SNew(SOverlay)
+				+ SOverlay::Slot()
+				[
+					MakeTransitionWidgetInternal(Trigger, FilterOptions)
+				]
+				+ SOverlay::Slot()
+				[
+					// Breakpoint box
+					SNew(SHorizontalBox)
+					+ SHorizontalBox::Slot()
+					.VAlign(VAlign_Top)
+					.HAlign(HAlign_Left)
+					.AutoWidth()
+					[
+						SNew(SBox)
+						.Padding(FMargin(-4.f, -4.f, 0.f, 0.f))
+						[
+							SNew(SImage)
+							.DesiredSizeOverride(FVector2D(10.f, 10.f))
+							.Image(FStateTreeEditorStyle::Get().GetBrush(TEXT("StateTreeEditor.Debugger.Breakpoint.EnabledAndValid")))
+							.Visibility(this, &SStateTreeViewRow::GetTransitionsBreakpointVisibility, Trigger)
+							.ToolTipText_Lambda([this, Trigger, InFilterOptions = FilterOptions]
+							{
+								FTransitionDescFilterOptions FilterOptions = InFilterOptions;
+								FilterOptions.WithBreakpoint = ETransitionDescRequirement::RequiredTrue;
+
+								return FText::Format(LOCTEXT("TransitionBreakpointTooltip", "Break when executing transition: {0}"),
+													 GetTransitionsDesc(Trigger, FilterOptions));
+							})
+						]
+					]
+				]
+			]
+		];
+};
+
+TSharedRef<SWidget> SStateTreeViewRow::MakeTransitionWidgetInternal(const EStateTreeTransitionTrigger Trigger, const FTransitionDescFilterOptions FilterOptions)
+{
+	const UStateTreeEditorData* TreeEditorData = WeakEditorData.Get();
+	const UStateTreeState* State = WeakState.Get();
+
+	if (!TreeEditorData || !State)
+	{
+		return SNullWidget::NullWidget;
+	}
+	
+	struct FItem
+	{
+		FItem() = default;
+		FItem(const FStateTreeStateLink& InLink, const FGuid InNodeID)
+			: Link(InLink)
+			, NodeID(InNodeID)
+		{
+		}
+		FItem(const FText& InDesc, const FText& InTooltip)
+			: Desc(InDesc)
+			, Tooltip(InTooltip)
+		{
+		}
+		
+		FText Desc;
+		FText Tooltip;
+		FStateTreeStateLink Link;
+		FGuid NodeID;
+	};
+
+	TArray<FItem> DescItems;
+
+	for (const FStateTreeTransition& Transition : State->Transitions)
+	{
+		// Apply filter for enabled/disabled transitions
+		if ((FilterOptions.Enabled == ETransitionDescRequirement::RequiredTrue && Transition.bTransitionEnabled == false)
+			|| (FilterOptions.Enabled == ETransitionDescRequirement::RequiredFalse && Transition.bTransitionEnabled))
+		{
+			continue;
+		}
+
+#if WITH_STATETREE_TRACE_DEBUGGER
+		// Apply filter for transitions with/without breakpoint
+		const bool bHasBreakpoint = TreeEditorData != nullptr && TreeEditorData->HasBreakpoint(Transition.ID, EStateTreeBreakpointType::OnTransition);
+		if ((FilterOptions.WithBreakpoint == ETransitionDescRequirement::RequiredTrue && bHasBreakpoint == false)
+			|| (FilterOptions.WithBreakpoint == ETransitionDescRequirement::RequiredFalse && bHasBreakpoint))
+		{
+			continue;
+		}
+#endif // WITH_STATETREE_TRACE_DEBUGGER
+
+		const bool bMatch = FilterOptions.bUseMask ? EnumHasAnyFlags(Transition.Trigger, Trigger) : Transition.Trigger == Trigger;
+		if (bMatch)
+		{
+			DescItems.Emplace(Transition.State, Transition.ID);
+		}
+	}
+
+	// Find states from transition tasks
+	if (EnumHasAnyFlags(Trigger, EStateTreeTransitionTrigger::OnTick | EStateTreeTransitionTrigger::OnEvent | EStateTreeTransitionTrigger::OnDelegate))
+	{
+		auto AddLinksFromStruct = [&DescItems, TreeEditorData](FStateTreeDataView Struct, const FGuid NodeID)
+		{
+			if (!Struct.IsValid())
+			{
+				return;
+			}
+			for (TPropertyValueIterator<FStructProperty> It(Struct.GetStruct(), Struct.GetMemory()); It; ++It)
+			{
+				const UScriptStruct* StructType = It.Key()->Struct;
+				if (StructType == TBaseStructure<FStateTreeStateLink>::Get())
+				{
+					const FStateTreeStateLink& Link = *static_cast<const FStateTreeStateLink*>(It.Value());
+					if (Link.LinkType != EStateTreeTransitionType::None)
+					{
+						DescItems.Emplace(Link, NodeID);
+					}
+				}
+			}
+		};
+		
+		for (const FStateTreeEditorNode& Task : State->Tasks)
+		{
+			AddLinksFromStruct(FStateTreeDataView(Task.Node.GetScriptStruct(), const_cast<uint8*>(Task.Node.GetMemory())), Task.ID);
+			AddLinksFromStruct(Task.GetInstance(), Task.ID);
+		}
+
+		AddLinksFromStruct(FStateTreeDataView(State->SingleTask.Node.GetScriptStruct(), const_cast<uint8*>(State->SingleTask.Node.GetMemory())), State->SingleTask.ID);
+		AddLinksFromStruct(State->SingleTask.GetInstance(), State->SingleTask.ID);
+	}
+
+	if (IsLeafState()
+		&& DescItems.Num() == 0
+		&& EnumHasAnyFlags(Trigger, EStateTreeTransitionTrigger::OnStateCompleted))
+	{
+		if (HasParentTransitionForTrigger(*State, Trigger))
+		{
+			DescItems.Emplace(
+				LOCTEXT("TransitionActionHandleInParentRich", "<i>Parent</>"),
+				LOCTEXT("TransitionActionHandleInParent", "Handle transition in parent State")
+				);
+		}
+		else
+		{
+			DescItems.Emplace(
+				LOCTEXT("TransitionActionRootRich", "<i>Root</>"),
+				LOCTEXT("TransitionActionRoot", "Transition to Root State.")
+				);
+		}
+	}
+
+	TSharedRef<SHorizontalBox> TransitionContainer = SNew(SHorizontalBox);
+
+	auto IsTransitionEnabledFunc = [WeakState = WeakState]
+	{
+		const UStateTreeState* State = WeakState.Get();
+		return State && State->bEnabled;
+	};
+
+	for (int32 Index = 0; Index < DescItems.Num(); Index++)
+	{
+		const FItem& Item = DescItems[Index];
+
+		if (Index > 0)
+		{
+			TransitionContainer->AddSlot()
+			.AutoWidth()
+			.VAlign(VAlign_Center)
+			[
+				SNew(STextBlock)
+				.Text(FText::FromString(TEXT(", ")))
+				.TextStyle(&FStateTreeEditorStyle::Get().GetWidgetStyle<FTextBlockStyle>("Transition.Subdued"))
+			];
+		}
+		
+		TransitionContainer->AddSlot()
+		.AutoWidth()
+		.VAlign(VAlign_Center)
+		[
+			SNew(SButton)
+			.ButtonStyle(FAppStyle::Get(), "SimpleButton")
+			.ContentPadding(FMargin(0, 0))
+			.OnClicked_Lambda([StateTreeViewModel = StateTreeViewModel, WeakState = WeakState, Item]()
+			{
+				if (StateTreeViewModel)
+				{
+					StateTreeViewModel->BringNodeToFocus(WeakState.Get(), Item.NodeID);
+				}
+				return FReply::Handled();
+			})
+			[
+				SNew(SBorder)
+				.VAlign(VAlign_Center)
+				.BorderImage(FAppStyle::GetNoBrush())
+				.Padding(0)
+				.IsEnabled_Lambda(IsTransitionEnabledFunc)
+				[
+					SNew(SRichTextBlock)
+					.Text_Lambda([WeakEditorData = WeakEditorData, Item]()
+					{
+						if (!Item.Desc.IsEmpty())
+						{
+							return Item.Desc;
+						}
+						return UE::StateTree::Editor::GetStateLinkDesc(WeakEditorData.Get(), Item.Link, EStateTreeNodeFormatting::RichText);
+					})
+					.ToolTipText_Lambda([this, Item]()
+					{
+						if (!Item.Tooltip.IsEmpty())
+						{
+							return Item.Tooltip;
+						}
+						return GetLinkTooltip(Item.Link, Item.NodeID);
+					})
+					.TextStyle(&FStateTreeEditorStyle::Get().GetWidgetStyle<FTextBlockStyle>("Transition.Normal"))
+					+SRichTextBlock::Decorator(FTextStyleDecorator::Create(TEXT(""), FStateTreeEditorStyle::Get().GetWidgetStyle<FTextBlockStyle>("Transition.Normal")))
+					+SRichTextBlock::Decorator(FTextStyleDecorator::Create(TEXT("b"), FStateTreeEditorStyle::Get().GetWidgetStyle<FTextBlockStyle>("Transition.Bold")))
+					+SRichTextBlock::Decorator(FTextStyleDecorator::Create(TEXT("i"), FStateTreeEditorStyle::Get().GetWidgetStyle<FTextBlockStyle>("Transition.Italic")))
+					+SRichTextBlock::Decorator(FTextStyleDecorator::Create(TEXT("s"), FStateTreeEditorStyle::Get().GetWidgetStyle<FTextBlockStyle>("Transition.Subdued")))
+				]
+			]
+		];
+	}
+	
+	return TransitionContainer;
+}
+
 void SStateTreeViewRow::MakeFlagsWidget()
@@ -1626,193 +1825,6 @@
 			|| State->Type == EStateTreeStateType::LinkedAsset);
 }
 
-TSharedRef<SWidget> SStateTreeViewRow::MakeTransitionWidgets(const EStateTreeTransitionTrigger Trigger, const FTransitionDescFilterOptions FilterOptions)
-{
-	const UStateTreeEditorData* TreeEditorData = WeakEditorData.Get();
-	const UStateTreeState* State = WeakState.Get();
-
-	if (!TreeEditorData || !State)
-	{
-		return SNullWidget::NullWidget;
-	}
-	
-	struct FItem
-	{
-		FItem() = default;
-		FItem(const FStateTreeStateLink& InLink, const FGuid InNodeID)
-			: Link(InLink)
-			, NodeID(InNodeID)
-		{
-		}
-		FItem(const FText& InDesc, const FText& InTooltip)
-			: Desc(InDesc)
-			, Tooltip(InTooltip)
-		{
-		}
-		
-		FText Desc;
-		FText Tooltip;
-		FStateTreeStateLink Link;
-		FGuid NodeID;
-	};
-
-	TArray<FItem> DescItems;
-
-	for (const FStateTreeTransition& Transition : State->Transitions)
-	{
-		// Apply filter for enabled/disabled transitions
-		if ((FilterOptions.Enabled == ETransitionDescRequirement::RequiredTrue && Transition.bTransitionEnabled == false)
-			|| (FilterOptions.Enabled == ETransitionDescRequirement::RequiredFalse && Transition.bTransitionEnabled))
-		{
-			continue;
-		}
-
-#if WITH_STATETREE_TRACE_DEBUGGER
-		// Apply filter for transitions with/without breakpoint
-		const bool bHasBreakpoint = TreeEditorData != nullptr && TreeEditorData->HasBreakpoint(Transition.ID, EStateTreeBreakpointType::OnTransition);
-		if ((FilterOptions.WithBreakpoint == ETransitionDescRequirement::RequiredTrue && bHasBreakpoint == false)
-			|| (FilterOptions.WithBreakpoint == ETransitionDescRequirement::RequiredFalse && bHasBreakpoint))
-		{
-			continue;
-		}
-#endif // WITH_STATETREE_TRACE_DEBUGGER
-
-		const bool bMatch = FilterOptions.bUseMask ? EnumHasAnyFlags(Transition.Trigger, Trigger) : Transition.Trigger == Trigger;
-		if (bMatch)
-		{
-			DescItems.Emplace(Transition.State, Transition.ID);
-		}
-	}
-
-	// Find states from transition tasks
-	if (EnumHasAnyFlags(Trigger, EStateTreeTransitionTrigger::OnTick | EStateTreeTransitionTrigger::OnEvent | EStateTreeTransitionTrigger::OnDelegate))
-	{
-		auto AddLinksFromStruct = [&DescItems, TreeEditorData](FStateTreeDataView Struct, const FGuid NodeID)
-		{
-			if (!Struct.IsValid())
-			{
-				return;
-			}
-			for (TPropertyValueIterator<FStructProperty> It(Struct.GetStruct(), Struct.GetMemory()); It; ++It)
-			{
-				const UScriptStruct* StructType = It.Key()->Struct;
-				if (StructType == TBaseStructure<FStateTreeStateLink>::Get())
-				{
-					const FStateTreeStateLink& Link = *static_cast<const FStateTreeStateLink*>(It.Value());
-					if (Link.LinkType != EStateTreeTransitionType::None)
-					{
-						DescItems.Emplace(Link, NodeID);
-					}
-				}
-			}
-		};
-		
-		for (const FStateTreeEditorNode& Task : State->Tasks)
-		{
-			AddLinksFromStruct(FStateTreeDataView(Task.Node.GetScriptStruct(), const_cast<uint8*>(Task.Node.GetMemory())), Task.ID);
-			AddLinksFromStruct(Task.GetInstance(), Task.ID);
-		}
-
-		AddLinksFromStruct(FStateTreeDataView(State->SingleTask.Node.GetScriptStruct(), const_cast<uint8*>(State->SingleTask.Node.GetMemory())), State->SingleTask.ID);
-		AddLinksFromStruct(State->SingleTask.GetInstance(), State->SingleTask.ID);
-	}
-
-	if (IsLeafState()
-		&& DescItems.Num() == 0
-		&& EnumHasAnyFlags(Trigger, EStateTreeTransitionTrigger::OnStateCompleted))
-	{
-		if (HasParentTransitionForTrigger(*State, Trigger))
-		{
-			DescItems.Emplace(
-				LOCTEXT("TransitionActionHandleInParentRich", "<i>Parent</>"),
-				LOCTEXT("TransitionActionHandleInParent", "Handle transition in parent State")
-				);
-		}
-		else
-		{
-			DescItems.Emplace(
-				LOCTEXT("TransitionActionRootRich", "<i>Root</>"),
-				LOCTEXT("TransitionActionRoot", "Transition to Root State.")
-				);
-		}
-	}
-
-	TSharedRef<SHorizontalBox> TransitionContainer = SNew(SHorizontalBox);
-
-	auto IsTransitionEnabledFunc = [WeakState = WeakState]
-	{
-		const UStateTreeState* State = WeakState.Get();
-		return State && State->bEnabled;
-	};
-
-	for (int32 Index = 0; Index < DescItems.Num(); Index++)
-	{
-		const FItem& Item = DescItems[Index];
-
-		if (Index > 0)
-		{
-			TransitionContainer->AddSlot()
-			.AutoWidth()
-			.VAlign(VAlign_Center)
-			[
-				SNew(STextBlock)
-				.Text(FText::FromString(TEXT(", ")))
-				.TextStyle(&FStateTreeEditorStyle::Get().GetWidgetStyle<FTextBlockStyle>("Transition.Subdued"))
-			];
-		}
-		
-		TransitionContainer->AddSlot()
-		.AutoWidth()
-		.VAlign(VAlign_Center)
-		[
-			SNew(SButton)
-			.ButtonStyle(FAppStyle::Get(), "SimpleButton")
-			.ContentPadding(FMargin(0, 0))
-			.OnClicked_Lambda([StateTreeViewModel = StateTreeViewModel, WeakState = WeakState, Item]()
-			{
-				if (StateTreeViewModel)
-				{
-					StateTreeViewModel->BringNodeToFocus(WeakState.Get(), Item.NodeID);
-				}
-				return FReply::Handled();
-			})
-			[
-				SNew(SBorder)
-				.VAlign(VAlign_Center)
-				.BorderImage(FAppStyle::GetNoBrush())
-				.Padding(0)
-				.IsEnabled_Lambda(IsTransitionEnabledFunc)
-				[
-					SNew(SRichTextBlock)
-					.Text_Lambda([WeakEditorData = WeakEditorData, Item]()
-					{
-						if (!Item.Desc.IsEmpty())
-						{
-							return Item.Desc;
-						}
-						return UE::StateTree::Editor::GetStateLinkDesc(WeakEditorData.Get(), Item.Link, EStateTreeNodeFormatting::RichText);
-					})
-					.ToolTipText_Lambda([this, Item]()
-					{
-						if (!Item.Tooltip.IsEmpty())
-						{
-							return Item.Tooltip;
-						}
-						return GetLinkTooltip(Item.Link, Item.NodeID);
-					})
-					.TextStyle(&FStateTreeEditorStyle::Get().GetWidgetStyle<FTextBlockStyle>("Transition.Normal"))
-					+SRichTextBlock::Decorator(FTextStyleDecorator::Create(TEXT(""), FStateTreeEditorStyle::Get().GetWidgetStyle<FTextBlockStyle>("Transition.Normal")))
-					+SRichTextBlock::Decorator(FTextStyleDecorator::Create(TEXT("b"), FStateTreeEditorStyle::Get().GetWidgetStyle<FTextBlockStyle>("Transition.Bold")))
-					+SRichTextBlock::Decorator(FTextStyleDecorator::Create(TEXT("i"), FStateTreeEditorStyle::Get().GetWidgetStyle<FTextBlockStyle>("Transition.Italic")))
-					+SRichTextBlock::Decorator(FTextStyleDecorator::Create(TEXT("s"), FStateTreeEditorStyle::Get().GetWidgetStyle<FTextBlockStyle>("Transition.Subdued")))
-				]
-			]
-		];
-	}
-	
-	return TransitionContainer;
-}
-
 FText SStateTreeViewRow::GetTransitionsDesc(const EStateTreeTransitionTrigger Trigger, const FTransitionDescFilterOptions FilterOptions) const
 {
 	const UStateTreeState* State = WeakState.Get();
@@ -2268,6 +2280,22 @@
 void SStateTreeViewRow::HandleAssetChanged()
 {
 	MakeFlagsWidget();
+	MakeTransitionsWidget();
 }
 
+void SStateTreeViewRow::HandleStatesChanged(const TSet<UStateTreeState*>& ChangedStates, const FPropertyChangedEvent& PropertyChangedEvent)
+{
+	if (const UStateTreeState* OwnerState = WeakState.Get())
+	{
+		if (ChangedStates.Contains(OwnerState))
+		{
+			if (PropertyChangedEvent.MemberProperty->GetFName() == GET_MEMBER_NAME_CHECKED(UStateTreeState, Transitions)
+				|| PropertyChangedEvent.Property->GetFName() == GET_MEMBER_NAME_CHECKED(FStateTreeStateLink, LinkType))
+			{
+				MakeTransitionsWidget();
+			}
+		}
+	}
+}
+
 #undef LOCTEXT_NAMESPACE
diff --git a/Engine/Plugins/Runtime/StateTree/Source/StateTreeEditorModule/Private/SStateTreeViewRow.h b/Engine/Plugins/Runtime/StateTree/Source/StateTreeEditorModule/Private/SStateTreeViewRow.h
--- a/Engine/Plugins/Runtime/StateTree/Source/StateTreeEditorModule/Private/SStateTreeViewRow.h
+++ b/Engine/Plugins/Runtime/StateTree/Source/StateTreeEditorModule/Private/SStateTreeViewRow.h
@@ -34,8 +34,31 @@ public:
 	void RequestRename() const;
 
 private:
+	enum class ETransitionDescRequirement : uint8
+	{
+		Any,
+		RequiredTrue,
+		RequiredFalse,
+	};
+
+	/**
+	 * Filtering options used to build the description of the transitions.
+	 * The default setup includes only enabled transition,
+	 * with or without breakpoints and requires exact trigger match (no partial mask)
+	 */
+	struct FTransitionDescFilterOptions
+	{
+		FTransitionDescFilterOptions() {};
+		ETransitionDescRequirement Enabled = ETransitionDescRequirement::RequiredTrue;
+		ETransitionDescRequirement WithBreakpoint = ETransitionDescRequirement::Any;
+		bool bUseMask = false;
+	};
+
 	TSharedRef<SWidget> MakeTasksWidget(const TSharedPtr<SScrollBox>& ViewBox);
 	TSharedRef<SWidget> MakeConditionsWidget(const TSharedPtr<SScrollBox>& ViewBox);
+	void MakeTransitionsWidget();
+	TSharedRef<SWidget> MakeTransitionWidget(const EStateTreeTransitionTrigger Trigger, const FSlateBrush* Icon);
+	TSharedRef<SWidget> MakeTransitionWidgetInternal(const EStateTreeTransitionTrigger Trigger, const FTransitionDescFilterOptions FilterOptions = {});
 	void MakeFlagsWidget();
 
 	FSlateColor GetTitleColor(const float Alpha = 1.0f, const float Lighten = 0.0f) const;
@@ -81,28 +104,6 @@ private:
 	FText GetStateDescription() const;
 
 	EVisibility GetTransitionDashVisibility() const;
-	
-	enum class ETransitionDescRequirement : uint8
-	{
-		Any,
-		RequiredTrue,
-		RequiredFalse,
-	};
-
-	/**
-	 * Filtering options used to build the description of the transitions.
-	 * The default setup includes only enabled transition,
-	 * with or without breakpoints and requires exact trigger match (no partial mask)
-	 */
-	struct FTransitionDescFilterOptions
-	{
-		FTransitionDescFilterOptions() {};
-		ETransitionDescRequirement Enabled = ETransitionDescRequirement::RequiredTrue;
-		ETransitionDescRequirement WithBreakpoint = ETransitionDescRequirement::Any;
-		bool bUseMask = false;
-	};
-
-	TSharedRef<SWidget> MakeTransitionWidgets(const EStateTreeTransitionTrigger Trigger, const FTransitionDescFilterOptions FilterOptions = {});
 
 	FText GetLinkTooltip(const FStateTreeStateLink& Link, const FGuid NodeID) const;
 	FText GetTransitionsDesc(const EStateTreeTransitionTrigger Trigger, const FTransitionDescFilterOptions FilterOptions = {}) const;
@@ -125,12 +126,14 @@ private:
 	FReply HandleAcceptDrop(const FDragDropEvent& DragDropEvent, EItemDropZone DropZone, TWeakObjectPtr<UStateTreeState> TargetState) const;
 
 	void HandleAssetChanged();
-
+	void HandleStatesChanged(const TSet<UStateTreeState*>& ChangedStates, const FPropertyChangedEvent& PropertyChangedEvent);
 	TSharedPtr<FStateTreeViewModel> StateTreeViewModel;
 	TWeakObjectPtr<UStateTreeState> WeakState;
 	TWeakObjectPtr<UStateTreeEditorData> WeakEditorData;
 	TSharedPtr<SInlineEditableTextBlock> NameTextBlock;
 	TSharedPtr<SBorder> FlagsContainer;
+	TSharedPtr<SHorizontalBox> TransitionsContainer;
 
 	FDelegateHandle AssetChangedHandle;
+	FDelegateHandle StatesChangedHandle;
 };
