- Only compile a StateTree in editor on load if it has EditorData (is has not been cooked)

diff --git a/Engine/Plugins/Runtime/StateTree/Source/StateTreeModule/Private/StateTree.cpp b/Engine/Plugins/Runtime/StateTree/Source/StateTreeModule/Private/StateTree.cpp
--- a/Engine/Plugins/Runtime/StateTree/Source/StateTreeModule/Private/StateTree.cpp
+++ b/Engine/Plugins/Runtime/StateTree/Source/StateTreeModule/Private/StateTree.cpp
@@ -448,8 +448,7 @@ void UStateTree::PostLoad()
 	{
 		// Make sure all the fix up logic in the editor data has had chance to happen.
 		EditorData->ConditionalPostLoad();
-	}
-	{
+
 		TGuardValueAccessors<bool> IsEditorLoadingPackageGuard(UE::GetIsEditorLoadingPackage, UE::SetIsEditorLoadingPackage, true);
 		Compile();
 	}
