[StateTree] Added asset information in some error messages. Also replaced GetFullName() by GetPathName() since the classname is useless

diff --git a/Engine/Plugins/Runtime/StateTree/Source/StateTreeModule/Private/StateTree.cpp b/Engine/Plugins/Runtime/StateTree/Source/StateTreeModule/Private/StateTree.cpp
--- a/Engine/Plugins/Runtime/StateTree/Source/StateTreeModule/Private/StateTree.cpp
+++ b/Engine/Plugins/Runtime/StateTree/Source/StateTreeModule/Private/StateTree.cpp
@@ -293,7 +293,7 @@ void UStateTree::OnObjectsReinstanced(const FReplacementObjectMap& ObjectMap)
 	{
 		if (!Link())
 		{
-			UE_LOG(LogStateTree, Error, TEXT("%s failed to link after Object reinstantiation. Take a look at the asset for any errors. Asset will not be usable at runtime."), *GetFullName());
+			UE_LOG(LogStateTree, Error, TEXT("%s failed to link after Object reinstantiation. Take a look at the asset for any errors. Asset will not be usable at runtime."), *GetPathName());
 		}
 	}
 }
@@ -310,7 +310,7 @@ void UStateTree::OnUserDefinedStructReinstanced(const UUserDefinedStruct& UserDe
 	{
 		if (!Link())
 		{
-			UE_LOG(LogStateTree, Error, TEXT("%s failed to link after Struct reinstantiation. Take a look at the asset for any errors. Asset will not be usable at runtime."), *GetFullName());
+			UE_LOG(LogStateTree, Error, TEXT("%s failed to link after Struct reinstantiation. Take a look at the asset for any errors. Asset will not be usable at runtime."), *GetPathName());
 		}
 	}
 }
@@ -381,7 +381,7 @@ EDataValidationResult UStateTree::IsDataValid(FDataValidationContext& Context) c
 			const uint32 CurrentHash = UE::StateTree::Delegates::OnRequestEditorHash.Execute(*this);
 			if (CurrentHash != LastCompiledEditorDataHash)
 			{
-				Context.AddWarning(FText::FromString(FString::Printf(TEXT("%s is not compiled. Please recompile the State Tree."), *GetFullName())));
+				Context.AddWarning(FText::FromString(FString::Printf(TEXT("%s is not compiled. Please recompile the State Tree."), *GetPathName())));
 				return EDataValidationResult::Invalid;
 			}
 		}
@@ -389,7 +389,7 @@ EDataValidationResult UStateTree::IsDataValid(FDataValidationContext& Context) c
 
 	if (!const_cast<UStateTree*>(this)->Link())
 	{
-		Context.AddError(FText::FromString(FString::Printf(TEXT("%s failed to link. Please recompile the State Tree for more details errors."), *GetFullName())));
+		Context.AddError(FText::FromString(FString::Printf(TEXT("%s failed to link. Please recompile the State Tree for more details errors."), *GetPathName())));
 		return EDataValidationResult::Invalid;
 	}
 
@@ -456,14 +456,14 @@ void UStateTree::PostLoad()
 	const int32 CurrentVersion = GetLinkerCustomVersion(FStateTreeCustomVersion::GUID);
 	if (CurrentVersion < FStateTreeCustomVersion::LatestVersion)
 	{		
-		UE_LOG(LogStateTree, Error, TEXT("%s: compiled data is in older format. Please recompile the StateTree asset."), *GetFullName());
+		UE_LOG(LogStateTree, Error, TEXT("%s: compiled data is in older format. Please recompile the StateTree asset."), *GetPathName());
 		return;
 	}
 #endif
 
 	if (!Link())
 	{
-		UE_LOG(LogStateTree, Log, TEXT("%s failed to link. Asset will not be usable at runtime."), *GetFullName());
+		UE_LOG(LogStateTree, Log, TEXT("%s failed to link. Asset will not be usable at runtime."), *GetPathName());
 	}
 }
 
@@ -564,7 +564,7 @@ bool UStateTree::ValidateInstanceData()
 
 				if (bLogError)
 				{
-					UE_LOG(LogStateTree, Error, TEXT("%s: node '%s' failed. The source Instance Data type '%s' has a newer version."), *GetFullName(), *WriteToString<64>(Node->StaticStruct()->GetFName()), *WriteToString<64>(CurrentInstanceDataType->GetFName()));
+					UE_LOG(LogStateTree, Error, TEXT("%s: node '%s' failed. The source Instance Data type '%s' has a newer version."), *GetPathName(), *WriteToString<64>(Node->StaticStruct()->GetFName()), *WriteToString<64>(CurrentInstanceDataType->GetFName()));
 				}
 
 				bResult = false;
@@ -595,7 +595,7 @@ bool UStateTree::ValidateInstanceData()
 #endif
 					if (bLogError)
 					{
-						UE_LOG(LogStateTree, Error, TEXT("%s: node '%s' failed. The source Instance Data type '%s' does not match '%s'"), *GetFullName(), *WriteToString<64>(Node->StaticStruct()->GetFName()), *GetNameSafe(CurrentInstanceDataType), *GetNameSafe(DesiredInstanceDataType));
+						UE_LOG(LogStateTree, Error, TEXT("%s: node '%s' failed. The source Instance Data type '%s' does not match '%s'"), *GetPathName(), *WriteToString<64>(Node->StaticStruct()->GetFName()), *GetNameSafe(CurrentInstanceDataType), *GetNameSafe(DesiredInstanceDataType));
 					}
 					bResult = false;
 				}
@@ -625,7 +625,7 @@ bool UStateTree::Link()
 		{
 			if (!Node.IsValid())
 			{
-				UE_LOG(LogStateTree, Error, TEXT("%s: State Tree asset was not properly loaded (missing node). See log for loading failures, or recompile the StateTree asset."), *GetFullName());
+				UE_LOG(LogStateTree, Error, TEXT("%s: State Tree asset was not properly loaded (missing node). See log for loading failures, or recompile the StateTree asset."), *GetPathName());
 				return false;
 			}
 		}
@@ -644,7 +644,7 @@ bool UStateTree::Link()
 				const bool bLinkSucceeded = NodePtr->Link(Linker);
 				if (!bLinkSucceeded || Linker.GetStatus() == EStateTreeLinkerStatus::Failed)
 				{
-					UE_LOG(LogStateTree, Error, TEXT("%s: node '%s' failed to resolve its references."), *GetFullName(), *NodePtr->StaticStruct()->GetName());
+					UE_LOG(LogStateTree, Error, TEXT("%s: node '%s' failed to resolve its references."), *GetPathName(), *NodePtr->StaticStruct()->GetName());
 					return false;
 				}
 			}
@@ -655,7 +655,7 @@ bool UStateTree::Link()
 			const bool bSchemaLinkSucceeded = Schema->Link(Linker);
 			if (!bSchemaLinkSucceeded || Linker.GetStatus() == EStateTreeLinkerStatus::Failed)
 			{
-				UE_LOG(LogStateTree, Error, TEXT("%s: schema failed to resolve its references."), *GetFullName());
+				UE_LOG(LogStateTree, Error, TEXT("%s: schema failed to resolve its references."), *GetPathName());
 				return false;
 			}
 		}
@@ -667,13 +667,13 @@ bool UStateTree::Link()
 
 	if (!DefaultInstanceData.AreAllInstancesValid())
 	{
-		UE_LOG(LogStateTree, Error, TEXT("%s: State Tree asset was not properly loaded (missing instance data). See log for loading failures, or recompile the StateTree asset."), *GetFullName());
+		UE_LOG(LogStateTree, Error, TEXT("%s: State Tree asset was not properly loaded (missing instance data). See log for loading failures, or recompile the StateTree asset."), *GetPathName());
 		return false;
 	}
 
 	if (!SharedInstanceData.AreAllInstancesValid())
 	{
-		UE_LOG(LogStateTree, Error, TEXT("%s: State Tree asset was not properly loaded (missing shared instance data). See log for loading failures, or recompile the StateTree asset."), *GetFullName());
+		UE_LOG(LogStateTree, Error, TEXT("%s: State Tree asset was not properly loaded (missing shared instance data). See log for loading failures, or recompile the StateTree asset."), *GetPathName());
 		return false;
 	}
 	
@@ -831,7 +831,7 @@ bool UStateTree::PatchBindings()
 		{
 			if (!State.ParameterTemplateIndex.IsValid())
 			{
-				UE_LOG(LogStateTree, Error, TEXT("%s: Data for state '%s' is malformed. Please recompile the StateTree asset."), *GetFullName(), *State.Name.ToString());
+				UE_LOG(LogStateTree, Error, TEXT("%s: Data for state '%s' is malformed. Please recompile the StateTree asset."), *GetPathName(), *State.Name.ToString());
 				return false;
 			}
 		}
@@ -843,7 +843,7 @@ bool UStateTree::PatchBindings()
 			FStateTreeBindableStructDesc* Desc = GetSourceStructByHandle(State.ParameterDataHandle);
 			if (!Desc)
 			{
-				UE_LOG(LogStateTree, Error, TEXT("%s: Data for state '%s' is malformed. Please recompile the StateTree asset."), *GetFullName(), *State.Name.ToString());
+				UE_LOG(LogStateTree, Error, TEXT("%s: Data for state '%s' is malformed. Please recompile the StateTree asset."), *GetPathName(), *State.Name.ToString());
 				return false;
 			}
 			Desc->Struct = Params.Parameters.GetPropertyBagStruct();
@@ -866,7 +866,7 @@ bool UStateTree::PatchBindings()
 			if (State.ParameterTemplateIndex.IsValid() == false
 				|| LinkedState.ParameterTemplateIndex.IsValid() == false)
 			{
-				UE_LOG(LogStateTree, Error, TEXT("%s: Data for state '%s' is malformed. Please recompile the StateTree asset."), *GetFullName(), *State.Name.ToString());
+				UE_LOG(LogStateTree, Error, TEXT("%s: Data for state '%s' is malformed. Please recompile the StateTree asset."), *GetPathName(), *State.Name.ToString());
 				return false;
 			}
 
@@ -876,7 +876,7 @@ bool UStateTree::PatchBindings()
 
 			if (LinkedStateParams.Parameters.GetPropertyBagStruct() != Params.Parameters.GetPropertyBagStruct())
 			{
-				UE_LOG(LogStateTree, Error, TEXT("%s: The parameters on state '%s' does not match the linked state parameters in state '%s'. Please recompile the StateTree asset."), *GetFullName(), *State.Name.ToString(), *LinkedState.Name.ToString());
+				UE_LOG(LogStateTree, Error, TEXT("%s: The parameters on state '%s' does not match the linked state parameters in state '%s'. Please recompile the StateTree asset."), *GetPathName(), *State.Name.ToString(), *LinkedState.Name.ToString());
 				return false;
 			}
 		}
@@ -889,7 +889,7 @@ bool UStateTree::PatchBindings()
 			if (TargetTreeParameters.GetPropertyBagStruct() != Params.Parameters.GetPropertyBagStruct())
 			{
 				UE_LOG(LogStateTree, Error, TEXT("%s: The parameters on state '%s' does not match the linked asset parameters '%s'. Please recompile the StateTree asset."),
-					*GetFullName(), *State.Name.ToString(), *State.LinkedAsset->GetFullName());
+					*GetPathName(), *State.Name.ToString(), *State.LinkedAsset->GetPathName());
 				return false;
 			}
 		}
@@ -974,7 +974,10 @@ bool UStateTree::PatchBindings()
 		FStateTreeDataView TargetView = GetBindingBatchDataView(FStateTreeIndex16(BatchIndex));
 		if (!TargetView.IsValid())
 		{
-			UE_LOG(LogStateTree, Error, TEXT("%hs: Invalid target struct when trying to bind to '%s'."), __FUNCTION__, *Batch.TargetStruct.Get().Name.ToString());
+			UE_LOG(LogStateTree, Error, TEXT("%hs: '%s' Invalid target struct when trying to bind to '%s'")
+				, __FUNCTION__
+				, *GetPathName()
+				, *Batch.TargetStruct.Get().Name.ToString());
 			return false;
 		}
 
@@ -992,14 +995,22 @@ bool UStateTree::PatchBindings()
 
 				if (!Binding.GetMutableSourcePath().UpdateSegmentsFromValue(SourceView, &ErrorMsg))
 				{
-					UE_LOG(LogStateTree, Error, TEXT("%hs: Failed to update source instance structs for property binding '%s'. Reason: %s"), __FUNCTION__, *Binding.GetTargetPath().ToString(), *ErrorMsg);
+					UE_LOG(LogStateTree, Error, TEXT("%hs: '%s' Failed to update source instance structs for property binding '%s'. Reason: %s")
+						, __FUNCTION__
+						, *GetPathName()
+						, *Binding.GetTargetPath().ToString()
+						, *ErrorMsg);
 					return false;
 				}
 			}
 
 			if (!Binding.GetMutableTargetPath().UpdateSegmentsFromValue(TargetView, &ErrorMsg))
 			{
-				UE_LOG(LogStateTree, Error, TEXT("%hs: Failed to update target instance structs for property binding '%s'. Reason: %s"), __FUNCTION__, *Binding.GetTargetPath().ToString(), *ErrorMsg);
+				UE_LOG(LogStateTree, Error, TEXT("%hs: '%s' Failed to update target instance structs for property binding '%s'. Reason: %s")
+					, __FUNCTION__
+					, *GetPathName()
+					, *Binding.GetTargetPath().ToString()
+					, *ErrorMsg);
 				return false;
 			}
 		}
@@ -1245,14 +1256,14 @@ void UStateTree::CompileIfChanged()
 		const uint32 CurrentHash = UE::StateTree::Delegates::OnRequestEditorHash.Execute(*this);
 		if (LastCompiledEditorDataHash != CurrentHash)
 		{
-			UE_LOG(LogStateTree, Log, TEXT("%s: Editor data has changed. Recompiling state tree."), *GetFullName());
+			UE_LOG(LogStateTree, Log, TEXT("%s: Editor data has changed. Recompiling state tree."), *GetPathName());
 			UE::StateTree::Delegates::OnRequestCompile.Execute(*this);
 		}
 	}
 	else
 	{
 		ResetCompiled();
-		UE_LOG(LogStateTree, Warning, TEXT("%s: could not compile. Please resave the StateTree asset."), *GetFullName());
+		UE_LOG(LogStateTree, Warning, TEXT("%s: could not compile. Please resave the StateTree asset."), *GetPathName());
 	}
 }
 
@@ -1260,13 +1271,13 @@ void UStateTree::Compile()
 {
 	if (UE::StateTree::Delegates::OnRequestCompile.IsBound())
 	{
-		UE_LOG(LogStateTree, Log, TEXT("%s: Editor data has changed. Recompiling state tree."), *GetFullName());
+		UE_LOG(LogStateTree, Log, TEXT("%s: Editor data has changed. Recompiling state tree."), *GetPathName());
 		UE::StateTree::Delegates::OnRequestCompile.Execute(*this);
 	}
 	else
 	{
 		ResetCompiled();
-		UE_LOG(LogStateTree, Warning, TEXT("%s: could not compile. Please resave the StateTree asset."), *GetFullName());
+		UE_LOG(LogStateTree, Warning, TEXT("%s: could not compile. Please resave the StateTree asset."), *GetPathName());
 	}
 }
 #endif //WITH_EDITOR
