StateTree: Make sure the frame is completed when the left of a frame contains a transition that retriggers.

diff --git a/Engine/Plugins/Runtime/StateTree/Source/StateTreeModule/Private/StateTreeExecutionContext.cpp b/Engine/Plugins/Runtime/StateTree/Source/StateTreeModule/Private/StateTreeExecutionContext.cpp
--- a/Engine/Plugins/Runtime/StateTree/Source/StateTreeModule/Private/StateTreeExecutionContext.cpp
+++ b/Engine/Plugins/Runtime/StateTree/Source/StateTreeModule/Private/StateTreeExecutionContext.cpp
@@ -1118,6 +1118,9 @@ void FStateTreeExecutionContext::TickTriggerTransitionsInternal()
 		return;
 	}
 
+	// Reset the completed subframe counter (for unit-test that do not recreate an execution context between each tick)
+	TriggerTransitionsFromFrameIndex.Reset();
+
 	// The state selection is repeated up to MaxIteration time. This allows failed EnterState() to potentially find a new state immediately.
 	// This helps event driven StateTrees to not require another event/tick to find a suitable state.
 	static constexpr int32 MaxIterations = 5;
@@ -4127,7 +4130,11 @@ bool FStateTreeExecutionContext::TriggerTransitions()
 		// Re-cache bHasEvents, RequestTransition above can create new events.
 		const bool bHasEvents = EventQueue && EventQueue->HasEvents();
 		const bool bHasBroadcastedDelegates = InstanceDataStorage->HasBroadcastedDelegates();
-		for (int32 FrameIndex = Exec.ActiveFrames.Num() - 1; FrameIndex >= 0; FrameIndex--)
+
+		// Transition() can TriggerTransitions() in a loop when a sub-frame completes.
+		//We do not want to evaluate the transition from that sub-frame.
+		const int32 EndFrameIndex = TriggerTransitionsFromFrameIndex.Get(Exec.ActiveFrames.Num() - 1);
+		for (int32 FrameIndex = EndFrameIndex; FrameIndex >= 0; FrameIndex--)
 		{
 			FStateTreeExecutionFrame& CurrentFrame = Exec.ActiveFrames[FrameIndex];
 			const UStateTree* CurrentStateTree = CurrentFrame.StateTree;
@@ -4598,7 +4605,7 @@ PRAGMA_ENABLE_DEPRECATION_WARNINGS
 					*GetSafeStateName(ParentFrame, ParentLinkedState), *GetSafeStateName(SourceFrame, NextTransition.SourceState), *UEnum::GetDisplayValueAsText(RunStatus).ToString());
 
 				// Set the parent linked state as last completed state, and update tick status to the status from the transition.
-				const int32 TaskIndex = INDEX_NONE;
+				constexpr int32 TaskIndex = INDEX_NONE;
 				const UE::StateTree::FActiveStateID ParentLastActiveStateID = ParentFrame.ActiveStates.StateIDs[ParentFrame.ActiveStates.Num() - 1];
 				FinishTaskInternal(UE::StateTree::FFinishedTask::EReasonType::InternalTransition, ParentFrame.FrameID, ParentLastActiveStateID, TaskIndex, RunStatus);
 				Exec.LastTickStatus = RunStatus;
@@ -4607,6 +4614,7 @@ PRAGMA_ENABLE_DEPRECATION_WARNINGS
 				// Since the LastTickStatus != running, the transition loop will try another transition
 				// now starting from the linked parent state. If we run out of retires in the selection loop (e.g. very deep hierarchy)
 				// we will continue on next tick.
+				TriggerTransitionsFromFrameIndex = ParentFrameIndex;
 				NextTransition.Reset();
 				return false;
 			}
diff --git a/Engine/Plugins/Runtime/StateTree/Source/StateTreeModule/Public/StateTreeExecutionContext.h b/Engine/Plugins/Runtime/StateTree/Source/StateTreeModule/Public/StateTreeExecutionContext.h
--- a/Engine/Plugins/Runtime/StateTree/Source/StateTreeModule/Public/StateTreeExecutionContext.h
+++ b/Engine/Plugins/Runtime/StateTree/Source/StateTreeModule/Public/StateTreeExecutionContext.h
@@ -997,6 +997,9 @@ protected:
 	/** Structure describing the origin of the state transition that caused the state change. */
 	FStateTreeTransitionSource NextTransitionSource;
 
+	/** When set, start the transitions loop from TriggerTransitionsFromFrameIndex. */
+	TOptional<int32> TriggerTransitionsFromFrameIndex;
+
 	/** Current frame we're processing. */
 	const FStateTreeExecutionFrame* CurrentlyProcessedParentFrame = nullptr; 
 	const FStateTreeExecutionFrame* CurrentlyProcessedFrame = nullptr; 
diff --git a/Engine/Plugins/Runtime/StateTree/Source/StateTreeTestSuite/Private/StateTreeLinkedStateTest.cpp b/Engine/Plugins/Runtime/StateTree/Source/StateTreeTestSuite/Private/StateTreeLinkedStateTest.cpp
--- a/Engine/Plugins/Runtime/StateTree/Source/StateTreeTestSuite/Private/StateTreeLinkedStateTest.cpp
+++ b/Engine/Plugins/Runtime/StateTree/Source/StateTreeTestSuite/Private/StateTreeLinkedStateTest.cpp
@@ -737,9 +737,148 @@ struct FStateTreeTest_RecursiveLinkedAsset : FStateTreeTestBase
 		return true;
 	}
 };
+IMPLEMENT_AI_INSTANT_TEST(FStateTreeTest_RecursiveLinkedAsset, "System.StateTree.LinkedAsset.RecursiveLinkedAsset");
 
+struct FStateTreeTest_LinkedAssetTransitionSameTick : FStateTreeTestBase
+{
+	virtual bool InstantTest() override
+	{
+		//Tree 1
+		//	Root
+		//		State1 -> Delay 1 -> StateLinkedTree1
+		//		LinkState2 (Tree2) -> Next
+		//		State3 -> Root
+		//Tree 2
+		//	Root
+		//		State1 -> Succeeded
 
-IMPLEMENT_AI_INSTANT_TEST(FStateTreeTest_RecursiveLinkedAsset, "System.StateTree.LinkedAsset.RecursiveLinkedAsset");
+		UStateTree& StateTree2 = NewStateTree();
+		UStateTreeState* Root2 = nullptr;
+		// Asset 2
+		{
+			UStateTreeEditorData& EditorData2 = *Cast<UStateTreeEditorData>(StateTree2.EditorData);
+			Root2 = &EditorData2.AddSubTree("Tree2StateRoot");
+
+			{
+				UStateTreeState& C1 = Root2->AddChildState("Tree2State1", EStateTreeStateType::State);
+				C1.AddTransition(EStateTreeTransitionTrigger::OnTick, EStateTreeTransitionType::Succeeded);
+			}
+			FStateTreeCompilerLog Log;
+			FStateTreeCompiler Compiler(Log);
+			const bool bResult = Compiler.Compile(StateTree2);
+			AITEST_TRUE("StateTree should not compiled", bResult);
+		}
+		// Asset 1
+		UStateTree& StateTree1 = NewStateTree();
+		{
+			UStateTreeState* Root1 = nullptr;
+			{
+				UStateTreeEditorData& EditorData1 = *Cast<UStateTreeEditorData>(StateTree1.EditorData);
+				Root1 = &EditorData1.AddSubTree("Tree1StateRoot");
+
+				FStateTreeCompilerLog Log;
+				FStateTreeCompiler Compiler(Log);
+				const bool bResult = Compiler.Compile(StateTree1);
+				AITEST_TRUE("StateTree 1 should get compiled", bResult);
+			}
+			{
+				TStateTreeEditorNode<FTestTask_PrintValue>& Task1 = Root1->AddTask<FTestTask_PrintValue>("Tree1StateRootTask1");
+				Task1.GetInstanceData().Value = 100;
+			}
+			{
+				UStateTreeState& State1 = Root1->AddChildState("Tree1State1", EStateTreeStateType::State);
+				TStateTreeEditorNode<FTestTask_PrintValue>& Task = State1.AddTask<FTestTask_PrintValue>("Tree1State1Task1");
+				Task.GetInstanceData().Value = 101;
+				Task.GetInstanceData().TickRunStatus = EStateTreeRunStatus::Succeeded;
+				State1.AddTransition(EStateTreeTransitionTrigger::OnStateCompleted, EStateTreeTransitionType::NextState);
+			}
+			{
+				UStateTreeState& LinkState2 = Root1->AddChildState("Tree1State2LinkedTree2", EStateTreeStateType::LinkedAsset);
+				LinkState2.SetLinkedStateAsset(&StateTree2);
+				LinkState2.AddTransition(EStateTreeTransitionTrigger::OnStateSucceeded, EStateTreeTransitionType::NextState);
+			}
+			{
+				UStateTreeState& State3 = Root1->AddChildState("Tree1State3", EStateTreeStateType::State);
+				TStateTreeEditorNode<FTestTask_PrintValue>& Task = State3.AddTask<FTestTask_PrintValue>("Tree1State3Task1");
+				Task.GetInstanceData().Value = 103;
+				FStateTreeTransition& Transition = State3.AddTransition(EStateTreeTransitionTrigger::OnTick, EStateTreeTransitionType::GotoState, Root1);
+				Transition.bDelayTransition = true;
+				Transition.DelayDuration = 1.0f;
+			}
+
+			FStateTreeCompilerLog Log;
+			FStateTreeCompiler Compiler(Log);
+			const bool bResult = Compiler.Compile(StateTree1);
+			AITEST_TRUE("StateTree should get compiled", bResult);
+		}
+
+		// Run test
+		{
+			FStateTreeInstanceData InstanceData;
+			FTestStateTreeExecutionContext Exec(StateTree1, StateTree1, InstanceData);
+			{
+				const bool bInitSucceeded = Exec.IsValid();
+				AITEST_TRUE("StateTree should init", bInitSucceeded);
+			}
+
+			{
+				EStateTreeRunStatus Status = Exec.Start();
+				AITEST_EQUAL(TEXT("Start should complete with Running"), Status, EStateTreeRunStatus::Running);
+				AITEST_TRUE(TEXT("In correct states"), Exec.ExpectInActiveStates("Tree1StateRoot", "Tree1State1"));
+				FTestStateTreeExecutionContext::FLogOrder LogOrder = Exec.Expect("Tree1StateRootTask1", TEXT("EnterState100"));
+				AITEST_TRUE(TEXT("Start should enter Tree1StateRootTask1"), LogOrder);
+				LogOrder = LogOrder.Then("Tree1State1Task1", TEXT("EnterState101"));
+				AITEST_TRUE(TEXT("Start should enter Tree1State1"), LogOrder);
+				Exec.LogClear();
+			}
+			{
+				EStateTreeRunStatus Status = Exec.Tick(1.01f);
+				AITEST_EQUAL(TEXT("Tick should complete with Running"), Status, EStateTreeRunStatus::Running);
+				AITEST_TRUE(TEXT("In correct states"), Exec.ExpectInActiveStates("Tree1StateRoot", "Tree1State2LinkedTree2", "Tree2StateRoot", "Tree2State1"));
+				FTestStateTreeExecutionContext::FLogOrder LogOrder = Exec.Expect("Tree1StateRootTask1", TEXT("Tick100"));
+				AITEST_TRUE(TEXT("Start should tick Tree1StateRootTask1"), LogOrder);
+				LogOrder = LogOrder.Then("Tree1State1Task1", TEXT("Tick101"));
+				AITEST_TRUE(TEXT("Start should tick Tree1State1"), LogOrder);
+				LogOrder = LogOrder.Then("Tree1State1Task1", TEXT("ExitState101"));
+				AITEST_TRUE(TEXT("Start should exit Tree1State1"), LogOrder);
+				Exec.LogClear();
+			}
+			//Tree2State1 -> Succeeded should transition to Tree1State3
+			{
+				EStateTreeRunStatus Status = Exec.Tick(1.01f);
+				AITEST_EQUAL(TEXT("Tick should complete with Running"), Status, EStateTreeRunStatus::Running);
+				AITEST_TRUE(TEXT("In correct states"), Exec.ExpectInActiveStates("Tree1StateRoot", "Tree1State3"));
+				FTestStateTreeExecutionContext::FLogOrder LogOrder = Exec.Expect("Tree1StateRootTask1", TEXT("Tick100"));
+				AITEST_TRUE(TEXT("Start should tick Tree1StateRootTask1"), LogOrder);
+				LogOrder = LogOrder.Then("Tree1State3Task1", TEXT("EnterState103"));
+				AITEST_TRUE(TEXT("Start should enter Tree1State3"), LogOrder);
+				Exec.LogClear();
+			}
+			{
+				EStateTreeRunStatus Status = Exec.Tick(1.01f);
+				AITEST_EQUAL(TEXT("Tick should complete with Running"), Status, EStateTreeRunStatus::Running);
+				AITEST_TRUE(TEXT("In correct states"), Exec.ExpectInActiveStates("Tree1StateRoot", "Tree1State3"));
+				Exec.LogClear();
+			}
+			{
+				EStateTreeRunStatus Status = Exec.Tick(1.01f);
+				AITEST_EQUAL(TEXT("Tick should complete with Running"), Status, EStateTreeRunStatus::Running);
+				AITEST_TRUE(TEXT("In correct states"), Exec.ExpectInActiveStates("Tree1StateRoot", "Tree1State1"));
+				FTestStateTreeExecutionContext::FLogOrder LogOrder = Exec.Expect("Tree1StateRootTask1", TEXT("Tick100"));
+				AITEST_TRUE(TEXT("Start should tick Tree1StateRootTask1"), LogOrder);
+				LogOrder = LogOrder.Then("Tree1State3Task1", TEXT("ExitState103"));
+				AITEST_TRUE(TEXT("Start should exit Tree1State3"), LogOrder);
+				LogOrder = LogOrder.Then("Tree1State1Task1", TEXT("EnterState101"));
+				AITEST_TRUE(TEXT("Start should enter Tree1State1"), LogOrder);
+				Exec.LogClear();
+			}
+			Exec.Stop();
+		}
+
+		return true;
+	}
+};
+IMPLEMENT_AI_INSTANT_TEST(FStateTreeTest_LinkedAssetTransitionSameTick, "System.StateTree.LinkedAsset.TransitionSameTick");
 
 } // namespace UE::StateTree::Tests
 
