StateTree: Remove CompletedStateRunStatus, the task can finish from a weak and it doesn't bring any benefit. Some frames can occur between the tasks tick and the transitions tick. Add a ticked flag to the finished task to know if the task was processed during the tasks tick phase. Only remove finished tasks that are processed.

diff --git a/Engine/Plugins/Runtime/StateTree/Source/StateTreeModule/Private/StateTreeExecutionContext.cpp b/Engine/Plugins/Runtime/StateTree/Source/StateTreeModule/Private/StateTreeExecutionContext.cpp
--- a/Engine/Plugins/Runtime/StateTree/Source/StateTreeModule/Private/StateTreeExecutionContext.cpp
+++ b/Engine/Plugins/Runtime/StateTree/Source/StateTreeModule/Private/StateTreeExecutionContext.cpp
@@ -22,12 +22,6 @@
 #define STATETREE_LOG(Verbosity, Format, ...) UE_VLOG_UELOG(GetOwner(), LogStateTree, Verbosity, TEXT("%s: ") Format, *GetInstanceDescriptionInternal(), ##__VA_ARGS__)
 #define STATETREE_CLOG(Condition, Verbosity, Format, ...) UE_CVLOG_UELOG((Condition), GetOwner(), LogStateTree, Verbosity, TEXT("%s: ") Format, *GetInstanceDescriptionInternal(), ##__VA_ARGS__)
 
-#if WITH_STATETREE_DEBUG
-#define UE_STATETREE_ENSURE_COMPLETED_STATE_RUN_STATUS() ensureMsgf(!CompletedStateRunStatus.IsSet(), TEXT("No public entry point should be able to set the completed state. Internal error."))
-#else
-#define UE_STATETREE_ENSURE_COMPLETED_STATE_RUN_STATUS()
-#endif //WITH_STATETREE_DEBUG
-
 namespace UE::StateTree::Debug
 {
 	constexpr int32 IndentSize = 2;	// Debug printing indent for hierarchical data.
@@ -104,6 +98,14 @@ namespace UE::StateTree::ExecutionContext::Private
 		static_assert(NumEStateTreeRunStatus() == UE_ARRAY_COUNT(PriorityMatrix), "The number of entries in EStateTreeRunStatus changed. Update GetPriorityRunStatus.");
 		return PriorityMatrix[(uint8)A] > PriorityMatrix[(uint8)B] ? A : B;
 	}
+
+	void RemoveTickProcessedFinishedTasks(FStateTreeExecutionState& Exec)
+	{
+		Exec.FinishedTasks.RemoveAllSwap([](const UE::StateTree::FFinishedTask& Task)
+			{
+				return Task.bTickProcessed;
+			});
+	}
 }
 
 FStateTreeMinimalExecutionContext::FStateTreeMinimalExecutionContext(UObject& InOwner, const UStateTree& InStateTree, FStateTreeInstanceData& InInstanceData)
@@ -462,6 +464,7 @@ FStateTreeExecutionContext::FStateTreeExecutionContext(const FStateTreeExecution
 	}
 }
 
+PRAGMA_DISABLE_DEPRECATION_WARNINGS
 FStateTreeExecutionContext::~FStateTreeExecutionContext()
 {
 	// Mark external data indices as invalid
@@ -471,6 +474,7 @@ FStateTreeExecutionContext::~FStateTreeExecutionContext()
 		Frame.ExternalDataBaseIndex = {};
 	}
 }
+PRAGMA_ENABLE_DEPRECATION_WARNINGS
 
 void FStateTreeExecutionContext::SetCollectExternalDataCallback(const FOnCollectStateTreeExternalData& Callback)
 {
@@ -820,7 +824,7 @@ EStateTreeRunStatus FStateTreeExecutionContext::Start(FStartParameters Parameter
 				Transition.NextActiveFrames = StateSelectionResult.GetSelectedFrames(); // Enter state will update Exec.ActiveFrames.
 				Transition.NextActiveFrameEvents = StateSelectionResult.GetFramesStateSelectionEvents();
 				const EStateTreeRunStatus LastTickStatus = EnterState(Transition);
-			
+
 				Exec.LastTickStatus = LastTickStatus;
 
 				// Report state completed immediately.
@@ -977,8 +981,6 @@ EStateTreeRunStatus FStateTreeExecutionContext::TickPrelude()
 	// From this point any calls to Stop should be deferred.
 	SetUpdatePhaseInExecutionState(Exec, EStateTreeUpdatePhase::TickStateTree);
 
-	UpdateCompletedStateList();
-
 	return EStateTreeRunStatus::Running;
 }
 
@@ -1051,6 +1053,14 @@ EStateTreeRunStatus FStateTreeExecutionContext::TickTriggerTransitions()
 		return PreludeResult;
 	}
 
+	if (!bCompletedStatesUpdated)
+	{
+		// Create the completed state list from the process tasks list.
+		//Do not mark the task as processed because the task could be added between TickTasks() and TickTriggerTransitions()
+		constexpr bool bMarkTaskProcessed = false;
+		UpdateCompletedStateList(bMarkTaskProcessed);
+	}
+
 	// If we have pending global task that completd, then stop the tree (like it would with a regular Tick).
 	//This should not occurs and its probably a user mistake to not tick the system in the proper order.
 	const EStateTreeRunStatus GlobalTaskStatus = GetGlobalTasksCompletedStatesStatus();
@@ -1080,6 +1090,10 @@ void FStateTreeExecutionContext::TickUpdateTasksInternal(const float DeltaTime)
 		DelayedState.TimeLeft -= DeltaTime;
 	}
 
+	// Create the completed list and mark the task as processed.
+	constexpr bool bMarkTaskProcessed = true;
+	UpdateCompletedStateList(bMarkTaskProcessed);
+
 	// Tick global evaluators and tasks.
 	const bool bTickGlobalTasks = true;
 	const EStateTreeRunStatus EvalAndGlobalTaskStatus = TickEvaluatorsAndGlobalTasks(DeltaTime, bTickGlobalTasks);
@@ -1090,8 +1104,8 @@ void FStateTreeExecutionContext::TickUpdateTasksInternal(const float DeltaTime)
 			// Tick tasks on active states.
 			Exec.LastTickStatus = TickTasks(DeltaTime);
 
-			// Report state completed immediately.
-			if (Exec.LastTickStatus != EStateTreeRunStatus::Running)
+			// Report state completed immediately (and no global task completes)
+			if (Exec.LastTickStatus != EStateTreeRunStatus::Running && Exec.RequestedStop == EStateTreeRunStatus::Unset)
 			{
 				StateCompleted();
 			}
@@ -1311,6 +1325,8 @@ void FStateTreeExecutionContext::RequestTransition(FStateTreeStateHandle InTarge
 
 void FStateTreeExecutionContext::FinishTask(const FStateTreeTaskBase& Task, EStateTreeFinishTaskType FinishType)
 {
+	using namespace UE::StateTree;
+
 	if (!IsValid())
 	{
 		STATETREE_LOG(Warning, TEXT("%hs: StateTree context is not initialized properly ('%s' using StateTree '%s')"),
@@ -1325,33 +1341,37 @@ void FStateTreeExecutionContext::FinishTask(const FStateTreeTaskBase& Task, ESta
 	}
 	check(CurrentlyProcessedFrame);
 	check(CurrentlyProcessedState.IsValid());
+	check(CurrentNodeIndex >= 0);
 
-	using namespace UE::StateTree;
-	using namespace UE::StateTree::ExecutionContext::Private;
 
 	FStateTreeExecutionState& Exec = GetExecState();
 
-	// Then we have the current frame and the current state.
 	const FActiveFrameID ActiveFrameID = CurrentlyProcessedFrame->FrameID;
 	check(ActiveFrameID.IsValid());
+	const FStateTreeIndex16 TaskIndex = FStateTreeIndex16(CurrentNodeIndex);
+	const EStateTreeRunStatus RunStatus = ExecutionContext::Private::GetRunStatusFromFinishTaskType(FinishType);
 	if (CurrentlyProcessedState.IsValid())
 	{
 		const int32 ActiveStateIndex = CurrentlyProcessedFrame->ActiveStates.IndexOfReverse(CurrentlyProcessedState);
 		check(ActiveStateIndex != INDEX_NONE);
 		const FActiveStateID ActiveStateID = CurrentlyProcessedFrame->ActiveStates.StateIDs[ActiveStateIndex];
-		FinishTaskInternal(FFinishedTask::EReasonType::StateTask, ActiveFrameID, ActiveStateID, CurrentNodeIndex, GetRunStatusFromFinishTaskType(FinishType));
+		constexpr FFinishedTask::EReasonType Reason = FFinishedTask::EReasonType::StateTask;
+		constexpr bool bTickProcessed = false;
+		Exec.FinishedTasks.Emplace(ActiveFrameID, ActiveStateID, TaskIndex, RunStatus, Reason, bTickProcessed);
 	}
 	else
 	{
-		FinishTaskInternal(FFinishedTask::EReasonType::GlobalTask, ActiveFrameID, FActiveStateID::Invalid, CurrentNodeIndex, GetRunStatusFromFinishTaskType(FinishType));
+		constexpr FFinishedTask::EReasonType Reason = FFinishedTask::EReasonType::GlobalTask;
+		constexpr bool bTickProcessed = false;
+		Exec.FinishedTasks.Emplace(ActiveFrameID, FActiveStateID::Invalid, TaskIndex, RunStatus, Reason, bTickProcessed);
 	}
-
-	// The FStateTreeTaskBase::Tick or FStateTreeTaskBase::OnEnter can returns a value. Set it here and test it later.
-	CompletedStateRunStatus = GetRunStatusFromFinishTaskType(FinishType);
 }
 
+// Deprecated
 void FStateTreeExecutionContext::FinishTask(const UE::StateTree::FFinishedTask& Task, EStateTreeFinishTaskType FinishType)
 {
+	using namespace UE::StateTree;
+
 	if (!IsValid())
 	{
 		STATETREE_LOG(Warning, TEXT("%hs: StateTree context is not initialized properly ('%s' using StateTree '%s')"),
@@ -1359,18 +1379,20 @@ void FStateTreeExecutionContext::FinishTask(const UE::StateTree::FFinishedTask&
 		return;
 	}
 
+	FStateTreeExecutionState& Exec = GetExecState();
+
 	if (CurrentNodeIndex == Task.TaskIndex.AsInt32() && CurrentlyProcessedFrame && Task.FrameID == CurrentlyProcessedFrame->FrameID)
 	{
-		// We are currently ticking that task. We are either in Tick, EnterState, ExitState, Transition.
-		//Note: Completed tasks are not considered during ExitState or Transition.
-		//The node will return the value and it will be processed once returned. Set it here and test it later.
-		CompletedStateRunStatus = UE::StateTree::ExecutionContext::Private:: GetRunStatusFromFinishTaskType(FinishType);
+		FFinishedTask& NewTask = Exec.FinishedTasks.Add_GetRef(Task);
+		// Deprecated path. Need to reset the flag.
+		NewTask.bTickProcessed = false;
 	}
 	else if (IsFinishedTaskValid(Task))
 	{
-		FStateTreeExecutionState& Exec = GetExecState();
-		const UE::StateTree::FFinishedTask& CompletedTask = Exec.FinishedTasks.Add_GetRef(Task);
+		FFinishedTask& CompletedTask = Exec.FinishedTasks.Add_GetRef(Task);
 		MarkStateCompleted(CompletedTask);
+		// Deprecated path. Need to reset the flag to the original value.
+		CompletedTask.bTickProcessed = false;
 		ScheduleNextTick();
 	}
 }
@@ -1381,15 +1403,12 @@ void FStateTreeExecutionContext::FinishTaskInternal(UE::StateTree::FFinishedTask
 
 	FStateTreeExecutionState& Exec = GetExecState();
 
+	// Add it for later. In case the user called TickTasks and will call TickTransitions later.
 	const FStateTreeIndex16 TaskIndex16 = TaskIndex >= 0 ? FStateTreeIndex16(TaskIndex) : FStateTreeIndex16::Invalid;
-
-	const FFinishedTask Task {.FrameID = ActiveFrameID, .StateID = ActiveStateID, .TaskIndex = TaskIndex16, .RunStatus = RunStatus, .Reason = Reason };
-	if (IsFinishedTaskValid(Task))
-	{
-		// Add it for later. In case the user called TickTasks and will call TickTransitions later.
-		Exec.FinishedTasks.Add(Task);
-		MarkStateCompleted(Task);
-	}
+	
+	constexpr bool bTickProcessed = true;
+	UE::StateTree::FFinishedTask& Task = Exec.FinishedTasks.Emplace_GetRef(ActiveFrameID, ActiveStateID, TaskIndex16, RunStatus, Reason, bTickProcessed);
+	MarkStateCompleted(Task);
 }
 
 bool FStateTreeExecutionContext::IsFinishedTaskValid(const UE::StateTree::FFinishedTask& Task) const
@@ -1457,8 +1476,13 @@ bool FStateTreeExecutionContext::IsFinishedTaskValid(const UE::StateTree::FFinis
 
 void FStateTreeExecutionContext::UpdateCompletedStateList()
 {
-	// Do not modify the PendingCompletedState until transitions completion.
-	//User can call TickTask and wait a couple of frame before calling TickTransitions.
+	constexpr bool bMarkTaskProcessed = false;
+	UpdateCompletedStateList(bMarkTaskProcessed);
+}
+
+void FStateTreeExecutionContext::UpdateCompletedStateList(bool bMarkTaskProcessed)
+{
+	bCompletedStatesUpdated = true;
 
 	// Recreate the completed states list from the pending list
 	CompletedStates.Empty();
@@ -1466,20 +1490,26 @@ void FStateTreeExecutionContext::UpdateCompletedStateList()
 	FStateTreeExecutionState& Exec = GetExecState();
 	for (int32 Index = Exec.FinishedTasks.Num() - 1; Index >= 0; --Index)
 	{
-		const UE::StateTree::FFinishedTask& Task = Exec.FinishedTasks[Index];
-		if (IsFinishedTaskValid(Task))
+		UE::StateTree::FFinishedTask& Task = Exec.FinishedTasks[Index];
+		// Because the user can tick only the transition (and not the task), we should not consider any finished tasks that were not ticked.
+		if (Task.bTickProcessed || bMarkTaskProcessed)
 		{
-			MarkStateCompleted(Task);
-		}
-		else
-		{
-			Exec.FinishedTasks.RemoveAtSwap(Index);
+			if (IsFinishedTaskValid(Task))
+			{
+				MarkStateCompleted(Task);
+			}
+			else
+			{
+				Exec.FinishedTasks.RemoveAtSwap(Index);
+			}
 		}
 	}
 }
 
-void FStateTreeExecutionContext::MarkStateCompleted(const UE::StateTree::FFinishedTask& NewFinishedTask)
+void FStateTreeExecutionContext::MarkStateCompleted(UE::StateTree::FFinishedTask& NewFinishedTask)
 {
+	NewFinishedTask.bTickProcessed = true;
+
 	bool bAdd = true;
 	// The priority is global tasks and failure. Only keep 1 FinishedTask per Frame per State.
 	//If the global task exist and a state (of the same tree) completes, only keep the global
@@ -2374,7 +2404,7 @@ EStateTreeRunStatus FStateTreeExecutionContext::EnterState(FStateTreeTransitionR
 	CaptureNewStateEvents(Exec.ActiveFrames, Transition.NextActiveFrames, Transition.NextActiveFrameEvents);
 
 	Exec.StateChangeCount++;
-	Exec.FinishedTasks.Empty();
+	UE::StateTree::ExecutionContext::Private::RemoveTickProcessedFinishedTasks(Exec);
 	Exec.EnterStateFailedFrameIndex = FStateTreeIndex16::Invalid; // This will make all tasks to be accepted.
 	Exec.EnterStateFailedTaskIndex = FStateTreeIndex16::Invalid; // This will make all tasks to be accepted.
 	
@@ -2394,7 +2424,6 @@ EStateTreeRunStatus FStateTreeExecutionContext::EnterState(FStateTreeTransitionR
 
 	// Reset the current active frames, new ones are added one by one.
 	Exec.ActiveFrames.Reset();
-	CompletedStateRunStatus.Reset();
 
 	// Track any changed state (i.e., not sustained) to prevent reused subtrees from being considered as sustained states since the whole
 	// tree is reused and should be considered as a new instance.
@@ -2406,6 +2435,7 @@ EStateTreeRunStatus FStateTreeExecutionContext::EnterState(FStateTreeTransitionR
 		
 		FStateTreeExecutionFrame* CurrentParentFrame = !Exec.ActiveFrames.IsEmpty() ? &Exec.ActiveFrames.Last() : nullptr;
 		FStateTreeExecutionFrame& CurrentFrame = Exec.ActiveFrames.Add_GetRef(NextFrame);
+		const UE::StateTree::FActiveFrameID CurrentFrameID = NextFrame.FrameID;
 		
 		// We'll add new states one by one, so that active states contain only the states which have EnterState called.
 		CurrentFrame.ActiveStates.Reset();
@@ -2424,6 +2454,7 @@ EStateTreeRunStatus FStateTreeExecutionContext::EnterState(FStateTreeTransitionR
 		for (int32 StateIndex = 0; StateIndex < NextFrame.ActiveStates.Num() && Result != EStateTreeRunStatus::Failed; StateIndex++)
 		{
 			const FStateTreeStateHandle CurrentHandle = NextFrame.ActiveStates[StateIndex];
+			const UE::StateTree::FActiveStateID CurrentStateID = NextFrame.ActiveStates.StateIDs[StateIndex];
 			const FStateTreeStateHandle PreviousHandle = PreviousActiveStates.GetStateSafe(StateIndex);
 			const FCompactStateTreeState& State = CurrentStateTree->States[CurrentHandle.Index];
 
@@ -2540,36 +2571,47 @@ EStateTreeRunStatus FStateTreeExecutionContext::EnterState(FStateTreeTransitionR
 				if (bIsEnteringState && bShouldCallEnterState)
 				{
 					STATETREE_LOG(Verbose, TEXT("%*sTask '%s'.EnterState()"), (FrameIndex + StateIndex + 1)*UE::StateTree::Debug::IndentSize, TEXT(""), *Task.Name.ToString());
-					UE_STATETREE_ENSURE_COMPLETED_STATE_RUN_STATUS();
 					UE_STATETREE_DEBUG_TASK_ENTER_STATE(this, NextFrame.StateTree, FStateTreeIndex16(TaskIndex));
 
-					EStateTreeRunStatus TaskResult = EStateTreeRunStatus::Unset;
+					const int32 PreviousFinishedTasksNum = Exec.FinishedTasks.Num();
+					EStateTreeRunStatus TaskStatus = EStateTreeRunStatus::Unset;
 					{
 						QUICK_SCOPE_CYCLE_COUNTER(StateTree_Task_EnterState);
 						CSV_SCOPED_TIMING_STAT_EXCLUSIVE(StateTree_Task_EnterState);
 					
-						TaskResult = Task.EnterState(*this, CurrentTransition);
+						TaskStatus = Task.EnterState(*this, CurrentTransition);
 					}
 
-					// Delayed from FinishTask(), use the best priority from what the task returns and FinishTask().
-					if (CompletedStateRunStatus.IsSet())
+					// Process FinishTask (weak, or this context)
+					bool bFinishTaskInternal = true;
+					for (int32 FinishedTaskIndex = PreviousFinishedTasksNum; FinishedTaskIndex < Exec.FinishedTasks.Num(); ++FinishedTaskIndex)
 					{
-						TaskResult = UE::StateTree::ExecutionContext::Private::GetPriorityRunStatus(TaskResult, CompletedStateRunStatus.GetValue());
-						CompletedStateRunStatus.Reset();
+						UE::StateTree::FFinishedTask& FinishedTask = Exec.FinishedTasks[FinishedTaskIndex];
+						// If same task, then update TaskStatus.
+						if (TaskIndex == FinishedTask.TaskIndex.AsInt32()
+							&& CurrentStateID == FinishedTask.StateID
+							&& CurrentFrameID == FinishedTask.FrameID)
+						{
+							// Delayed from FinishTask(), use the best priority from what the task returns and FinishTask().
+							TaskStatus = UE::StateTree::ExecutionContext::Private::GetPriorityRunStatus(TaskStatus, FinishedTask.RunStatus);
+							bFinishTaskInternal = false;
+							MarkStateCompleted(FinishedTask);
+						}
 					}
-					UE_STATETREE_ENSURE_COMPLETED_STATE_RUN_STATUS();
-					UE_STATETREE_DEBUG_TASK_EVENT(this, TaskIndex, TaskInstanceView, EStateTreeTraceEventType::OnEntered, TaskResult);
 
-					if (TaskResult != EStateTreeRunStatus::Running)
+					UE_STATETREE_DEBUG_TASK_EVENT(this, TaskIndex, TaskInstanceView, EStateTreeTraceEventType::OnEntered, TaskStatus);
+
+					if (TaskStatus != EStateTreeRunStatus::Running && bFinishTaskInternal)
 					{
-						Result = UE::StateTree::ExecutionContext::Private::GetPriorityRunStatus(Result, TaskResult);
-						FinishTaskInternal(UE::StateTree::FFinishedTask::EReasonType::StateTask, CurrentFrame.FrameID, CurrentFrame.ActiveStates.StateIDs[StateIndex], TaskIndex, TaskResult);
-						if (TaskResult == EStateTreeRunStatus::Failed)
-						{
-							Exec.EnterStateFailedFrameIndex = FStateTreeIndex16(FrameIndex);
-							Exec.EnterStateFailedTaskIndex = FStateTreeIndex16(TaskIndex);
-							break;
-						}
+						FinishTaskInternal(UE::StateTree::FFinishedTask::EReasonType::StateTask, CurrentFrame.FrameID, CurrentFrame.ActiveStates.StateIDs[StateIndex], TaskIndex, TaskStatus);
+					}
+
+					Result = UE::StateTree::ExecutionContext::Private::GetPriorityRunStatus(Result, TaskStatus);
+					if (Result == EStateTreeRunStatus::Failed)
+					{
+						Exec.EnterStateFailedFrameIndex = FStateTreeIndex16(FrameIndex);
+						Exec.EnterStateFailedTaskIndex = FStateTreeIndex16(TaskIndex);
+						break;
 					}
 				}
 			}
@@ -2989,7 +3031,6 @@ EStateTreeRunStatus FStateTreeExecutionContext::TickEvaluatorsAndGlobalTasks(con
 {
 	// When a global task is completed it completes the tree execution.
 	// A global task can complete async. See CompletedStates.
-	// CompletedStateRunStatus is used to detect that a global task is completed with FinishTask instead of a global task returning a completed result.
 	// When a global task fails, stop ticking the following tasks. If bForceCopyBoundPropertiesOnNonTickedTask, update their bindings.
 
 	CSV_SCOPED_TIMING_STAT_EXCLUSIVE(StateTree_TickEvaluators);
@@ -3000,14 +3041,21 @@ EStateTreeRunStatus FStateTreeExecutionContext::TickEvaluatorsAndGlobalTasks(con
 	FStateTreeExecutionState& Exec = GetExecState();
 
 	EStateTreeRunStatus Result = EStateTreeRunStatus::Running;
-	CompletedStateRunStatus.Reset();
 
 	const bool bForceCopyBoundPropertiesOnNonTickedTask = UE::StateTree::ExecutionContext::Private::bForceCopyBoundPropertiesOnNonTickedTask;
 
+	FTickTaskArguments TickArgs;
+	TickArgs.DeltaTime = DeltaTime;
+	TickArgs.bIsGlobalTasks = true;
+	TickArgs.bShouldTickTasks = true;
+
 	for (int32 FrameIndex = 0; FrameIndex < Exec.ActiveFrames.Num(); FrameIndex++)
 	{
 		FStateTreeExecutionFrame* CurrentParentFrame = FrameIndex > 0 ? &Exec.ActiveFrames[FrameIndex - 1] : nullptr;
 		FStateTreeExecutionFrame& CurrentFrame = Exec.ActiveFrames[FrameIndex];
+		TickArgs.ParentFrame = CurrentParentFrame;
+		TickArgs.Frame = &CurrentFrame;
+
 		if (CurrentFrame.bIsGlobalFrame)
 		{
 			FCurrentlyProcessedFrameScope FrameScope(*this, CurrentParentFrame, CurrentFrame);
@@ -3041,10 +3089,12 @@ EStateTreeRunStatus FStateTreeExecutionContext::TickEvaluatorsAndGlobalTasks(con
 
 			const FCompletedState* FrameCompletedState = CompletedStates.FindByPredicate([FrameID = CurrentFrame.FrameID](FCompletedState& Other)
 				{
-					return Other.FrameID == FrameID && Other.Reason == UE::StateTree::FFinishedTask::EReasonType::GlobalTask;
+					return Other.FrameID == FrameID
+						&& Other.Reason == UE::StateTree::FFinishedTask::EReasonType::GlobalTask;
 				});
 			if (FrameCompletedState)
 			{
+				// Completed global task stops the frame execution.
 				bShouldTickTasks = false;
 				Result = UE::StateTree::ExecutionContext::Private::GetPriorityRunStatus(Result, FrameCompletedState->RunStatus);
 			}
@@ -3052,77 +3102,20 @@ EStateTreeRunStatus FStateTreeExecutionContext::TickEvaluatorsAndGlobalTasks(con
 			if (bTickGlobalTasks)
 			{
 				const bool bHasEvents = EventQueue && EventQueue->HasEvents();
-
 				if (bForceCopyBoundPropertiesOnNonTickedTask || CurrentStateTree->ShouldTickGlobalTasks(bHasEvents))
 				{
-					for (int32 TaskIndex = CurrentStateTree->GlobalTasksBegin; TaskIndex < (CurrentStateTree->GlobalTasksBegin + CurrentStateTree->GlobalTasksNum); TaskIndex++)
-					{
-						const FStateTreeTaskBase& Task = CurrentStateTree->Nodes[TaskIndex].Get<const FStateTreeTaskBase>();
-						const FStateTreeDataView TaskInstanceView = GetDataView(CurrentParentFrame, CurrentFrame, Task.InstanceDataHandle);
-						FNodeInstanceDataScope DataScope(*this, &Task, TaskIndex, Task.InstanceDataHandle, TaskInstanceView);
-
-						// Ignore disabled task
-						if (Task.bTaskEnabled == false)
-						{
-							STATETREE_LOG(VeryVerbose, TEXT("%*sSkipped 'Tick' for disabled Task: '%s'"), UE::StateTree::Debug::IndentSize, TEXT(""), *Task.Name.ToString());
-							continue;
-						}
+					// Update Tasks data and tick if possible (ie. if no task has yet failed and bShouldTickTasks is true)
+					TickArgs.TasksBegin = CurrentStateTree->GlobalTasksBegin;
+					TickArgs.TasksEnd = CurrentStateTree->GlobalTasksBegin + CurrentStateTree->GlobalTasksNum;
+					TickArgs.Indent = FrameIndex + 1;
+					const FTickTaskResult TickTasksResult = TickTasks(TickArgs);
 
-						const bool bNeedsTick = bShouldTickTasks && (Task.bShouldCallTick || (bHasEvents && Task.bShouldCallTickOnlyOnEvents));
-						STATETREE_LOG(VeryVerbose, TEXT("  Tick: '%s' %s"), *Task.Name.ToString(), !bNeedsTick ? TEXT("[not ticked]") : TEXT(""));
-
-						// Copy bound properties.
-						// Only copy properties when the task is actually ticked, and copy properties at tick is requested.
-						const bool bCopyBatch = (bForceCopyBoundPropertiesOnNonTickedTask || bNeedsTick)
-							&& Task.BindingsBatch.IsValid()
-							&& Task.bShouldCopyBoundPropertiesOnTick;
-						if (bCopyBatch)
-						{
-							CopyBatchOnActiveInstances(CurrentParentFrame, CurrentFrame, TaskInstanceView, Task.BindingsBatch);
-						}
-
-						if (!bNeedsTick)
-						{
-							continue;
-						}
+					TickArgs.bShouldTickTasks = TickTasksResult.bShouldTickTasks;
+					Result = UE::StateTree::ExecutionContext::Private::GetPriorityRunStatus(Result, TickTasksResult.Status);
 
-						//UE_STATETREE_DEBUG_TASK_EVENT(this, TaskIndex, TaskDataView, EStateTreeTraceEventType::OnTickingTask, EStateTreeRunStatus::Running);
-						UE_STATETREE_ENSURE_COMPLETED_STATE_RUN_STATUS();
-						UE_STATETREE_DEBUG_TASK_TICK(this, CurrentStateTree, TaskIndex);
-
-						EStateTreeRunStatus TaskResult = EStateTreeRunStatus::Unset;
-						{
-							QUICK_SCOPE_CYCLE_COUNTER(StateTree_Task_Tick);
-							CSV_SCOPED_TIMING_STAT_EXCLUSIVE(StateTree_Task_Tick);
-
-							TaskResult = Task.Tick(*this, DeltaTime);
-						}
-
-						// Delayed from FinishTask(), use the best priority from what the task returns and FinishTask().
-						if (CompletedStateRunStatus.IsSet())
-						{
-							TaskResult = UE::StateTree::ExecutionContext::Private::GetPriorityRunStatus(TaskResult, CompletedStateRunStatus.GetValue());
-							CompletedStateRunStatus.Reset();
-						}
-						UE_STATETREE_ENSURE_COMPLETED_STATE_RUN_STATUS();
-						UE_STATETREE_DEBUG_TASK_EVENT(this, TaskIndex, TaskInstanceView,
-							TaskResult != EStateTreeRunStatus::Running ? EStateTreeTraceEventType::OnTaskCompleted : EStateTreeTraceEventType::OnTicked,
-							TaskResult);
-
-						// If a global task succeeds or fails, it will stop the whole tree.
-						if (TaskResult != EStateTreeRunStatus::Running)
-						{
-							Result = UE::StateTree::ExecutionContext::Private::GetPriorityRunStatus(Result, TaskResult);
-							FinishTaskInternal(UE::StateTree::FFinishedTask::EReasonType::GlobalTask, CurrentFrame.FrameID, UE::StateTree::FActiveStateID::Invalid, TaskIndex, TaskResult);
-							if (TaskResult == EStateTreeRunStatus::Failed)
-							{
-								bShouldTickTasks = false;
-								if (!bForceCopyBoundPropertiesOnNonTickedTask)
-								{
-									break; // go to next frame
-								}
-							}
-						}
+					if (!bForceCopyBoundPropertiesOnNonTickedTask && !TickTasksResult.bShouldTickTasks)
+					{
+						break;
 					}
 				}
 			}
@@ -3148,6 +3141,7 @@ EStateTreeRunStatus FStateTreeExecutionContext::StartEvaluatorsAndGlobalTasks(FS
 	{
 		FStateTreeExecutionFrame* CurrentParentFrame = FrameIndex > 0 ? &Exec.ActiveFrames[FrameIndex - 1] : nullptr;
 		FStateTreeExecutionFrame& CurrentFrame = Exec.ActiveFrames[FrameIndex];
+		const UE::StateTree::FActiveFrameID CurrentFrameID = CurrentFrame.FrameID;
 		if (CurrentFrame.bIsGlobalFrame)
 		{
 			FCurrentlyProcessedFrameScope FrameScope(*this, CurrentParentFrame, CurrentFrame);
@@ -3179,7 +3173,6 @@ EStateTreeRunStatus FStateTreeExecutionContext::StartEvaluatorsAndGlobalTasks(FS
 			// Start Global tasks
 			// Even if we call Enter/ExitState() on global tasks, they do not enter any specific state.
 			const FStateTreeTransitionResult Transition = {}; // Empty transition
-			CompletedStateRunStatus.Reset();
 
 			for (int32 TaskIndex = CurrentStateTree->GlobalTasksBegin; TaskIndex < (CurrentStateTree->GlobalTasksBegin + CurrentStateTree->GlobalTasksNum); TaskIndex++)
 			{
@@ -3201,28 +3194,45 @@ EStateTreeRunStatus FStateTreeExecutionContext::StartEvaluatorsAndGlobalTasks(FS
 				}
 
 				STATETREE_LOG(Verbose, TEXT("  Start: '%s'"), *Task.Name.ToString());
-				UE_STATETREE_ENSURE_COMPLETED_STATE_RUN_STATUS();
 				UE_STATETREE_DEBUG_TASK_ENTER_STATE(this, CurrentStateTree, FStateTreeIndex16(TaskIndex));
+
+				const int32 PreviousFinishedTasksNum = Exec.FinishedTasks.Num();
+				EStateTreeRunStatus TaskStatus = EStateTreeRunStatus::Unset;
 				{
 					QUICK_SCOPE_CYCLE_COUNTER(StateTree_Task_TreeStart);
-					EStateTreeRunStatus TaskStatus = Task.EnterState(*this, Transition);
 
-					// Delayed from FinishTask(), use the best priority from what the task returns and FinishTask().
-					if (CompletedStateRunStatus.IsSet())
-					{
-						TaskStatus = UE::StateTree::ExecutionContext::Private::GetPriorityRunStatus(TaskStatus, CompletedStateRunStatus.GetValue());
-						CompletedStateRunStatus.Reset();
-					}
-					UE_STATETREE_ENSURE_COMPLETED_STATE_RUN_STATUS();
-					UE_STATETREE_DEBUG_TASK_EVENT(this, TaskIndex, TaskInstanceView, EStateTreeTraceEventType::OnEntered, TaskStatus);
+					TaskStatus = Task.EnterState(*this, Transition);
+				}
 
-					if (TaskStatus != EStateTreeRunStatus::Running)
+				// Process FinishTask (weak, or this context)
+				bool bFinishTaskInternal = true;
+				for (int32 FinishedTaskIndex = PreviousFinishedTasksNum; FinishedTaskIndex < Exec.FinishedTasks.Num(); ++FinishedTaskIndex)
+				{
+					// If same task, then update TaskStatus.
+					UE::StateTree::FFinishedTask& FinishedTask = Exec.FinishedTasks[FinishedTaskIndex];
+					if (TaskIndex == FinishedTask.TaskIndex.AsInt32()
+						&& CurrentFrameID == FinishedTask.FrameID)
 					{
-						OutLastInitializedTaskIndex = FStateTreeIndex16(TaskIndex);
-						Result = TaskStatus;
-						break;
+						ensureMsgf(!FinishedTask.StateID.IsValid(), TEXT("EnterState a global task. The state should be invalid."));
+						// Delayed from FinishTask(), use the best priority from what the task returns and mark the state completed.
+						TaskStatus = UE::StateTree::ExecutionContext::Private::GetPriorityRunStatus(TaskStatus, FinishedTask.RunStatus);
+						bFinishTaskInternal = false;
+						MarkStateCompleted(FinishedTask);
 					}
 				}
+
+				UE_STATETREE_DEBUG_TASK_EVENT(this, TaskIndex, TaskInstanceView, EStateTreeTraceEventType::OnEntered, TaskStatus);
+
+				if (TaskStatus != EStateTreeRunStatus::Running && bFinishTaskInternal)
+				{
+					FinishTaskInternal(UE::StateTree::FFinishedTask::EReasonType::GlobalTask, CurrentFrame.FrameID, UE::StateTree::FActiveStateID::Invalid, TaskIndex, TaskStatus);					
+				}
+
+				Result = UE::StateTree::ExecutionContext::Private::GetPriorityRunStatus(Result, TaskStatus);
+				if (Result != EStateTreeRunStatus::Running)
+				{
+					OutLastInitializedTaskIndex = FStateTreeIndex16(TaskIndex);
+				}
 			}
 		}
 	}
@@ -3362,10 +3372,13 @@ EStateTreeRunStatus FStateTreeExecutionContext::StartTemporaryEvaluatorsAndGloba
 	CSV_SCOPED_TIMING_STAT_EXCLUSIVE(StateTree_StartEvaluators);
 	STATETREE_LOG(Verbose, TEXT("Start Temporary Evaluators & Global tasks while trying to select linked asset: %s"), *GetNameSafe(CurrentFrame.StateTree));
 
+	FStateTreeExecutionState& Exec = GetExecState();
+
 	FCurrentlyProcessedFrameScope FrameScope(*this, CurrentParentFrame, CurrentFrame);
 	UE_STATETREE_DEBUG_SCOPED_PHASE(this, EStateTreeUpdatePhase::StartGlobalTasksForSelection);
 
 	EStateTreeRunStatus Result = EStateTreeRunStatus::Running;
+	const UE::StateTree::FActiveFrameID CurrentFrameID = CurrentFrame.FrameID;
 	const UStateTree* CurrentStateTree = CurrentFrame.StateTree;
 
 	// Start evaluators
@@ -3426,24 +3439,42 @@ EStateTreeRunStatus FStateTreeExecutionContext::StartTemporaryEvaluatorsAndGloba
 		}
 
 		STATETREE_LOG(Verbose, TEXT("  Start: '%s'"), *Task.Name.ToString());
-		UE_STATETREE_ENSURE_COMPLETED_STATE_RUN_STATUS();
 		UE_STATETREE_DEBUG_TASK_ENTER_STATE(this, CurrentStateTree, FStateTreeIndex16(TaskIndex));
-		QUICK_SCOPE_CYCLE_COUNTER(StateTree_Task_TreeStart);
 
-		EStateTreeRunStatus TaskStatus = Task.EnterState(*this, Transition);
+		const int32 PreviousFinishedTasksNum = Exec.FinishedTasks.Num();
+		EStateTreeRunStatus TaskStatus = EStateTreeRunStatus::Unset;
+		{
+			QUICK_SCOPE_CYCLE_COUNTER(StateTree_Task_TreeStart);
+			TaskStatus = Task.EnterState(*this, Transition);
+		}
 
-		// Delayed from FinishTask(), use the best priority from what the task returns and FinishTask().
-		if (CompletedStateRunStatus.IsSet())
+		// Process FinishTask (weak, or this context)
+		bool bFinishTaskInternal = true;
+		for (int32 FinishedTaskIndex = PreviousFinishedTasksNum; FinishedTaskIndex < Exec.FinishedTasks.Num(); ++FinishedTaskIndex)
 		{
-			TaskStatus = UE::StateTree::ExecutionContext::Private::GetPriorityRunStatus(TaskStatus, CompletedStateRunStatus.GetValue());
-			CompletedStateRunStatus.Reset();
+			// If same task, then update TaskStatus.
+			UE::StateTree::FFinishedTask& FinishedTask = Exec.FinishedTasks[FinishedTaskIndex];
+			if (TaskIndex == FinishedTask.TaskIndex.AsInt32()
+				&& CurrentFrameID == FinishedTask.FrameID)
+			{
+				ensureMsgf(!FinishedTask.StateID.IsValid(), TEXT("EnterState a global task. The state should be invalid."));
+				// Delayed from FinishTask(), use the best priority from what the task returns and mark the state completed.
+				TaskStatus = UE::StateTree::ExecutionContext::Private::GetPriorityRunStatus(TaskStatus, FinishedTask.RunStatus);
+				bFinishTaskInternal = false;
+				MarkStateCompleted(FinishedTask);
+			}
 		}
-		UE_STATETREE_ENSURE_COMPLETED_STATE_RUN_STATUS();
+
 		UE_STATETREE_DEBUG_TASK_EVENT(this, TaskIndex, TaskDataView, EStateTreeTraceEventType::OnEntered, TaskStatus);
 
-		if (TaskStatus != EStateTreeRunStatus::Running)
+		if (TaskStatus != EStateTreeRunStatus::Running && bFinishTaskInternal)
+		{
+			FinishTaskInternal(UE::StateTree::FFinishedTask::EReasonType::GlobalTask, CurrentFrame.FrameID, UE::StateTree::FActiveStateID::Invalid, TaskIndex, TaskStatus);			
+		}
+
+		Result = UE::StateTree::ExecutionContext::Private::GetPriorityRunStatus(Result, TaskStatus);
+		if (Result != EStateTreeRunStatus::Running)
 		{
-			Result = TaskStatus;
 			break;
 		}
 	}
@@ -3521,13 +3552,14 @@ EStateTreeRunStatus FStateTreeExecutionContext::TickTasks(const float DeltaTime)
 {
 	// When a task is completed it also completes the state and triggers the completion transition (because LastTickStatus is set).
 	// A task can complete async. See CompletedStates.
-	// CompletedStateRunStatus is used to detect that a task is completed with FinishTask instead of the task returning a completed result.
 	// When a task fails, stop ticking the following tasks. If bForceCopyBoundPropertiesOnNonTickedTask, update their bindings.
 	// When no task ticks, then the leaf completes.
 
 	CSV_SCOPED_TIMING_STAT_EXCLUSIVE(StateTree_TickTasks);
 	UE_STATETREE_DEBUG_SCOPED_PHASE(this, EStateTreeUpdatePhase::TickingTasks);
 
+	using namespace UE::StateTree;
+
 	FStateTreeExecutionState& Exec = GetExecState();
 
 	if (Exec.ActiveFrames.IsEmpty())
@@ -3538,117 +3570,69 @@ EStateTreeRunStatus FStateTreeExecutionContext::TickTasks(const float DeltaTime)
 	EStateTreeRunStatus Result = EStateTreeRunStatus::Running;
 	int32 NumTotalEnabledTasks = 0;
 
-	const bool bHasEvents = EventQueue && EventQueue->HasEvents();
-	
-	const bool bForceCopyBoundPropertiesOnNonTickedTask = UE::StateTree::ExecutionContext::Private::bForceCopyBoundPropertiesOnNonTickedTask;
-	bool bShouldTickTasks = CompletedStates.Num() == 0;
-	CompletedStateRunStatus.Reset();
+	const bool bForceCopyBoundPropertiesOnNonTickedTask = ExecutionContext::Private::bForceCopyBoundPropertiesOnNonTickedTask;
+
+	FTickTaskArguments TickArgs;
+	TickArgs.DeltaTime = DeltaTime;
+	TickArgs.bIsGlobalTasks = false;
+	TickArgs.bShouldTickTasks = CompletedStates.Num() == 0;
 
 	STATETREE_CLOG(Exec.ActiveFrames.Num() > 0, VeryVerbose, TEXT("Ticking Tasks"));
 
 	for (int32 FrameIndex = 0; FrameIndex < Exec.ActiveFrames.Num(); FrameIndex++)
 	{
-		const FStateTreeExecutionFrame* CurrentParentFrame = FrameIndex > 0 ? &Exec.ActiveFrames[FrameIndex - 1] : nullptr;
-		const FStateTreeExecutionFrame& CurrentFrame = Exec.ActiveFrames[FrameIndex];
-		const UStateTree* CurrentStateTree = CurrentFrame.StateTree;
-
-		FCurrentlyProcessedFrameScope FrameScope(*this, CurrentParentFrame, CurrentFrame);
+		TickArgs.ParentFrame = FrameIndex > 0 ? &Exec.ActiveFrames[FrameIndex - 1] : nullptr;
+		TickArgs.Frame = &Exec.ActiveFrames[FrameIndex];
+		const UStateTree* CurrentStateTree = TickArgs.Frame->StateTree;
+		FCurrentlyProcessedFrameScope FrameScope(*this, TickArgs.ParentFrame, *TickArgs.Frame);
 
-		for (int32 StateIndex = 0; StateIndex < CurrentFrame.ActiveStates.Num(); StateIndex++)
+		for (int32 StateIndex = 0; StateIndex < TickArgs.Frame->ActiveStates.Num(); StateIndex++)
 		{
-			const FStateTreeStateHandle CurrentHandle = CurrentFrame.ActiveStates[StateIndex];
+			const FStateTreeStateHandle CurrentHandle = TickArgs.Frame->ActiveStates[StateIndex];
+			TickArgs.StateID = TickArgs.Frame->ActiveStates.StateIDs[StateIndex];
 			const FCompactStateTreeState& State = CurrentStateTree->States[CurrentHandle.Index];
 
 			FCurrentlyProcessedStateScope StateScope(*this, CurrentHandle);
 			UE_STATETREE_DEBUG_SCOPED_STATE(this, CurrentHandle);
 
-			STATETREE_CLOG(State.TasksNum > 0, VeryVerbose, TEXT("%*sState '%s'"), (FrameIndex + StateIndex + 1)*UE::StateTree::Debug::IndentSize, TEXT("")
-				, *DebugGetStatePath(Exec.ActiveFrames, &CurrentFrame, StateIndex));
+			STATETREE_CLOG(State.TasksNum > 0, VeryVerbose, TEXT("%*sState '%s'")
+				, (FrameIndex + StateIndex + 1)*UE::StateTree::Debug::IndentSize, TEXT("")
+				, *DebugGetStatePath(Exec.ActiveFrames, TickArgs.Frame, StateIndex));
 
 			if (State.Type == EStateTreeStateType::Linked || State.Type == EStateTreeStateType::LinkedAsset)
 			{
 				if (State.ParameterDataHandle.IsValid() && State.ParameterBindingsBatch.IsValid())
 				{
-					const FStateTreeDataView StateParamsDataView = GetDataView(CurrentParentFrame, CurrentFrame, State.ParameterDataHandle);
-					CopyBatchOnActiveInstances(CurrentParentFrame, CurrentFrame, StateParamsDataView, State.ParameterBindingsBatch);
+					const FStateTreeDataView StateParamsDataView = GetDataView(TickArgs.ParentFrame, *TickArgs.Frame, State.ParameterDataHandle);
+					CopyBatchOnActiveInstances(TickArgs.ParentFrame, *TickArgs.Frame, StateParamsDataView, State.ParameterBindingsBatch);
 				}
 			}
 
+			const bool bHasEvents = EventQueue && EventQueue->HasEvents();
 			if (bForceCopyBoundPropertiesOnNonTickedTask || State.ShouldTickTasks(bHasEvents))
 			{
-				// Update Tasks data and tick if possible (ie. if no task has yet failed and so bShouldTickTasks is true)
-				for (int32 TaskIndex = State.TasksBegin; TaskIndex < (State.TasksBegin + State.TasksNum); TaskIndex++)
-				{
-					const FStateTreeTaskBase& Task = CurrentStateTree->Nodes[TaskIndex].Get<const FStateTreeTaskBase>();
-					const FStateTreeDataView TaskInstanceView = GetDataView(CurrentParentFrame, CurrentFrame, Task.InstanceDataHandle);
-					FNodeInstanceDataScope DataScope(*this, &Task, TaskIndex, Task.InstanceDataHandle, TaskInstanceView);
-
-					// Ignore disabled task
-					if (Task.bTaskEnabled == false)
-					{
-						STATETREE_LOG(VeryVerbose, TEXT("%*sSkipped 'Tick' for disabled Task: '%s'"), UE::StateTree::Debug::IndentSize, TEXT(""), *Task.Name.ToString());
-						continue;
-					}
-
-					const bool bNeedsTick = bShouldTickTasks && (Task.bShouldCallTick || (bHasEvents && Task.bShouldCallTickOnlyOnEvents));
-					STATETREE_LOG(VeryVerbose, TEXT("%*s  Tick: '%s' %s"), (FrameIndex + StateIndex + 1)*UE::StateTree::Debug::IndentSize, TEXT("")
-						, *Task.Name.ToString()
-						, !bNeedsTick ? TEXT("[not ticked]") : TEXT(""));
-
-					// Copy bound properties.
-					// Only copy properties when the task is actually ticked, and copy properties at tick is requested.
-					const bool bCopyBatch = (bForceCopyBoundPropertiesOnNonTickedTask || bNeedsTick)
-						&& Task.BindingsBatch.IsValid()
-						&& Task.bShouldCopyBoundPropertiesOnTick;
-					if (bCopyBatch)
-					{
-						CopyBatchOnActiveInstances(CurrentParentFrame, CurrentFrame, TaskInstanceView, Task.BindingsBatch);
-					}
+				// Update Tasks data and tick if possible (ie. if no task has yet failed and bShouldTickTasks is true)
+				TickArgs.TasksBegin = State.TasksBegin;
+				TickArgs.TasksEnd = State.TasksBegin + State.TasksNum;
+				TickArgs.Indent = (FrameIndex + StateIndex + 1);
+				const FTickTaskResult TickTasksResult = TickTasks(TickArgs);
 
-					if (!bNeedsTick)
-					{
-						continue;
-					}
+				TickArgs.bShouldTickTasks = TickTasksResult.bShouldTickTasks;
+				Result = UE::StateTree::ExecutionContext::Private::GetPriorityRunStatus(Result, TickTasksResult.Status);
 
-					//UE_STATETREE_DEBUG_TASK_EVENT(this, TaskIndex, TaskDataView, EStateTreeTraceEventType::OnTickingTask, EStateTreeRunStatus::Running);
-					UE_STATETREE_ENSURE_COMPLETED_STATE_RUN_STATUS();
-					UE_STATETREE_DEBUG_TASK_TICK(this, CurrentStateTree, TaskIndex);
-
-					EStateTreeRunStatus TaskResult = EStateTreeRunStatus::Unset;
-					{
-						QUICK_SCOPE_CYCLE_COUNTER(StateTree_Task_Tick);
-						CSV_SCOPED_TIMING_STAT_EXCLUSIVE(StateTree_Task_Tick);
-
-						TaskResult = Task.Tick(*this, DeltaTime);
-					}
-
-					// Delayed from FinishTask(), use the best priority from what the task returns and the FinishTask().
-					if (CompletedStateRunStatus.IsSet())
-					{
-						Result = UE::StateTree::ExecutionContext::Private::GetPriorityRunStatus(Result, CompletedStateRunStatus.GetValue());
-						CompletedStateRunStatus.Reset();
-					}
-					UE_STATETREE_ENSURE_COMPLETED_STATE_RUN_STATUS();
-					UE_STATETREE_DEBUG_TASK_EVENT(this, TaskIndex, TaskInstanceView,
-						TaskResult != EStateTreeRunStatus::Running ? EStateTreeTraceEventType::OnTaskCompleted : EStateTreeTraceEventType::OnTicked,
-						TaskResult);
-
-					if (TaskResult != EStateTreeRunStatus::Running)
-					{
-						Result = UE::StateTree::ExecutionContext::Private::GetPriorityRunStatus(Result, TaskResult);
-						FinishTaskInternal(UE::StateTree::FFinishedTask::EReasonType::StateTask, CurrentFrame.FrameID, CurrentFrame.ActiveStates.StateIDs[StateIndex], TaskIndex, TaskResult);
-						if (Result == EStateTreeRunStatus::Failed)
-						{
-							bShouldTickTasks = false;
-							if (!bForceCopyBoundPropertiesOnNonTickedTask)
-							{
-								break; // go to next frame
-							}
-						}
-					}
+				if (!bForceCopyBoundPropertiesOnNonTickedTask && !TickTasksResult.bShouldTickTasks)
+				{
+					// Stop ticking
+					break;
 				}
 			}
 
+			if (Exec.RequestedStop != EStateTreeRunStatus::Unset)
+			{
+				// Stop the state tasks execution
+				return Result;
+			}
+
 			NumTotalEnabledTasks += State.EnabledTasksNum;
 		}
 	}
@@ -3656,8 +3640,22 @@ EStateTreeRunStatus FStateTreeExecutionContext::TickTasks(const float DeltaTime)
 	// Delayed finished task but no task were ticked to return the value.
 	if (Result == EStateTreeRunStatus::Running && CompletedStates.Num() > 0)
 	{
-		const FCompletedState* FirstFailureCompletedState = CompletedStates.FindByPredicate([](const FCompletedState& Other) { return Other.RunStatus != EStateTreeRunStatus::Succeeded; });
-		Result = FirstFailureCompletedState == nullptr ? EStateTreeRunStatus::Succeeded : FirstFailureCompletedState->RunStatus;
+		// Return the "worst" run status from the completed states.
+		//Since no states/tasks ticked, the result didn't have the chance to get updated.
+		for (const FCompletedState& CompletedState : CompletedStates)
+		{
+			if (CompletedState.Reason == UE::StateTree::FFinishedTask::EReasonType::GlobalTask)
+			{
+				// The value was set here in TickTasks. RequestedStop should already have the correct value.
+				//Or the value was set in GlobalTickTasks.
+				//Don't take any chances and set it back.
+				Exec.RequestedStop = ExecutionContext::Private::GetPriorityRunStatus(Exec.RequestedStop, CompletedState.RunStatus);
+			}
+			else
+			{
+				Result = UE::StateTree::ExecutionContext::Private::GetPriorityRunStatus(Result, CompletedState.RunStatus);
+			}
+		}
 	}
 	else if (NumTotalEnabledTasks == 0 && CompletedStates.Num() == 0)
 	{
@@ -3679,6 +3677,143 @@ EStateTreeRunStatus FStateTreeExecutionContext::TickTasks(const float DeltaTime)
 	return Result;
 }
 
+FStateTreeExecutionContext::FTickTaskResult FStateTreeExecutionContext::TickTasks(const FTickTaskArguments& Args)
+{
+	using namespace UE::StateTree;
+
+	check(Args.Frame);
+
+	EStateTreeRunStatus Result = EStateTreeRunStatus::Unset;
+	bool bShouldTickTasks = Args.bShouldTickTasks;
+
+	FStateTreeExecutionState& Exec = GetExecState();
+	const bool bCopyBoundPropertiesOnNonTickedTask = ExecutionContext::Private::bForceCopyBoundPropertiesOnNonTickedTask;
+	const UStateTree* CurrentStateTree = Args.Frame->StateTree;
+	const FActiveFrameID CurrentFrameID = Args.Frame->FrameID;
+	check(CurrentStateTree);
+
+	for (int32 TaskIndex = Args.TasksBegin; TaskIndex < Args.TasksEnd; ++TaskIndex)
+	{
+		const FStateTreeTaskBase& Task = CurrentStateTree->Nodes[TaskIndex].Get<const FStateTreeTaskBase>();
+
+		// Ignore disabled task
+		if (Task.bTaskEnabled == false)
+		{
+			STATETREE_LOG(VeryVerbose, TEXT("%*sSkipped 'Tick' for disabled Task: '%s'"), Debug::IndentSize, TEXT(""), *Task.Name.ToString());
+			continue;
+		}
+
+		const FStateTreeDataView TaskInstanceView = GetDataView(Args.ParentFrame, *Args.Frame, Task.InstanceDataHandle);
+		FNodeInstanceDataScope DataScope(*this, &Task, TaskIndex, Task.InstanceDataHandle, TaskInstanceView);
+
+		const bool bHasEvents = EventQueue && EventQueue->HasEvents();
+		const bool bNeedsTick = bShouldTickTasks
+			&& (Task.bShouldCallTick || (bHasEvents && Task.bShouldCallTickOnlyOnEvents));
+		STATETREE_LOG(VeryVerbose, TEXT("%*s  Tick: '%s' %s"), Args.Indent * Debug::IndentSize, TEXT("")
+			, *Task.Name.ToString()
+			, !bNeedsTick ? TEXT("[not ticked]") : TEXT(""));
+
+		// Copy bound properties.
+		// Only copy properties when the task is actually ticked, and copy properties at tick is requested.
+		const bool bCopyBatch = (bCopyBoundPropertiesOnNonTickedTask || bNeedsTick)
+			&& Task.BindingsBatch.IsValid()
+			&& Task.bShouldCopyBoundPropertiesOnTick;
+		if (bCopyBatch)
+		{
+			CopyBatchOnActiveInstances(Args.ParentFrame, *Args.Frame, TaskInstanceView, Task.BindingsBatch);
+		}
+
+		if (!bNeedsTick)
+		{
+			continue;
+		}
+
+		//UE_STATETREE_DEBUG_TASK_EVENT(this, TaskIndex, TaskDataView, EStateTreeTraceEventType::OnTickingTask, EStateTreeRunStatus::Running);
+		UE_STATETREE_DEBUG_TASK_TICK(this, CurrentStateTree, TaskIndex);
+
+		const int32 PreviousFinishedTasksNum = Exec.FinishedTasks.Num();
+		EStateTreeRunStatus TaskResult = EStateTreeRunStatus::Unset;
+		{
+			QUICK_SCOPE_CYCLE_COUNTER(StateTree_Task_Tick);
+			CSV_SCOPED_TIMING_STAT_EXCLUSIVE(StateTree_Task_Tick);
+
+			TaskResult = Task.Tick(*this, Args.DeltaTime);
+		}
+
+		// Process FinishTask (weak context, or this context)
+		bool bFinishTaskInternal = true;
+		for (int32 FinishedTaskIndex = PreviousFinishedTasksNum; FinishedTaskIndex < Exec.FinishedTasks.Num(); ++FinishedTaskIndex)
+		{
+			FFinishedTask& FinishedTask = Exec.FinishedTasks[FinishedTaskIndex];
+			if (!IsFinishedTaskValid(FinishedTask))
+			{
+				continue;
+			}
+
+			// If same task, then update TaskStatus.
+			// If state, then continue unless it's a failure.
+			// If global task, then mark the completed execution.
+			if (TaskIndex == FinishedTask.TaskIndex.AsInt32()
+				&& (Args.StateID == FinishedTask.StateID || Args.bIsGlobalTasks)
+				&& CurrentFrameID == FinishedTask.FrameID)
+			{
+				if (Args.bIsGlobalTasks)
+				{
+					ensureMsgf(!FinishedTask.StateID.IsValid(), TEXT("Ticking a global task. The state should be invalid."));
+				}
+				// Delayed from FinishTask(), use the best priority from what the task returns and mark the state completed.
+				TaskResult = ExecutionContext::Private::GetPriorityRunStatus(TaskResult, FinishedTask.RunStatus);
+				bFinishTaskInternal = false;
+				MarkStateCompleted(FinishedTask);
+			}
+			else
+			{
+				MarkStateCompleted(FinishedTask);
+				if (Args.bIsGlobalTasks)
+				{
+					if (FinishedTask.Reason == FFinishedTask::EReasonType::GlobalTask)
+					{
+						// Mark it completed and continue ticking the other tasks.
+						// If it's a failure,. then the loop will stop bellow. We do the trace first.
+						Result = ExecutionContext::Private::GetPriorityRunStatus(Result, FinishedTask.RunStatus);
+					}
+					// Else will be processed in TickTasks()
+				}
+				else
+				{
+					if (FinishedTask.Reason == FFinishedTask::EReasonType::GlobalTask)
+					{
+						Exec.RequestedStop = ExecutionContext::Private::GetPriorityRunStatus(Exec.RequestedStop, FinishedTask.RunStatus);
+					}
+					else
+					{
+						// If it's a failure, then the loops will stop bellow. We do the trace first.
+						Result = ExecutionContext::Private::GetPriorityRunStatus(Result, FinishedTask.RunStatus);
+					}
+				}
+			}
+		}
+
+		UE_STATETREE_DEBUG_TASK_EVENT(this, TaskIndex, TaskInstanceView,
+			TaskResult != EStateTreeRunStatus::Running ? EStateTreeTraceEventType::OnTaskCompleted : EStateTreeTraceEventType::OnTicked,
+			TaskResult);
+
+		if (TaskResult != EStateTreeRunStatus::Running && bFinishTaskInternal)
+		{
+			const FFinishedTask::EReasonType Reason = Args.bIsGlobalTasks ? FFinishedTask::EReasonType::GlobalTask : FFinishedTask::EReasonType::StateTask;
+			ensureMsgf(Args.bIsGlobalTasks != Args.StateID.IsValid(), TEXT("The state ID needs to be valid when the we are ticking state tasks."));
+			FinishTaskInternal(Reason, CurrentFrameID, Args.StateID, TaskIndex, TaskResult);
+		}
+
+		Result = ExecutionContext::Private::GetPriorityRunStatus(Result, TaskResult);
+		if (Result == EStateTreeRunStatus::Failed)
+		{
+			bShouldTickTasks = false;
+		}
+	}
+	return FTickTaskResult{Result, bShouldTickTasks};
+}
+
 bool FStateTreeExecutionContext::TestAllConditions(const FStateTreeExecutionFrame* CurrentParentFrame, const FStateTreeExecutionFrame& CurrentFrame, const int32 ConditionsOffset, const int32 ConditionsNum)
 {
 	CSV_SCOPED_TIMING_STAT_EXCLUSIVE(StateTree_TestConditions);
@@ -4465,7 +4600,7 @@ PRAGMA_ENABLE_DEPRECATION_WARNINGS
 			}
 		}
 
-		Exec.FinishedTasks.Empty();
+		UE::StateTree::ExecutionContext::Private::RemoveTickProcessedFinishedTasks(Exec);
 		CompletedStates.Empty();
 
 		// Start from the last completed state if specified, else only check the last state from the last frame.
@@ -4579,7 +4714,7 @@ PRAGMA_ENABLE_DEPRECATION_WARNINGS
 	}
 	else
 	{
-		Exec.FinishedTasks.Empty();
+		UE::StateTree::ExecutionContext::Private::RemoveTickProcessedFinishedTasks(Exec);
 		CompletedStates.Empty();
 	}
 
diff --git a/Engine/Plugins/Runtime/StateTree/Source/StateTreeModule/Public/StateTreeExecutionContext.h b/Engine/Plugins/Runtime/StateTree/Source/StateTreeModule/Public/StateTreeExecutionContext.h
--- a/Engine/Plugins/Runtime/StateTree/Source/StateTreeModule/Public/StateTreeExecutionContext.h
+++ b/Engine/Plugins/Runtime/StateTree/Source/StateTreeModule/Public/StateTreeExecutionContext.h
@@ -764,6 +764,27 @@ protected:
 	 */
 	EStateTreeRunStatus TickTasks(const float DeltaTime);
 
+	struct FTickTaskResult
+	{
+		EStateTreeRunStatus Status;
+		bool bShouldTickTasks = true;
+	};
+	struct FTickTaskArguments
+	{
+		FTickTaskArguments() = default;
+		float DeltaTime = 0.f;
+		int32 TasksBegin = 0;
+		int32 TasksEnd = 0;
+		int32 Indent = 0;
+		const FStateTreeExecutionFrame* ParentFrame = nullptr;
+		const FStateTreeExecutionFrame* Frame = nullptr;
+		UE::StateTree::FActiveStateID StateID;
+		bool bIsGlobalTasks = false;
+		bool bShouldTickTasks = true;
+	};
+	/** Ticks tasks for a specific state or frame. */
+	FTickTaskResult TickTasks(const FTickTaskArguments& Args);
+
 	/** Common functionality shared by the tick methods. */
 	EStateTreeRunStatus TickPrelude();
 	EStateTreeRunStatus TickPostlude();
@@ -826,11 +847,14 @@ protected:
 	/**
 	 * Removes stale entries and fills the completed states list.
 	 * The tree can be in a different state from when a finished task was added to the pending list and the finished task may not be valid.
+	 * @param bMarkTaskProcessed when true, marks the tasks as processed. When false, only used the already marked tasks.
 	 */
+	void UpdateCompletedStateList(bool bMarkTaskProcessed);
+	UE_DEPRECATED(5.6, "Use UpdateCompletedStateList with the bMarkTaskProcessed argument")
 	void UpdateCompletedStateList();
 
-	/** Adds the state to the completed list from a completed task. */
-	void MarkStateCompleted(const UE::StateTree::FFinishedTask& FinisedTask);
+	/** Adds the state to the completed list from a finished task. */
+	void MarkStateCompleted(UE::StateTree::FFinishedTask& FinisedTask);
 
 	/** In the list of completed states, returns the run status when it's completed by a global task. */
 	EStateTreeRunStatus GetGlobalTasksCompletedStatesStatus() const;
@@ -972,8 +996,17 @@ protected:
 
 	bool bActiveExternalDataCollected = false;
 
+#if WITH_EDITORONLY_DATA
 	/** Run status set in FinishTask and that used after the task ticks. */
+	UE_DEPRECATED(5.6, "CompletedStateRunStatus is not used.")
 	TOptional<EStateTreeRunStatus> CompletedStateRunStatus;
+#endif
+
+	/**
+	 * UpdateCompletedStateList is called once per context. Either from TickTasks or TickTransitions.
+	 * User can call TickTasks and wait some frame before calling TickTransitions.
+	 */
+	bool bCompletedStatesUpdated = false;
 
 	/** State that was completed in this frame. */
 	struct FCompletedState : UE::StateTree::FFinishedTask
diff --git a/Engine/Plugins/Runtime/StateTree/Source/StateTreeModule/Public/StateTreeExecutionTypes.h b/Engine/Plugins/Runtime/StateTree/Source/StateTreeModule/Public/StateTreeExecutionTypes.h
--- a/Engine/Plugins/Runtime/StateTree/Source/StateTreeModule/Public/StateTreeExecutionTypes.h
+++ b/Engine/Plugins/Runtime/StateTree/Source/StateTreeModule/Public/StateTreeExecutionTypes.h
@@ -699,6 +699,12 @@ struct FFinishedTask
 	EStateTreeRunStatus RunStatus = EStateTreeRunStatus::Failed;
 	/** The reason of the finished task. */
 	EReasonType Reason = EReasonType::GlobalTask;
+	/**
+	 * Set to true if the task is completed before or during the TickTasks.
+	 * Used to identify tasks that were completed and had the chance to be processed by the algo.
+	 * If not processed, they won't trigger the transition in this frame.
+	 */
+	bool bTickProcessed = false;
 };
 } //namespace UE::StateTree
 
diff --git a/Engine/Plugins/Runtime/StateTree/Source/StateTreeTestSuite/Private/StateTreeWeakContextTest.cpp b/Engine/Plugins/Runtime/StateTree/Source/StateTreeTestSuite/Private/StateTreeWeakContextTest.cpp
new file mode 100644
--- /dev/null
+++ b/Engine/Plugins/Runtime/StateTree/Source/StateTreeTestSuite/Private/StateTreeWeakContextTest.cpp
@@ -0,0 +1,328 @@
+// Copyright Epic Games, Inc. All Rights Reserved.
+
+#include "StateTreeTest.h"
+#include "StateTreeTestBase.h"
+#include "StateTreeTestTypes.h"
+
+#include "StateTreeCompilerLog.h"
+#include "StateTreeEditorData.h"
+#include "StateTreeCompiler.h"
+
+#define LOCTEXT_NAMESPACE "AITestSuite_StateTreeTestWeakContext"
+
+UE_DISABLE_OPTIMIZATION_SHIP
+
+namespace UE::StateTree::Tests
+{
+struct FStateTreeTest_WeakContext_FinishTask : FStateTreeTestBase
+{
+	virtual bool InstantTest() override
+	{
+		//Tree 1 (with globals)
+		//	Root
+		//		State1 -> Next
+		//			State2 -> Success
+
+		struct FWeakContext
+		{
+			FStateTreeWeakExecutionContext ContextTree1GlobalTask1;
+			FStateTreeWeakExecutionContext ContextTree1GlobalTask2;
+			FStateTreeWeakExecutionContext ContextTree1RootTask1;
+			FStateTreeWeakExecutionContext ContextTree1RootTask2;
+			FStateTreeWeakExecutionContext ContextTree1State1Task1;
+			FStateTreeWeakExecutionContext ContextTree1State1Task2;
+			FStateTreeWeakExecutionContext ContextTree1State2Task1;
+			FStateTreeWeakExecutionContext ContextTree1State2Task2;
+			FStateTreeWeakTaskRef TaskTree1GlobalTask1;
+			FStateTreeWeakTaskRef TaskTree1GlobalTask2;
+			FStateTreeWeakTaskRef TaskTree1RootTask1;
+			FStateTreeWeakTaskRef TaskTree1RootTask2;
+			FStateTreeWeakTaskRef TaskTree1State1Task1;
+			FStateTreeWeakTaskRef TaskTree1State1Task2;
+			FStateTreeWeakTaskRef TaskTree1State2Task1;
+			FStateTreeWeakTaskRef TaskTree1State2Task2;
+
+			bool bWhen_Tree1GlobalTask1_Finish_Tree1GlobalTask2 = false;
+			bool bWhen_TaskTree1State1Task1_Finish_TaskTree1State1Task2 = false;
+			bool bWhen_TaskTree1State1Task2_Finish_TaskTree1State1Task1 = false;
+			bool bWhen_TaskTree1State1Task2_Finish_TaskTree1State1Task2 = false;
+			bool bWhen_TaskTree1State1Task2_Finish_Tree1GlobalTask2 = false;
+		};
+		FWeakContext WeakContext;
+
+		UStateTree& StateTree1 = NewStateTree();
+		{
+			UStateTreeEditorData& EditorData1 = *Cast<UStateTreeEditorData>(StateTree1.EditorData);
+			{
+				TStateTreeEditorNode<FTestTask_PrintValue>& Tree1GlobalTask1 = EditorData1.AddGlobalTask<FTestTask_PrintValue>(FName(TEXT("Tree1GlobalTask1")));
+				Tree1GlobalTask1.GetNode().CustomEnterStateFunc = [&WeakContext](FStateTreeExecutionContext& Context, const FTestTask_PrintValue* Task)
+					{
+						WeakContext.ContextTree1GlobalTask1 = Context.MakeWeakExecutionContext();
+						WeakContext.TaskTree1GlobalTask1 = Context.MakeWeakTaskRef(*Task);
+					};
+				Tree1GlobalTask1.GetNode().CustomTickFunc = [&WeakContext](FStateTreeExecutionContext& Context, const FTestTask_PrintValue* Task)
+					{
+						if (WeakContext.bWhen_Tree1GlobalTask1_Finish_Tree1GlobalTask2)
+						{
+							WeakContext.ContextTree1GlobalTask1.FinishTask(WeakContext.TaskTree1GlobalTask1, EStateTreeFinishTaskType::Failed);
+						}
+					};
+				TStateTreeEditorNode<FTestTask_PrintValue>& Tree1GlobalTask2 = EditorData1.AddGlobalTask<FTestTask_PrintValue>(FName(TEXT("Tree1GlobalTask2")));
+				Tree1GlobalTask2.GetNode().CustomEnterStateFunc = [&WeakContext](FStateTreeExecutionContext& Context, const FTestTask_PrintValue* Task)
+					{
+						WeakContext.ContextTree1GlobalTask2 = Context.MakeWeakExecutionContext();
+						WeakContext.TaskTree1GlobalTask2 = Context.MakeWeakTaskRef(*Task);
+					};
+			}
+			UStateTreeState& Root1 = EditorData1.AddSubTree("Tree1StateRoot");
+			{
+				TStateTreeEditorNode<FTestTask_PrintValue>& Tree1StateRootTask1 = Root1.AddTask<FTestTask_PrintValue>(FName(TEXT("Tree1RootTask1")));
+				Tree1StateRootTask1.GetNode().CustomEnterStateFunc = [&WeakContext](FStateTreeExecutionContext& Context, const FTestTask_PrintValue* Task)
+					{
+						WeakContext.ContextTree1RootTask1 = Context.MakeWeakExecutionContext();
+						WeakContext.TaskTree1RootTask1 = Context.MakeWeakTaskRef(*Task);
+					};
+				TStateTreeEditorNode<FTestTask_PrintValue>& Tree1StateRootTask2 = Root1.AddTask<FTestTask_PrintValue>(FName(TEXT("Tree1RootTask2")));
+				Tree1StateRootTask2.GetNode().CustomEnterStateFunc = [&WeakContext](FStateTreeExecutionContext& Context, const FTestTask_PrintValue* Task)
+					{
+						WeakContext.ContextTree1RootTask2 = Context.MakeWeakExecutionContext();
+						WeakContext.TaskTree1RootTask2 = Context.MakeWeakTaskRef(*Task);
+					};
+			}
+			UStateTreeState& State1 = Root1.AddChildState("Tree1State1");
+			{
+				TStateTreeEditorNode<FTestTask_PrintValue>& Tree1State1Task1 = State1.AddTask<FTestTask_PrintValue>(FName(TEXT("Tree1State1Task1")));
+				Tree1State1Task1.GetNode().CustomEnterStateFunc = [&WeakContext](FStateTreeExecutionContext& Context, const FTestTask_PrintValue* Task)
+					{
+						WeakContext.ContextTree1State1Task1 = Context.MakeWeakExecutionContext();
+						WeakContext.TaskTree1State1Task1 = Context.MakeWeakTaskRef(*Task);
+					};
+				Tree1State1Task1.GetNode().CustomTickFunc = [&WeakContext](FStateTreeExecutionContext& Context, const FTestTask_PrintValue* Task)
+					{
+						if (WeakContext.bWhen_TaskTree1State1Task1_Finish_TaskTree1State1Task2)
+						{
+							WeakContext.ContextTree1State1Task2.FinishTask(WeakContext.TaskTree1State1Task2, EStateTreeFinishTaskType::Failed);
+						}
+					};				
+				TStateTreeEditorNode<FTestTask_PrintValue>& Tree1State1Task2 = State1.AddTask<FTestTask_PrintValue>(FName(TEXT("Tree1State1Task2")));
+				Tree1State1Task2.GetNode().CustomEnterStateFunc = [&WeakContext](FStateTreeExecutionContext& Context, const FTestTask_PrintValue* Task)
+					{
+						WeakContext.ContextTree1State1Task2 = Context.MakeWeakExecutionContext();
+						WeakContext.TaskTree1State1Task2 = Context.MakeWeakTaskRef(*Task);
+					};
+				Tree1State1Task2.GetNode().CustomTickFunc = [&WeakContext](FStateTreeExecutionContext& Context, const FTestTask_PrintValue* Task)
+					{
+						if (WeakContext.bWhen_TaskTree1State1Task2_Finish_TaskTree1State1Task2)
+						{
+							WeakContext.ContextTree1State1Task2.FinishTask(WeakContext.TaskTree1State1Task2, EStateTreeFinishTaskType::Failed);
+						}
+						if (WeakContext.bWhen_TaskTree1State1Task2_Finish_TaskTree1State1Task1)
+						{
+							WeakContext.ContextTree1State1Task1.FinishTask(WeakContext.TaskTree1State1Task1, EStateTreeFinishTaskType::Failed);
+						}
+						if (WeakContext.bWhen_TaskTree1State1Task2_Finish_Tree1GlobalTask2)
+						{
+							WeakContext.ContextTree1GlobalTask2.FinishTask(WeakContext.TaskTree1GlobalTask2, EStateTreeFinishTaskType::Failed);
+						}
+					};
+			}
+			UStateTreeState& State2 = State1.AddChildState("Tree1State2");
+			{
+				TStateTreeEditorNode<FTestTask_PrintValue>& Tree1State2Task1 = State2.AddTask<FTestTask_PrintValue>(FName(TEXT("Tree1State2Task1")));
+				Tree1State2Task1.GetNode().CustomEnterStateFunc = [&WeakContext](FStateTreeExecutionContext& Context, const FTestTask_PrintValue* Task)
+					{
+						WeakContext.ContextTree1State1Task1 = Context.MakeWeakExecutionContext();
+						WeakContext.TaskTree1State1Task1 = Context.MakeWeakTaskRef(*Task);
+					};
+				TStateTreeEditorNode<FTestTask_PrintValue>& Tree1State2Task2 = State2.AddTask<FTestTask_PrintValue>(FName(TEXT("Tree1State2Task2")));
+				Tree1State2Task2.GetNode().CustomEnterStateFunc = [&WeakContext](FStateTreeExecutionContext& Context, const FTestTask_PrintValue* Task)
+					{
+						WeakContext.ContextTree1State1Task2 = Context.MakeWeakExecutionContext();
+						WeakContext.TaskTree1State1Task2 = Context.MakeWeakTaskRef(*Task);
+					};
+			}
+		}
+
+		// Compile tree
+		{
+			FStateTreeCompilerLog Log;
+			FStateTreeCompiler Compiler(Log);
+			const bool bResult = Compiler.Compile(StateTree1);
+			AITEST_TRUE("StateTree1 should get compiled", bResult);
+		}
+
+		// Create context
+		FStateTreeInstanceData InstanceData;
+		{
+			{
+				FTestStateTreeExecutionContext Exec(StateTree1, StateTree1, InstanceData);
+				const bool bInitSucceeded = Exec.IsValid();
+				AITEST_TRUE("StateTree should init", bInitSucceeded);
+			}
+			{
+				FTestStateTreeExecutionContext Exec(StateTree1, StateTree1, InstanceData);
+				EStateTreeRunStatus Status = Exec.Start();
+				AITEST_EQUAL("Start should complete with Running", Status, EStateTreeRunStatus::Running);
+				AITEST_TRUE(TEXT("In correct states"), Exec.ExpectInActiveStates("Tree1StateRoot", "Tree1State1", "Tree1State2"));
+
+				AITEST_TRUE(TEXT("Start should EnterState"), Exec.Expect("Tree1GlobalTask1", TEXT("EnterState0"))
+					.Then("Tree1GlobalTask2", TEXT("EnterState0"))
+					.Then("Tree1RootTask1", TEXT("EnterState0"))
+					.Then("Tree1RootTask2", TEXT("EnterState0"))
+					.Then("Tree1State1Task1", TEXT("EnterState0"))
+					.Then("Tree1State1Task2", TEXT("EnterState0"))
+					.Then("Tree1State2Task1", TEXT("EnterState0"))
+					.Then("Tree1State2Task2", TEXT("EnterState0")));
+				Exec.LogClear();
+			}
+			// Test that everything tick and there are no transitions.
+			{
+				FTestStateTreeExecutionContext Exec(StateTree1, StateTree1, InstanceData);
+				EStateTreeRunStatus Status = Exec.Tick(0.1f);
+				AITEST_EQUAL("Tick should complete with Running", Status, EStateTreeRunStatus::Running);
+				AITEST_TRUE(TEXT("In correct states"), Exec.ExpectInActiveStates("Tree1StateRoot", "Tree1State1", "Tree1State2"));
+
+				AITEST_TRUE(TEXT("Start should Tick"), Exec.Expect("Tree1GlobalTask1", TEXT("Tick0"))
+					.Then("Tree1GlobalTask2", TEXT("Tick0"))
+					.Then("Tree1RootTask1", TEXT("Tick0"))
+					.Then("Tree1RootTask2", TEXT("Tick0"))
+					.Then("Tree1State1Task1", TEXT("Tick0"))
+					.Then("Tree1State1Task2", TEXT("Tick0"))
+					.Then("Tree1State2Task1", TEXT("Tick0"))
+					.Then("Tree1State2Task2", TEXT("Tick0")));
+				Exec.LogClear();
+			}
+			// Test Weak::FinishTask outside of the tick.
+			WeakContext.ContextTree1RootTask2.FinishTask(WeakContext.TaskTree1RootTask2, EStateTreeFinishTaskType::Succeeded);
+			{
+				FTestStateTreeExecutionContext Exec(StateTree1, StateTree1, InstanceData);
+				EStateTreeRunStatus Status = Exec.Tick(0.1f);
+				AITEST_EQUAL("Tick should complete with Running", Status, EStateTreeRunStatus::Running);
+				AITEST_TRUE(TEXT("In correct states"), Exec.ExpectInActiveStates("Tree1StateRoot", "Tree1State1", "Tree1State2"));
+
+				AITEST_TRUE(TEXT("Tick should Tick"), Exec.Expect("Tree1GlobalTask1", TEXT("Tick0"))
+					.Then("Tree1GlobalTask2", TEXT("Tick0"))
+					.Then("Tree1State2Task2", TEXT("ExitState0"))
+					.Then("Tree1State2Task1", TEXT("ExitState0"))
+					.Then("Tree1State1Task2", TEXT("ExitState0"))
+					.Then("Tree1State1Task1", TEXT("ExitState0"))
+					.Then("Tree1RootTask2", TEXT("ExitState0"))
+					.Then("Tree1RootTask1", TEXT("ExitState0"))
+					.Then("Tree1RootTask1", TEXT("EnterState0"))
+					.Then("Tree1RootTask2", TEXT("EnterState0"))
+					.Then("Tree1State1Task1", TEXT("EnterState0"))
+					.Then("Tree1State1Task2", TEXT("EnterState0"))
+					.Then("Tree1State2Task1", TEXT("EnterState0"))
+					.Then("Tree1State2Task2", TEXT("EnterState0")));
+				AITEST_FALSE(TEXT("FinishTask should not Tick"), Exec.Expect("Tree1RootTask1", TEXT("Tick0")));
+				AITEST_FALSE(TEXT("FinishTask should not Tick"), Exec.Expect("Tree1RootTask2", TEXT("Tick0")));
+				AITEST_FALSE(TEXT("FinishTask should not Tick"), Exec.Expect("Tree1State1Task1", TEXT("Tick0")));
+				AITEST_FALSE(TEXT("FinishTask should not Tick"), Exec.Expect("Tree1State1Task2", TEXT("Tick0")));
+				AITEST_FALSE(TEXT("FinishTask should not Tick"), Exec.Expect("Tree1State2Task1", TEXT("Tick0")));
+				AITEST_FALSE(TEXT("FinishTask should not Tick"), Exec.Expect("Tree1State2Task2", TEXT("Tick0")));
+				Exec.LogClear();
+			}
+			// Test Weak::FinishTask of a task after the current task, during the tick tasks.
+			{
+				WeakContext.bWhen_TaskTree1State1Task1_Finish_TaskTree1State1Task2 = true;
+				FTestStateTreeExecutionContext Exec(StateTree1, StateTree1, InstanceData);
+				EStateTreeRunStatus Status = Exec.Tick(0.1f);
+				AITEST_EQUAL("Tick should complete with Running", Status, EStateTreeRunStatus::Running);
+				AITEST_TRUE(TEXT("In correct states"), Exec.ExpectInActiveStates("Tree1StateRoot", "Tree1State1", "Tree1State2"));
+
+				AITEST_TRUE(TEXT("Tick should Tick"), Exec.Expect("Tree1GlobalTask1", TEXT("Tick0"))
+					.Then("Tree1GlobalTask2", TEXT("Tick0"))
+					.Then("Tree1RootTask1", TEXT("Tick0"))
+					.Then("Tree1RootTask2", TEXT("Tick0"))
+					.Then("Tree1State1Task1", TEXT("Tick0")));
+				AITEST_FALSE(TEXT("FinishTask should not Tick"), Exec.Expect("Tree1State1Task2", TEXT("Tick0")));
+				AITEST_FALSE(TEXT("FinishTask should not Tick"), Exec.Expect("Tree1State2Task1", TEXT("Tick0")));
+				AITEST_FALSE(TEXT("FinishTask should not Tick"), Exec.Expect("Tree1State2Task2", TEXT("Tick0")));
+				AITEST_TRUE(TEXT("Tick should Tick"), Exec.Expect("Tree1State2Task2", TEXT("ExitState0"))
+					.Then("Tree1State2Task1", TEXT("ExitState0"))
+					.Then("Tree1State1Task2", TEXT("ExitState0"))
+					.Then("Tree1State1Task2", TEXT("EnterState0"))
+					.Then("Tree1State2Task1", TEXT("EnterState0"))
+					.Then("Tree1State2Task2", TEXT("EnterState0")));
+				Exec.LogClear();
+				WeakContext.bWhen_TaskTree1State1Task1_Finish_TaskTree1State1Task2 = false;
+			}
+			// Test Weak::FinishTask of a task before the current task, during the tick tasks.
+			{
+				WeakContext.bWhen_TaskTree1State1Task2_Finish_TaskTree1State1Task1 = true;
+				FTestStateTreeExecutionContext Exec(StateTree1, StateTree1, InstanceData);
+				EStateTreeRunStatus Status = Exec.Tick(0.1f);
+				AITEST_EQUAL("Tick should complete with Running", Status, EStateTreeRunStatus::Running);
+				AITEST_TRUE(TEXT("In correct states"), Exec.ExpectInActiveStates("Tree1StateRoot", "Tree1State1", "Tree1State2"));
+
+				AITEST_TRUE(TEXT("Tick should Tick"), Exec.Expect("Tree1GlobalTask1", TEXT("Tick0"))
+					.Then("Tree1GlobalTask2", TEXT("Tick0"))
+					.Then("Tree1RootTask1", TEXT("Tick0"))
+					.Then("Tree1RootTask2", TEXT("Tick0"))
+					.Then("Tree1State1Task1", TEXT("Tick0"))
+					.Then("Tree1State1Task2", TEXT("Tick0")));
+				AITEST_FALSE(TEXT("FinishTask should not Tick"), Exec.Expect("Tree1State2Task1", TEXT("Tick0")));
+				AITEST_FALSE(TEXT("FinishTask should not Tick"), Exec.Expect("Tree1State2Task2", TEXT("Tick0")));
+				AITEST_TRUE(TEXT("Tick should Tick"), Exec.Expect("Tree1State2Task2", TEXT("ExitState0"))
+					.Then("Tree1State2Task1", TEXT("ExitState0"))
+					.Then("Tree1State1Task2", TEXT("ExitState0"))
+					.Then("Tree1State1Task2", TEXT("EnterState0"))
+					.Then("Tree1State2Task1", TEXT("EnterState0"))
+					.Then("Tree1State2Task2", TEXT("EnterState0")));
+				Exec.LogClear();
+				WeakContext.bWhen_TaskTree1State1Task2_Finish_TaskTree1State1Task1 = false;
+			}
+			// Test Weak::FinishTask of a global task, during the global tick tasks.
+			{
+				WeakContext.bWhen_Tree1GlobalTask1_Finish_Tree1GlobalTask2 = true;
+				FTestStateTreeExecutionContext Exec(StateTree1, StateTree1, InstanceData);
+				EStateTreeRunStatus Status = Exec.Tick(0.1f);
+				AITEST_EQUAL("Tick should complete with Running", Status, EStateTreeRunStatus::Failed);
+				AITEST_TRUE(TEXT("Tick should Tick"), Exec.Expect("Tree1GlobalTask1", TEXT("Tick0"))
+					.Then("Tree1State2Task2", TEXT("ExitState0"))
+					.Then("Tree1State2Task1", TEXT("ExitState0"))
+					.Then("Tree1State1Task2", TEXT("ExitState0"))
+					.Then("Tree1GlobalTask2", TEXT("ExitState0"))
+					.Then("Tree1GlobalTask1", TEXT("ExitState0")));
+				Exec.LogClear();
+				WeakContext.bWhen_Tree1GlobalTask1_Finish_Tree1GlobalTask2 = false;
+			}
+			// Finished global task stop the execution. Reset the execution.
+			{
+				FTestStateTreeExecutionContext Exec(StateTree1, StateTree1, InstanceData);
+				Exec.Stop();
+				Exec.Start();
+				Exec.LogClear();
+			}
+			// Test Weak::FinishTask of global task, during the tick tasks.
+			{
+				WeakContext.bWhen_TaskTree1State1Task2_Finish_Tree1GlobalTask2 = true;
+				FTestStateTreeExecutionContext Exec(StateTree1, StateTree1, InstanceData);
+				EStateTreeRunStatus Status = Exec.Tick(0.1f);
+				AITEST_EQUAL("Tick should complete with Failed", Status, EStateTreeRunStatus::Failed);
+
+				AITEST_TRUE(TEXT("Tick should Tick"), Exec.Expect("Tree1GlobalTask1", TEXT("Tick0"))
+					.Then("Tree1State1Task2", TEXT("Tick0")));
+				AITEST_FALSE(TEXT("FinishTask should not Tick"), Exec.Expect("Tree1State2Task1", TEXT("Tick0")));
+				AITEST_TRUE(TEXT("Tick should Tick"), Exec.Expect("Tree1State2Task2", TEXT("ExitState0"))
+					.Then("Tree1GlobalTask1", TEXT("ExitState0")));
+				Exec.LogClear();
+				WeakContext.bWhen_TaskTree1State1Task2_Finish_Tree1GlobalTask2 = false;
+			}
+
+			{
+				FTestStateTreeExecutionContext Exec(StateTree1, StateTree1, InstanceData);
+				Exec.Stop();
+			}
+		}
+
+		return true;
+	}
+};
+IMPLEMENT_AI_INSTANT_TEST(FStateTreeTest_WeakContext_FinishTask, "System.StateTree.WeakContext.FinishTask");
+
+} // namespace UE::StateTree::Tests
+
+UE_ENABLE_OPTIMIZATION_SHIP
+
+#undef LOCTEXT_NAMESPACE
