Fix linux build. Issue introduced with CL 40296435.
#rnx

diff --git a/Engine/Plugins/Runtime/StateTree/Source/StateTreeModule/Private/StateTreeExecutionContext.cpp b/Engine/Plugins/Runtime/StateTree/Source/StateTreeModule/Private/StateTreeExecutionContext.cpp
--- a/Engine/Plugins/Runtime/StateTree/Source/StateTreeModule/Private/StateTreeExecutionContext.cpp
+++ b/Engine/Plugins/Runtime/StateTree/Source/StateTreeModule/Private/StateTreeExecutionContext.cpp
@@ -30,11 +30,11 @@ namespace UE::StateTree::Debug
 
 namespace UE::StateTree::ExecutionContext::Private
 {
-	bool bForceCopyBoundPropertiesOnNonTickedTask = false;
+	bool bCopyBoundPropertiesOnNonTickedTask = false;
 	FAutoConsoleVariableRef CVarLogCompiledStateTree(
-		TEXT("StateTree.ForceCopyBoundPropertiesOnNonTickedTask"),
-		bForceCopyBoundPropertiesOnNonTickedTask,
-		TEXT("When ticking the tasks, copy the bindings when the tasks are not ticked. A task can not ticked because of another task failure or because it disabled ticking.")
+		TEXT("StateTree.CopyBoundPropertiesOnNonTickedTask"),
+		bCopyBoundPropertiesOnNonTickedTask,
+		TEXT("When ticking the tasks, copy the bindings when the tasks are not ticked because it disabled ticking or completed. The bindings are not updated on failure.")
 	);
 
 	bool MarkDelegateAsBroadcasted(FStateTreeDelegateDispatcher Dispatcher, const FStateTreeExecutionFrame& CurrentFrame, FStateTreeInstanceStorage& Storage)
@@ -3031,7 +3031,7 @@ EStateTreeRunStatus FStateTreeExecutionContext::TickEvaluatorsAndGlobalTasks(con
 {
 	// When a global task is completed it completes the tree execution.
 	// A global task can complete async. See CompletedStates.
-	// When a global task fails, stop ticking the following tasks. If bForceCopyBoundPropertiesOnNonTickedTask, update their bindings.
+	// When a global task fails, stop ticking the following tasks.
 
 	CSV_SCOPED_TIMING_STAT_EXCLUSIVE(StateTree_TickEvaluators);
 	UE_STATETREE_DEBUG_SCOPED_PHASE(this, EStateTreeUpdatePhase::TickingGlobalTasks);
@@ -3042,7 +3042,7 @@ EStateTreeRunStatus FStateTreeExecutionContext::TickEvaluatorsAndGlobalTasks(con
 
 	EStateTreeRunStatus Result = EStateTreeRunStatus::Running;
 
-	const bool bForceCopyBoundPropertiesOnNonTickedTask = UE::StateTree::ExecutionContext::Private::bForceCopyBoundPropertiesOnNonTickedTask;
+	const bool bCopyBoundPropertiesOnNonTickedTask = UE::StateTree::ExecutionContext::Private::bCopyBoundPropertiesOnNonTickedTask;
 
 	FTickTaskArguments TickArgs;
 	TickArgs.DeltaTime = DeltaTime;
@@ -3084,38 +3084,36 @@ EStateTreeRunStatus FStateTreeExecutionContext::TickEvaluatorsAndGlobalTasks(con
 				}
 			}
 
-			// Used to stop ticking tasks after one fails, but if bForceCopyBoundPropertiesOnNonTickedTask, we still want to keep updating the data views so that property binding works properly.
-			bool bShouldTickTasks = true;
-
-			const FCompletedState* FrameCompletedState = CompletedStates.FindByPredicate([FrameID = CurrentFrame.FrameID](FCompletedState& Other)
-				{
-					return Other.FrameID == FrameID
-						&& Other.Reason == UE::StateTree::FFinishedTask::EReasonType::GlobalTask;
-				});
-			if (FrameCompletedState)
-			{
-				// Completed global task stops the frame execution.
-				bShouldTickTasks = false;
-				Result = UE::StateTree::ExecutionContext::Private::GetPriorityRunStatus(Result, FrameCompletedState->RunStatus);
-			}
-
 			if (bTickGlobalTasks)
 			{
-				const bool bHasEvents = EventQueue && EventQueue->HasEvents();
-				if (bForceCopyBoundPropertiesOnNonTickedTask || CurrentStateTree->ShouldTickGlobalTasks(bHasEvents))
+				const FCompletedState* FrameCompletedState = CompletedStates.FindByPredicate([FrameID = CurrentFrame.FrameID](FCompletedState& Other)
+					{
+						return Other.FrameID == FrameID
+							&& Other.Reason == UE::StateTree::FFinishedTask::EReasonType::GlobalTask;
+					});
+				if (FrameCompletedState)
+				{
+					// Completed global task stops the frame execution.
+					Result = UE::StateTree::ExecutionContext::Private::GetPriorityRunStatus(Result, FrameCompletedState->RunStatus);
+				}
+				else
 				{
-					// Update Tasks data and tick if possible (ie. if no task has yet failed and bShouldTickTasks is true)
-					TickArgs.TasksBegin = CurrentStateTree->GlobalTasksBegin;
-					TickArgs.TasksEnd = CurrentStateTree->GlobalTasksBegin + CurrentStateTree->GlobalTasksNum;
-					TickArgs.Indent = FrameIndex + 1;
-					const FTickTaskResult TickTasksResult = TickTasks(TickArgs);
+					const bool bHasEvents = EventQueue && EventQueue->HasEvents();
+					if (bCopyBoundPropertiesOnNonTickedTask || CurrentStateTree->ShouldTickGlobalTasks(bHasEvents))
+					{
+						// Update Tasks data and tick if possible (ie. if no task has yet failed and bShouldTickTasks is true)
+						TickArgs.TasksBegin = CurrentStateTree->GlobalTasksBegin;
+						TickArgs.TasksEnd = CurrentStateTree->GlobalTasksBegin + CurrentStateTree->GlobalTasksNum;
+						TickArgs.Indent = FrameIndex + 1;
+						const FTickTaskResult TickTasksResult = TickTasks(TickArgs);
 
-					TickArgs.bShouldTickTasks = TickTasksResult.bShouldTickTasks;
-					Result = UE::StateTree::ExecutionContext::Private::GetPriorityRunStatus(Result, TickTasksResult.Status);
+						TickArgs.bShouldTickTasks = TickTasksResult.bShouldTickTasks;
+						Result = UE::StateTree::ExecutionContext::Private::GetPriorityRunStatus(Result, TickTasksResult.Status);
 
-					if (!bForceCopyBoundPropertiesOnNonTickedTask && !TickTasksResult.bShouldTickTasks)
-					{
-						break;
+						if (!bCopyBoundPropertiesOnNonTickedTask && !TickTasksResult.bShouldTickTasks)
+						{
+							break;
+						}
 					}
 				}
 			}
@@ -3552,7 +3550,7 @@ EStateTreeRunStatus FStateTreeExecutionContext::TickTasks(const float DeltaTime)
 {
 	// When a task is completed it also completes the state and triggers the completion transition (because LastTickStatus is set).
 	// A task can complete async. See CompletedStates.
-	// When a task fails, stop ticking the following tasks. If bForceCopyBoundPropertiesOnNonTickedTask, update their bindings.
+	// When a task fails, stop ticking the following tasks.
 	// When no task ticks, then the leaf completes.
 
 	CSV_SCOPED_TIMING_STAT_EXCLUSIVE(StateTree_TickTasks);
@@ -3570,7 +3568,7 @@ EStateTreeRunStatus FStateTreeExecutionContext::TickTasks(const float DeltaTime)
 	EStateTreeRunStatus Result = EStateTreeRunStatus::Running;
 	int32 NumTotalEnabledTasks = 0;
 
-	const bool bForceCopyBoundPropertiesOnNonTickedTask = ExecutionContext::Private::bForceCopyBoundPropertiesOnNonTickedTask;
+	const bool bCopyBoundPropertiesOnNonTickedTask = ExecutionContext::Private::bCopyBoundPropertiesOnNonTickedTask;
 
 	FTickTaskArguments TickArgs;
 	TickArgs.DeltaTime = DeltaTime;
@@ -3609,7 +3607,7 @@ EStateTreeRunStatus FStateTreeExecutionContext::TickTasks(const float DeltaTime)
 			}
 
 			const bool bHasEvents = EventQueue && EventQueue->HasEvents();
-			if (bForceCopyBoundPropertiesOnNonTickedTask || State.ShouldTickTasks(bHasEvents))
+			if (bCopyBoundPropertiesOnNonTickedTask || State.ShouldTickTasks(bHasEvents))
 			{
 				// Update Tasks data and tick if possible (ie. if no task has yet failed and bShouldTickTasks is true)
 				TickArgs.TasksBegin = State.TasksBegin;
@@ -3620,7 +3618,7 @@ EStateTreeRunStatus FStateTreeExecutionContext::TickTasks(const float DeltaTime)
 				TickArgs.bShouldTickTasks = TickTasksResult.bShouldTickTasks;
 				Result = UE::StateTree::ExecutionContext::Private::GetPriorityRunStatus(Result, TickTasksResult.Status);
 
-				if (!bForceCopyBoundPropertiesOnNonTickedTask && !TickTasksResult.bShouldTickTasks)
+				if (!bCopyBoundPropertiesOnNonTickedTask && !TickTasksResult.bShouldTickTasks)
 				{
 					// Stop ticking
 					break;
@@ -3687,7 +3685,7 @@ FStateTreeExecutionContext::FTickTaskResult FStateTreeExecutionContext::TickTask
 	bool bShouldTickTasks = Args.bShouldTickTasks;
 
 	FStateTreeExecutionState& Exec = GetExecState();
-	const bool bCopyBoundPropertiesOnNonTickedTask = ExecutionContext::Private::bForceCopyBoundPropertiesOnNonTickedTask;
+	const bool bCopyBoundPropertiesOnNonTickedTask = ExecutionContext::Private::bCopyBoundPropertiesOnNonTickedTask;
 	const UStateTree* CurrentStateTree = Args.Frame->StateTree;
 	const FActiveFrameID CurrentFrameID = Args.Frame->FrameID;
 	check(CurrentStateTree);
@@ -6156,14 +6154,8 @@ bool FStateTreeWeakExecutionContext::FinishTask(const FStateTreeWeakTaskRef& Tas
 					const bool bIsGlobalTask = TaskPin.GetTaskIndex().AsInt32() >= GlobalTaskBeginIndex && TaskPin.GetTaskIndex().AsInt32() < GlobalTaskEndIndex;
 					const UE::StateTree::FFinishedTask::EReasonType Reason = bIsGlobalTask ? UE::StateTree::FFinishedTask::EReasonType::GlobalTask : UE::StateTree::FFinishedTask::EReasonType::StateTask;
 					const EStateTreeRunStatus RunStatus = UE::StateTree::ExecutionContext::Private::GetRunStatusFromFinishTaskType(FinishType);
-					UE::StateTree::FFinishedTask CompletedTask{
-						.FrameID = FrameID,
-						.StateID = StateID,
-						.TaskIndex = TaskPin.GetTaskIndex(),
-						.RunStatus = RunStatus,
-						.Reason = Reason
-					};
-
+					constexpr bool bTickProcessed = false;
+					UE::StateTree::FFinishedTask CompletedTask(FrameID, StateID, TaskPin.GetTaskIndex(), RunStatus, Reason, bTickProcessed);
 					Exec.FinishedTasks.Add(CompletedTask);
 					ScheduleNextTick(OwnerPinned.Get(), StateTreePinned.Get(), *StoragePinned);
 					bCompleted = true;
diff --git a/Engine/Plugins/Runtime/StateTree/Source/StateTreeModule/Private/StateTreeExecutionTypes.cpp b/Engine/Plugins/Runtime/StateTree/Source/StateTreeModule/Private/StateTreeExecutionTypes.cpp
--- a/Engine/Plugins/Runtime/StateTree/Source/StateTreeModule/Private/StateTreeExecutionTypes.cpp
+++ b/Engine/Plugins/Runtime/StateTree/Source/StateTreeModule/Private/StateTreeExecutionTypes.cpp
@@ -241,6 +241,19 @@ FRecordedStateTreeExecutionFrame::FRecordedStateTreeExecutionFrame(const FStateT
 	bIsGlobalFrame = ExecutionFrame.bIsGlobalFrame;
 }
 
+//----------------------------------------------------------------------//
+// FFinishedTask
+//----------------------------------------------------------------------//
+UE::StateTree::FFinishedTask::FFinishedTask(FActiveFrameID InFrameID, FActiveStateID InStateID, FStateTreeIndex16 InTaskIndex, EStateTreeRunStatus InRunStatus, EReasonType InReason, bool bInTickProcessed)
+	: FrameID(InFrameID)
+	, StateID(InStateID)
+	, TaskIndex(InTaskIndex)
+	, RunStatus(InRunStatus)
+	, Reason(InReason)
+	, bTickProcessed(bInTickProcessed)
+{
+}
+
 //----------------------------------------------------------------------//
 // FStateTreeScheduledTick
 //----------------------------------------------------------------------//
diff --git a/Engine/Plugins/Runtime/StateTree/Source/StateTreeModule/Public/StateTreeExecutionTypes.h b/Engine/Plugins/Runtime/StateTree/Source/StateTreeModule/Public/StateTreeExecutionTypes.h
--- a/Engine/Plugins/Runtime/StateTree/Source/StateTreeModule/Public/StateTreeExecutionTypes.h
+++ b/Engine/Plugins/Runtime/StateTree/Source/StateTreeModule/Public/StateTreeExecutionTypes.h
@@ -689,6 +689,8 @@ struct FFinishedTask
 		/** An internal transition finish the state. The FrameID and StateID are valid. */
 		InternalTransition,
 	};
+	FFinishedTask() = default;
+	STATETREEMODULE_API FFinishedTask(FActiveFrameID FrameID, FActiveStateID StateID, FStateTreeIndex16 TaskIndex, EStateTreeRunStatus RunStatus, EReasonType Reason, bool bTickProcessed);
 	/** Frame ID that identifies 1the active frame. */
 	FActiveFrameID FrameID;
 	/** State ID that contains the finished task. */
