StateTree: Tick the global tasks and evaluators in the state hierarchy. That solves bindings order issues and data that are one frame late for the global tasks/evaluators when they execute in a linked state. The old behavior can be set with cvar StateTree.TickGlobalNodesWithHierarchy.
[RN] The state ticking order changed. Use StateTree.TickGlobalNodesWithHierarchy to help migrate to the new behavior.

diff --git a/Engine/Plugins/Runtime/StateTree/Source/StateTreeModule/Private/StateTreeExecutionContext.cpp b/Engine/Plugins/Runtime/StateTree/Source/StateTreeModule/Private/StateTreeExecutionContext.cpp
--- a/Engine/Plugins/Runtime/StateTree/Source/StateTreeModule/Private/StateTreeExecutionContext.cpp
+++ b/Engine/Plugins/Runtime/StateTree/Source/StateTreeModule/Private/StateTreeExecutionContext.cpp
@@ -31,12 +31,22 @@ namespace UE::StateTree::Debug
 namespace UE::StateTree::ExecutionContext::Private
 {
 	bool bCopyBoundPropertiesOnNonTickedTask = false;
-	FAutoConsoleVariableRef CVarLogCompiledStateTree(
+	FAutoConsoleVariableRef CVarCopyBoundPropertiesOnNonTickedTask(
 		TEXT("StateTree.CopyBoundPropertiesOnNonTickedTask"),
 		bCopyBoundPropertiesOnNonTickedTask,
 		TEXT("When ticking the tasks, copy the bindings when the tasks are not ticked because it disabled ticking or completed. The bindings are not updated on failure.")
 	);
 
+	bool bTickGlobalNodesFollowingTreeHierarchy = true;
+	FAutoConsoleVariableRef CVarTickGlobalNodesFollowingTreeHierarchy(
+		TEXT("StateTree.TickGlobalNodesFollowingTreeHierarchy"),
+		bTickGlobalNodesFollowingTreeHierarchy,
+		TEXT("If true, then the global evaluators and global tasks are ticked following the asset hierarchy.\n")
+		TEXT("The order is (1)root evaluators, (2)root global tasks, (3)state tasks, (4)linked asset evaluators, (5)linked asset global tasks, (6)linked asset state tasks.\n")
+		TEXT("If false (5.5 behavior), then all the global nodes, from all linked assets, are evaluated, then the state tasks are ticked.\n")
+		TEXT("You should upgrade your asset. This option is to help migrate to the new behavior.")
+	);
+
 	bool MarkDelegateAsBroadcasted(FStateTreeDelegateDispatcher Dispatcher, const FStateTreeExecutionFrame& CurrentFrame, FStateTreeInstanceStorage& Storage)
 	{
 		const UStateTree* StateTree = CurrentFrame.StateTree;
@@ -1061,7 +1071,7 @@ EStateTreeRunStatus FStateTreeExecutionContext::TickTriggerTransitions()
 		UpdateCompletedStateList(bMarkTaskProcessed);
 	}
 
-	// If we have pending global task that completd, then stop the tree (like it would with a regular Tick).
+	// If we have pending global task that completed, then stop the tree (like it would with a regular Tick).
 	//This should not occurs and its probably a user mistake to not tick the system in the proper order.
 	const EStateTreeRunStatus GlobalTaskStatus = GetGlobalTasksCompletedStatesStatus();
 	if (GlobalTaskStatus == EStateTreeRunStatus::Running)
@@ -1080,10 +1090,19 @@ EStateTreeRunStatus FStateTreeExecutionContext::TickTriggerTransitions()
 	return TickPostlude();
 }
 
-void FStateTreeExecutionContext::TickUpdateTasksInternal(const float DeltaTime)
+void FStateTreeExecutionContext::TickUpdateTasksInternal(float DeltaTime)
 {
 	FStateTreeExecutionState& Exec = GetExecState();
+
+	// If stop is requested, do not try to tick tasks.
+	if (Exec.RequestedStop != EStateTreeRunStatus::Unset)
+	{
+		return;
+	}
 	
+	// Prevent wrong user input.
+	DeltaTime = FMath::Max(0.f, DeltaTime);
+
 	// Update the delayed transitions.
 	for (FStateTreeTransitionDelayedState& DelayedState : Exec.DelayedTransitions)
 	{
@@ -1094,29 +1113,49 @@ void FStateTreeExecutionContext::TickUpdateTasksInternal(const float DeltaTime)
 	constexpr bool bMarkTaskProcessed = true;
 	UpdateCompletedStateList(bMarkTaskProcessed);
 
-	// Tick global evaluators and tasks.
-	const bool bTickGlobalTasks = true;
-	const EStateTreeRunStatus EvalAndGlobalTaskStatus = TickEvaluatorsAndGlobalTasks(DeltaTime, bTickGlobalTasks);
-	if (EvalAndGlobalTaskStatus == EStateTreeRunStatus::Running)
-	{
-		if (Exec.LastTickStatus == EStateTreeRunStatus::Running)
+	auto LogRequestStop = [&Exec, this]()
+		{
+			if (Exec.RequestedStop != EStateTreeRunStatus::Unset)
+			{
+				UE_STATETREE_DEBUG_LOG_EVENT(this, Log, TEXT("Global tasks completed (%s), stopping the tree"), *UEnum::GetDisplayValueAsText(Exec.RequestedStop).ToString());
+				STATETREE_LOG(Log, TEXT("Global tasks completed (%s), stopping the tree"), *UEnum::GetDisplayValueAsText(Exec.RequestedStop).ToString());
+			}
+		};
+	auto TickTaskLogic = [&Exec, &LogRequestStop, this](float DeltaTime)
 		{
 			// Tick tasks on active states.
 			Exec.LastTickStatus = TickTasks(DeltaTime);
-
 			// Report state completed immediately (and no global task completes)
 			if (Exec.LastTickStatus != EStateTreeRunStatus::Running && Exec.RequestedStop == EStateTreeRunStatus::Unset)
 			{
 				StateCompleted();
 			}
-		}
+
+			LogRequestStop();
+		};
+
+	if (UE::StateTree::ExecutionContext::Private::bTickGlobalNodesFollowingTreeHierarchy)
+	{
+		TickTaskLogic(DeltaTime);
 	}
 	else
 	{
-		//@TODO check what we want to do with global task from linked asset state. Should they request a stop or close their own tree.
-		UE_STATETREE_DEBUG_LOG_EVENT(this, Log, TEXT("Global tasks completed (%s), stopping the tree"), *UEnum::GetDisplayValueAsText(EvalAndGlobalTaskStatus).ToString());
-		STATETREE_LOG(Log, TEXT("Global tasks completed (%s), stopping the tree"), *UEnum::GetDisplayValueAsText(EvalAndGlobalTaskStatus).ToString());
-		Exec.RequestedStop = EvalAndGlobalTaskStatus;
+		// Tick global evaluators and tasks.
+		const bool bTickGlobalTasks = true;
+		const EStateTreeRunStatus EvalAndGlobalTaskStatus = TickEvaluatorsAndGlobalTasks(DeltaTime, bTickGlobalTasks);
+		if (EvalAndGlobalTaskStatus == EStateTreeRunStatus::Running)
+		{
+			if (Exec.LastTickStatus == EStateTreeRunStatus::Running)
+			{
+				TickTaskLogic(DeltaTime);
+			}
+		}
+		else
+		{
+			//@TODO check what we want to do with global task from linked asset state. Should they request a stop or close their own tree.
+			Exec.RequestedStop = UE::StateTree::ExecutionContext::Private::GetPriorityRunStatus(Exec.RequestedStop, EvalAndGlobalTaskStatus);
+			LogRequestStop();
+		}
 	}
 }
 
@@ -3042,80 +3081,90 @@ EStateTreeRunStatus FStateTreeExecutionContext::TickEvaluatorsAndGlobalTasks(con
 
 	EStateTreeRunStatus Result = EStateTreeRunStatus::Running;
 
+	for (int32 FrameIndex = 0; FrameIndex < Exec.ActiveFrames.Num(); ++FrameIndex)
+	{
+		FStateTreeExecutionFrame* CurrentParentFrame = FrameIndex > 0 ? &Exec.ActiveFrames[FrameIndex - 1] : nullptr;
+		FStateTreeExecutionFrame& CurrentFrame = Exec.ActiveFrames[FrameIndex];
+		if (CurrentFrame.bIsGlobalFrame)
+		{
+			FCurrentlyProcessedFrameScope FrameScope(*this, CurrentParentFrame, CurrentFrame);
+
+			const EStateTreeRunStatus FrameResult = TickEvaluatorsAndGlobalTasksForFrame(DeltaTime, bTickGlobalTasks, FrameIndex, CurrentParentFrame, &CurrentFrame);
+			Result = UE::StateTree::ExecutionContext::Private::GetPriorityRunStatus(Result, FrameResult);
+
+			if (Result == EStateTreeRunStatus::Failed)
+			{
+				break;
+			}
+		}
+	}
+
+	return Result;
+}
+
+EStateTreeRunStatus FStateTreeExecutionContext::TickEvaluatorsAndGlobalTasksForFrame(const float DeltaTime, const bool bTickGlobalTasks, const int32 FrameIndex, const FStateTreeExecutionFrame* CurrentParentFrame, const TNotNull<const FStateTreeExecutionFrame*> CurrentFrame)
+{
+	check(CurrentFrame->bIsGlobalFrame);
+
+	EStateTreeRunStatus Result = EStateTreeRunStatus::Running;
+
 	const bool bCopyBoundPropertiesOnNonTickedTask = UE::StateTree::ExecutionContext::Private::bCopyBoundPropertiesOnNonTickedTask;
 
+	const UStateTree* CurrentStateTree = CurrentFrame->StateTree;
+
 	FTickTaskArguments TickArgs;
 	TickArgs.DeltaTime = DeltaTime;
 	TickArgs.bIsGlobalTasks = true;
 	TickArgs.bShouldTickTasks = true;
+	TickArgs.ParentFrame = CurrentParentFrame;
+	TickArgs.Frame = CurrentFrame;
 
-	for (int32 FrameIndex = 0; FrameIndex < Exec.ActiveFrames.Num(); FrameIndex++)
+	// Tick evaluators
+	for (int32 EvalIndex = CurrentStateTree->EvaluatorsBegin; EvalIndex < (CurrentStateTree->EvaluatorsBegin + CurrentStateTree->EvaluatorsNum); EvalIndex++)
 	{
-		FStateTreeExecutionFrame* CurrentParentFrame = FrameIndex > 0 ? &Exec.ActiveFrames[FrameIndex - 1] : nullptr;
-		FStateTreeExecutionFrame& CurrentFrame = Exec.ActiveFrames[FrameIndex];
-		TickArgs.ParentFrame = CurrentParentFrame;
-		TickArgs.Frame = &CurrentFrame;
+		const FStateTreeEvaluatorBase& Eval = CurrentStateTree->Nodes[EvalIndex].Get<const FStateTreeEvaluatorBase>();
+		const FStateTreeDataView EvalInstanceView = GetDataView(CurrentParentFrame, *CurrentFrame, Eval.InstanceDataHandle);
+		FNodeInstanceDataScope DataScope(*this, &Eval, EvalIndex, Eval.InstanceDataHandle, EvalInstanceView);
 
-		if (CurrentFrame.bIsGlobalFrame)
+		// Copy bound properties.
+		if (Eval.BindingsBatch.IsValid())
 		{
-			FCurrentlyProcessedFrameScope FrameScope(*this, CurrentParentFrame, CurrentFrame);
+			CopyBatchOnActiveInstances(CurrentParentFrame, *CurrentFrame, EvalInstanceView, Eval.BindingsBatch);
+		}
+		STATETREE_LOG(VeryVerbose, TEXT("  Tick: '%s'"), *Eval.Name.ToString());
+		UE_STATETREE_DEBUG_EVALUATOR_TICK(this, CurrentStateTree, EvalIndex);
+		{
+			QUICK_SCOPE_CYCLE_COUNTER(StateTree_Eval_Tick);
+			Eval.Tick(*this, DeltaTime);
 
-			const UStateTree* CurrentStateTree = CurrentFrame.StateTree;
+			UE_STATETREE_DEBUG_EVALUATOR_EVENT(this, EvalIndex, EvalInstanceView, EStateTreeTraceEventType::OnTicked);
+		}
+	}
 
-			// Tick evaluators
-			for (int32 EvalIndex = CurrentStateTree->EvaluatorsBegin; EvalIndex < (CurrentStateTree->EvaluatorsBegin + CurrentStateTree->EvaluatorsNum); EvalIndex++)
+	if (bTickGlobalTasks)
+	{
+		const FCompletedState* FrameCompletedState = CompletedStates.FindByPredicate([FrameID = CurrentFrame->FrameID](FCompletedState& Other)
 			{
-				const FStateTreeEvaluatorBase& Eval = CurrentStateTree->Nodes[EvalIndex].Get<const FStateTreeEvaluatorBase>();
-				const FStateTreeDataView EvalInstanceView = GetDataView(CurrentParentFrame, CurrentFrame, Eval.InstanceDataHandle);
-				FNodeInstanceDataScope DataScope(*this, &Eval, EvalIndex, Eval.InstanceDataHandle, EvalInstanceView);
-
-				// Copy bound properties.
-				if (Eval.BindingsBatch.IsValid())
-				{
-					CopyBatchOnActiveInstances(CurrentParentFrame, CurrentFrame, EvalInstanceView, Eval.BindingsBatch);
-				}
-				STATETREE_LOG(VeryVerbose, TEXT("  Tick: '%s'"), *Eval.Name.ToString());
-				UE_STATETREE_DEBUG_EVALUATOR_TICK(this, CurrentStateTree, EvalIndex);
-				{
-					QUICK_SCOPE_CYCLE_COUNTER(StateTree_Eval_Tick);
-					Eval.Tick(*this, DeltaTime);
-
-					UE_STATETREE_DEBUG_EVALUATOR_EVENT(this, EvalIndex, EvalInstanceView, EStateTreeTraceEventType::OnTicked);
-				}
-			}
-
-			if (bTickGlobalTasks)
+				return Other.FrameID == FrameID
+					&& Other.Reason == UE::StateTree::FFinishedTask::EReasonType::GlobalTask;
+			});
+		if (FrameCompletedState)
+		{
+			// Completed global task stops the frame execution.
+			Result = FrameCompletedState->RunStatus;
+		}
+		else
+		{
+			const bool bHasEvents = EventQueue && EventQueue->HasEvents();
+			if (bCopyBoundPropertiesOnNonTickedTask || CurrentStateTree->ShouldTickGlobalTasks(bHasEvents))
 			{
-				const FCompletedState* FrameCompletedState = CompletedStates.FindByPredicate([FrameID = CurrentFrame.FrameID](FCompletedState& Other)
-					{
-						return Other.FrameID == FrameID
-							&& Other.Reason == UE::StateTree::FFinishedTask::EReasonType::GlobalTask;
-					});
-				if (FrameCompletedState)
-				{
-					// Completed global task stops the frame execution.
-					Result = UE::StateTree::ExecutionContext::Private::GetPriorityRunStatus(Result, FrameCompletedState->RunStatus);
-				}
-				else
-				{
-					const bool bHasEvents = EventQueue && EventQueue->HasEvents();
-					if (bCopyBoundPropertiesOnNonTickedTask || CurrentStateTree->ShouldTickGlobalTasks(bHasEvents))
-					{
-						// Update Tasks data and tick if possible (ie. if no task has yet failed and bShouldTickTasks is true)
-						TickArgs.TasksBegin = CurrentStateTree->GlobalTasksBegin;
-						TickArgs.TasksEnd = CurrentStateTree->GlobalTasksBegin + CurrentStateTree->GlobalTasksNum;
-						TickArgs.Indent = FrameIndex + 1;
-						const FTickTaskResult TickTasksResult = TickTasks(TickArgs);
-
-						TickArgs.bShouldTickTasks = TickTasksResult.bShouldTickTasks;
-						Result = UE::StateTree::ExecutionContext::Private::GetPriorityRunStatus(Result, TickTasksResult.Status);
+				// Update Tasks data and tick if possible (ie. if no task has yet failed and bShouldTickTasks is true)
+				TickArgs.TasksBegin = CurrentStateTree->GlobalTasksBegin;
+				TickArgs.TasksEnd = CurrentStateTree->GlobalTasksBegin + CurrentStateTree->GlobalTasksNum;
+				TickArgs.Indent = FrameIndex + 1;
+				const FTickTaskResult TickTasksResult = TickTasks(TickArgs);
 
-						if (!bCopyBoundPropertiesOnNonTickedTask && !TickTasksResult.bShouldTickTasks)
-						{
-							break;
-						}
-					}
-				}
+				Result = TickTasksResult.Status;
 			}
 		}
 	}
@@ -3584,6 +3633,20 @@ EStateTreeRunStatus FStateTreeExecutionContext::TickTasks(const float DeltaTime)
 		const UStateTree* CurrentStateTree = TickArgs.Frame->StateTree;
 		FCurrentlyProcessedFrameScope FrameScope(*this, TickArgs.ParentFrame, *TickArgs.Frame);
 
+		if (UE::StateTree::ExecutionContext::Private::bTickGlobalNodesFollowingTreeHierarchy)
+		{
+			if (TickArgs.Frame->bIsGlobalFrame)
+			{
+				constexpr bool bTickGlobalTasks = true;
+				const EStateTreeRunStatus FrameResult = TickEvaluatorsAndGlobalTasksForFrame(DeltaTime, bTickGlobalTasks, FrameIndex, TickArgs.ParentFrame, TickArgs.Frame);
+				if (FrameResult != EStateTreeRunStatus::Running)
+				{
+					Exec.RequestedStop = ExecutionContext::Private::GetPriorityRunStatus(Exec.RequestedStop, FrameResult);
+					break;
+				}
+			}
+		}
+
 		for (int32 StateIndex = 0; StateIndex < TickArgs.Frame->ActiveStates.Num(); StateIndex++)
 		{
 			const FStateTreeStateHandle CurrentHandle = TickArgs.Frame->ActiveStates[StateIndex];
diff --git a/Engine/Plugins/Runtime/StateTree/Source/StateTreeModule/Public/StateTreeExecutionContext.h b/Engine/Plugins/Runtime/StateTree/Source/StateTreeModule/Public/StateTreeExecutionContext.h
--- a/Engine/Plugins/Runtime/StateTree/Source/StateTreeModule/Public/StateTreeExecutionContext.h
+++ b/Engine/Plugins/Runtime/StateTree/Source/StateTreeModule/Public/StateTreeExecutionContext.h
@@ -736,6 +736,11 @@ protected:
 	 */
 	EStateTreeRunStatus TickEvaluatorsAndGlobalTasks(const float DeltaTime, bool bTickGlobalTasks = true);
 
+	/**
+	 * Tick evaluators and global tasks of a specific frame.
+	 */
+	EStateTreeRunStatus TickEvaluatorsAndGlobalTasksForFrame(const float DeltaTime, const bool bTickGlobalTasks, const int32 FrameIndex, const FStateTreeExecutionFrame* CurrentParentFrame, const TNotNull<const FStateTreeExecutionFrame*> CurrentFrame);
+
 	/**
 	 * Starts evaluators and global tasks.
 	 * @return run status returned by the global tasks.
@@ -790,7 +795,7 @@ protected:
 	EStateTreeRunStatus TickPostlude();
 
 	/** Handles task ticking part of the tick. */
-	void TickUpdateTasksInternal(const float DeltaTime);
+	void TickUpdateTasksInternal(float DeltaTime);
 	
 	/** Handles transition triggering part of the tick. */
 	void TickTriggerTransitionsInternal();
diff --git a/Engine/Plugins/Runtime/StateTree/Source/StateTreeTestSuite/Private/StateTreeLinkedStateTest.cpp b/Engine/Plugins/Runtime/StateTree/Source/StateTreeTestSuite/Private/StateTreeLinkedStateTest.cpp
--- a/Engine/Plugins/Runtime/StateTree/Source/StateTreeTestSuite/Private/StateTreeLinkedStateTest.cpp
+++ b/Engine/Plugins/Runtime/StateTree/Source/StateTreeTestSuite/Private/StateTreeLinkedStateTest.cpp
@@ -16,6 +16,36 @@ UE_DISABLE_OPTIMIZATION_SHIP
 namespace UE::StateTree::Tests
 {
 
+namespace Private
+{
+	struct FScopedCVarBool
+	{
+		FScopedCVarBool(const TCHAR* VariableName)
+		{
+			CVar = IConsoleManager::Get().FindConsoleVariable(VariableName);
+			check(CVar);
+			bPreviousValue = CVar->GetBool();
+		}
+		~FScopedCVarBool()
+		{
+			CVar->SetWithCurrentPriority(bPreviousValue);
+		}
+
+		FScopedCVarBool(const FScopedCVarBool&) = delete;
+		FScopedCVarBool& operator=(const FScopedCVarBool&) = delete;
+
+		void Set(bool NewValue)
+		{
+			CVar->SetWithCurrentPriority(NewValue);
+			check(CVar->GetBool() == NewValue);
+		}
+
+	private:
+		IConsoleVariable* CVar;
+		bool bPreviousValue;
+	};
+} // namespace Private
+
 struct FStateTreeTest_FailEnterLinkedAsset : FStateTreeTestBase
 {
 	virtual bool InstantTest() override
@@ -271,7 +301,12 @@ struct FStateTreeTest_MultipleSameLinkedAsset : FStateTreeTestBase
 			AITEST_TRUE("StateTree should get compiled", bResult);
 		}
 
+		Private::FScopedCVarBool CVarTickGlobalNodesWithHierarchy = TEXT("StateTree.TickGlobalNodesFollowingTreeHierarchy");
+		for (int32 Counter = 0; Counter < 2; ++Counter)
 		{
+			const bool bTickGlobalNodesWithHierarchy = Counter == 0;
+			CVarTickGlobalNodesWithHierarchy.Set(bTickGlobalNodesWithHierarchy);
+
 			FStateTreeInstanceData InstanceData;
 			FTestStateTreeExecutionContext Exec(StateTree1, StateTree1, InstanceData);
 			{
@@ -351,45 +386,91 @@ struct FStateTreeTest_MultipleSameLinkedAsset : FStateTreeTestBase
 			{
 				EStateTreeRunStatus Status = Exec.Tick(0.001f);
 				AITEST_EQUAL("5th Tick should complete with Running", Status, EStateTreeRunStatus::Running);
-				FTestStateTreeExecutionContext::FLogOrder LogOrder = Exec.Expect("Tree1GlobalTaskA", TEXT("Tick99"));
-				AITEST_TRUE("5h Tick should tick Tree1GlobalTaskA", LogOrder);
-				LogOrder = LogOrder.Then("Tree2GlobalTaskA", TEXT("Tick111"));
-				AITEST_TRUE("5h Tick should tick Tree2GlobalTaskA", LogOrder);
-				LogOrder = LogOrder.Then("Tree1StateRootTaskA", TEXT("Tick88"));
-				AITEST_TRUE("5h Tick should tick Tree1StateRootTaskA", LogOrder);
-				LogOrder = LogOrder.Then("Tree2StateRootTaskA", TEXT("Tick111"));
-				AITEST_TRUE("5h Tick should tick Tree2StateRootTaskA", LogOrder);
-				LogOrder = LogOrder.Then("Tree2StateATaskA", TEXT("Tick111"));
-				AITEST_TRUE("5h Tick should tick Tree2StateATaskA", LogOrder);
+				if (bTickGlobalNodesWithHierarchy)
+				{
+					FTestStateTreeExecutionContext::FLogOrder LogOrder = Exec.Expect("Tree1GlobalTaskA", TEXT("Tick99"));
+					AITEST_TRUE("5h Tick should tick Tree1GlobalTaskA", LogOrder);
+					LogOrder = LogOrder.Then("Tree1StateRootTaskA", TEXT("Tick88"));
+					AITEST_TRUE("5h Tick should tick Tree1StateRootTaskA", LogOrder);
+					LogOrder = LogOrder.Then("Tree2GlobalTaskA", TEXT("Tick111"));
+					AITEST_TRUE("5h Tick should tick Tree2GlobalTaskA", LogOrder);
+					LogOrder = LogOrder.Then("Tree2StateRootTaskA", TEXT("Tick111"));
+					AITEST_TRUE("5h Tick should tick Tree2StateRootTaskA", LogOrder);
+					LogOrder = LogOrder.Then("Tree2StateATaskA", TEXT("Tick111"));
+					AITEST_TRUE("5h Tick should tick Tree2StateATaskA", LogOrder);
+				}
+				else
+				{
+					FTestStateTreeExecutionContext::FLogOrder LogOrder = Exec.Expect("Tree1GlobalTaskA", TEXT("Tick99"));
+					AITEST_TRUE("5h Tick should tick Tree1GlobalTaskA", LogOrder);
+					LogOrder = LogOrder.Then("Tree2GlobalTaskA", TEXT("Tick111"));
+					AITEST_TRUE("5h Tick should tick Tree2GlobalTaskA", LogOrder);
+					LogOrder = LogOrder.Then("Tree1StateRootTaskA", TEXT("Tick88"));
+					AITEST_TRUE("5h Tick should tick Tree1StateRootTaskA", LogOrder);
+					LogOrder = LogOrder.Then("Tree2StateRootTaskA", TEXT("Tick111"));
+					AITEST_TRUE("5h Tick should tick Tree2StateRootTaskA", LogOrder);
+					LogOrder = LogOrder.Then("Tree2StateATaskA", TEXT("Tick111"));
+					AITEST_TRUE("5h Tick should tick Tree2StateATaskA", LogOrder);
+				}
 				Exec.LogClear();
 			}
 			{
 				EStateTreeRunStatus Status = Exec.Tick(1.0f);
 				AITEST_EQUAL("6th Tick should complete with Running", Status, EStateTreeRunStatus::Running);
-				FTestStateTreeExecutionContext::FLogOrder LogOrder = Exec.Expect("Tree1GlobalTaskA", TEXT("Tick99"));
-				AITEST_TRUE("6th Tick should tick Tree1GlobalTaskA", LogOrder);
-				LogOrder = LogOrder.Then("Tree2GlobalTaskA", TEXT("Tick111"));
-				AITEST_TRUE("6th Tick should tick Tree2GlobalTaskA", LogOrder);
-				LogOrder = LogOrder.Then("Tree1StateRootTaskA", TEXT("Tick88"));
-				AITEST_TRUE("6th Tick should tick Tree1StateRootTaskA", LogOrder);
-				LogOrder = LogOrder.Then("Tree2StateRootTaskA", TEXT("Tick111"));
-				AITEST_TRUE("6th Tick should tick Tree2StateRootTaskA", LogOrder);
-				LogOrder = LogOrder.Then("Tree2StateATaskA", TEXT("Tick111"));
-				AITEST_TRUE("6th Tick should tick Tree2StateATaskA", LogOrder);
-				LogOrder = LogOrder.Then("Tree2StateATaskA", TEXT("ExitState111"));
-				AITEST_TRUE("6th Tick should exit Tree2StateATaskA", LogOrder);
-				LogOrder = LogOrder.Then("Tree2StateRootTaskA", TEXT("ExitState111"));
-				AITEST_TRUE("6th Tick should exit Tree2StateRootTaskA", LogOrder);
-				//@todo: bug: The exit state is not called on the global, it's a new frame.
-				//LogOrder = LogOrder.Then("Tree2GlobalTaskA", TEXT("ExitState111"));
-				//AITEST_TRUE("6th Tick should exit Tree2GlobalTaskA", LogOrder);
-				//@todo: bug: The enter state is not called on the global, it's a new frame.
-				//LogOrder = LogOrder.Then("Tree2GlobalTaskA", TEXT("EnterState222"));
-				//AITEST_TRUE("6th Tick should enter Tree2GlobalTaskA", LogOrder);
-				LogOrder = LogOrder.Then("Tree2StateRootTaskA", TEXT("EnterState222"));
-				AITEST_TRUE("6th Tick should enter Tree2StateRootTaskA", LogOrder);
-				LogOrder = LogOrder.Then("Tree2StateATaskA", TEXT("EnterState222"));
-				AITEST_TRUE("6th Tick should enter Tree2StateATaskA", LogOrder);
+				if (bTickGlobalNodesWithHierarchy)
+				{
+					FTestStateTreeExecutionContext::FLogOrder LogOrder = Exec.Expect("Tree1GlobalTaskA", TEXT("Tick99"));
+					AITEST_TRUE("6th Tick should tick Tree1GlobalTaskA", LogOrder);
+					LogOrder = LogOrder.Then("Tree1StateRootTaskA", TEXT("Tick88"));
+					AITEST_TRUE("6th Tick should tick Tree1StateRootTaskA", LogOrder);
+					LogOrder = LogOrder.Then("Tree2GlobalTaskA", TEXT("Tick111"));
+					AITEST_TRUE("6th Tick should tick Tree2GlobalTaskA", LogOrder);
+					LogOrder = LogOrder.Then("Tree2StateRootTaskA", TEXT("Tick111"));
+					AITEST_TRUE("6th Tick should tick Tree2StateRootTaskA", LogOrder);
+					LogOrder = LogOrder.Then("Tree2StateATaskA", TEXT("Tick111"));
+					AITEST_TRUE("6th Tick should tick Tree2StateATaskA", LogOrder);
+					LogOrder = LogOrder.Then("Tree2StateATaskA", TEXT("ExitState111"));
+					AITEST_TRUE("6th Tick should exit Tree2StateATaskA", LogOrder);
+					LogOrder = LogOrder.Then("Tree2StateRootTaskA", TEXT("ExitState111"));
+					AITEST_TRUE("6th Tick should exit Tree2StateRootTaskA", LogOrder);
+					//@todo: bug: The exit state is not called on the global, it's a new frame.
+					//LogOrder = LogOrder.Then("Tree2GlobalTaskA", TEXT("ExitState111"));
+					//AITEST_TRUE("6th Tick should exit Tree2GlobalTaskA", LogOrder);
+					//@todo: bug: The enter state is not called on the global, it's a new frame.
+					//LogOrder = LogOrder.Then("Tree2GlobalTaskA", TEXT("EnterState222"));
+					//AITEST_TRUE("6th Tick should enter Tree2GlobalTaskA", LogOrder);
+					LogOrder = LogOrder.Then("Tree2StateRootTaskA", TEXT("EnterState222"));
+					AITEST_TRUE("6th Tick should enter Tree2StateRootTaskA", LogOrder);
+					LogOrder = LogOrder.Then("Tree2StateATaskA", TEXT("EnterState222"));
+					AITEST_TRUE("6th Tick should enter Tree2StateATaskA", LogOrder);
+				}
+				else
+				{
+					FTestStateTreeExecutionContext::FLogOrder LogOrder = Exec.Expect("Tree1GlobalTaskA", TEXT("Tick99"));
+					AITEST_TRUE("6th Tick should tick Tree1GlobalTaskA", LogOrder);
+					LogOrder = LogOrder.Then("Tree2GlobalTaskA", TEXT("Tick111"));
+					AITEST_TRUE("6th Tick should tick Tree2GlobalTaskA", LogOrder);
+					LogOrder = LogOrder.Then("Tree1StateRootTaskA", TEXT("Tick88"));
+					AITEST_TRUE("6th Tick should tick Tree1StateRootTaskA", LogOrder);
+					LogOrder = LogOrder.Then("Tree2StateRootTaskA", TEXT("Tick111"));
+					AITEST_TRUE("6th Tick should tick Tree2StateRootTaskA", LogOrder);
+					LogOrder = LogOrder.Then("Tree2StateATaskA", TEXT("Tick111"));
+					AITEST_TRUE("6th Tick should tick Tree2StateATaskA", LogOrder);
+					LogOrder = LogOrder.Then("Tree2StateATaskA", TEXT("ExitState111"));
+					AITEST_TRUE("6th Tick should exit Tree2StateATaskA", LogOrder);
+					LogOrder = LogOrder.Then("Tree2StateRootTaskA", TEXT("ExitState111"));
+					AITEST_TRUE("6th Tick should exit Tree2StateRootTaskA", LogOrder);
+					//@todo: bug: The exit state is not called on the global, it's a new frame.
+					//LogOrder = LogOrder.Then("Tree2GlobalTaskA", TEXT("ExitState111"));
+					//AITEST_TRUE("6th Tick should exit Tree2GlobalTaskA", LogOrder);
+					//@todo: bug: The enter state is not called on the global, it's a new frame.
+					//LogOrder = LogOrder.Then("Tree2GlobalTaskA", TEXT("EnterState222"));
+					//AITEST_TRUE("6th Tick should enter Tree2GlobalTaskA", LogOrder);
+					LogOrder = LogOrder.Then("Tree2StateRootTaskA", TEXT("EnterState222"));
+					AITEST_TRUE("6th Tick should enter Tree2StateRootTaskA", LogOrder);
+					LogOrder = LogOrder.Then("Tree2StateATaskA", TEXT("EnterState222"));
+					AITEST_TRUE("6th Tick should enter Tree2StateATaskA", LogOrder);
+				}
 				Exec.LogClear();
 			}
 			Exec.Stop();
@@ -499,7 +580,12 @@ struct FStateTreeTest_EmptyStateWithTickTransitionLinkedAsset : FStateTreeTestBa
 			AITEST_TRUE("StateTree should get compiled", bResult);
 		}
 
+		Private::FScopedCVarBool CVarTickGlobalNodesWithHierarchy = TEXT("StateTree.TickGlobalNodesFollowingTreeHierarchy");
+		for (int32 Counter = 0; Counter < 2; ++Counter)
 		{
+			const bool bTickGlobalNodesWithHierarchy = Counter == 0;
+			CVarTickGlobalNodesWithHierarchy.Set(bTickGlobalNodesWithHierarchy);
+
 			FStateTreeInstanceData InstanceData;
 			FTestStateTreeExecutionContext Exec(StateTree1, StateTree1, InstanceData);
 			{
@@ -556,34 +642,68 @@ struct FStateTreeTest_EmptyStateWithTickTransitionLinkedAsset : FStateTreeTestBa
 			{
 				EStateTreeRunStatus Status = Exec.Tick(1.0f);
 				AITEST_EQUAL("3rd Tick should complete with Running", Status, EStateTreeRunStatus::Running);
-				FTestStateTreeExecutionContext::FLogOrder LogOrder = Exec.Expect("Tree1GlobalTaskA", TEXT("Tick11"));
-				AITEST_TRUE("3rd Tick should tick Tree1GlobalTaskA", LogOrder);
-				LogOrder = LogOrder.Then("Tree2GlobalTaskA", TEXT("Tick21"));
-				AITEST_TRUE("3rd Tick should tick Tree2GlobalTaskA", LogOrder);
-				LogOrder = LogOrder.Then("Tree1StateRootTaskA", TEXT("Tick12"));
-				AITEST_TRUE("3rd Tick should tick Tree1StateRootTaskA", LogOrder);
-				LogOrder = LogOrder.Then("Tree2StateRootTaskA", TEXT("Tick22"));
-				AITEST_TRUE("3rd Tick should tick Tree2StateRootTaskA", LogOrder);
 				AITEST_TRUE("3rd Tick should be in the correct states", Exec.ExpectInActiveStates("Tree1StateRoot", "Tree1StateLinkedTree", "Tree2StateRoot", "Tree2StateB"));
+				if (bTickGlobalNodesWithHierarchy)
+				{
+					FTestStateTreeExecutionContext::FLogOrder LogOrder = Exec.Expect("Tree1GlobalTaskA", TEXT("Tick11"));
+					AITEST_TRUE("3rd Tick should tick Tree1GlobalTaskA", LogOrder);
+					LogOrder = LogOrder.Then("Tree1StateRootTaskA", TEXT("Tick12"));
+					AITEST_TRUE("3rd Tick should tick Tree1StateRootTaskA", LogOrder);
+					LogOrder = LogOrder.Then("Tree2GlobalTaskA", TEXT("Tick21"));
+					AITEST_TRUE("3rd Tick should tick Tree2GlobalTaskA", LogOrder);
+					LogOrder = LogOrder.Then("Tree2StateRootTaskA", TEXT("Tick22"));
+					AITEST_TRUE("3rd Tick should tick Tree2StateRootTaskA", LogOrder);
+				}
+				else
+				{
+					FTestStateTreeExecutionContext::FLogOrder LogOrder = Exec.Expect("Tree1GlobalTaskA", TEXT("Tick11"));
+					AITEST_TRUE("3rd Tick should tick Tree1GlobalTaskA", LogOrder);
+					LogOrder = LogOrder.Then("Tree2GlobalTaskA", TEXT("Tick21"));
+					AITEST_TRUE("3rd Tick should tick Tree2GlobalTaskA", LogOrder);
+					LogOrder = LogOrder.Then("Tree1StateRootTaskA", TEXT("Tick12"));
+					AITEST_TRUE("3rd Tick should tick Tree1StateRootTaskA", LogOrder);
+					LogOrder = LogOrder.Then("Tree2StateRootTaskA", TEXT("Tick22"));
+					AITEST_TRUE("3rd Tick should tick Tree2StateRootTaskA", LogOrder);
+				}
 				Exec.LogClear();
 			}
 			{
 				EStateTreeRunStatus Status = Exec.Tick(1.0f);
 				AITEST_EQUAL("4th Tick should complete with Running", Status, EStateTreeRunStatus::Running);
-				FTestStateTreeExecutionContext::FLogOrder LogOrder = Exec.Expect("Tree1GlobalTaskA", TEXT("Tick11"));
-				AITEST_TRUE("4th Tick should tick Tree1GlobalTaskA", LogOrder);
-				LogOrder = LogOrder.Then("Tree2GlobalTaskA", TEXT("Tick21"));
-				AITEST_TRUE("4th Tick should tick Tree2GlobalTaskA", LogOrder);
-				LogOrder = LogOrder.Then("Tree1StateRootTaskA", TEXT("Tick12"));
-				AITEST_TRUE("4th Tick should tick Tree1StateRootTaskA", LogOrder);
-				LogOrder = LogOrder.Then("Tree2StateRootTaskA", TEXT("Tick22"));
-				AITEST_TRUE("4th Tick should tick Tree2StateRootTaskA", LogOrder);
-				LogOrder = LogOrder.Then("Tree2StateRootTaskA", TEXT("ExitState22"));
-				AITEST_TRUE("4th Tick should exit Tree2StateRootTaskA", LogOrder);
-				LogOrder = LogOrder.Then("Tree2GlobalTaskA", TEXT("ExitState21"));
-				AITEST_TRUE("4th Tick should exit Tree2GlobalTaskA", LogOrder);
-				LogOrder = LogOrder.Then("Tree1StateBTaskA", TEXT("EnterState14"));
-				AITEST_TRUE("4th Tick should enter Tree1StateBTaskA", LogOrder);
+				if (bTickGlobalNodesWithHierarchy)
+				{
+					FTestStateTreeExecutionContext::FLogOrder LogOrder = Exec.Expect("Tree1GlobalTaskA", TEXT("Tick11"));
+					AITEST_TRUE("4th Tick should tick Tree1GlobalTaskA", LogOrder);
+					LogOrder = LogOrder.Then("Tree1StateRootTaskA", TEXT("Tick12"));
+					AITEST_TRUE("4th Tick should tick Tree1StateRootTaskA", LogOrder);
+					LogOrder = LogOrder.Then("Tree2GlobalTaskA", TEXT("Tick21"));
+					AITEST_TRUE("4th Tick should tick Tree2GlobalTaskA", LogOrder);
+					LogOrder = LogOrder.Then("Tree2StateRootTaskA", TEXT("Tick22"));
+					AITEST_TRUE("4th Tick should tick Tree2StateRootTaskA", LogOrder);
+					LogOrder = LogOrder.Then("Tree2StateRootTaskA", TEXT("ExitState22"));
+					AITEST_TRUE("4th Tick should exit Tree2StateRootTaskA", LogOrder);
+					LogOrder = LogOrder.Then("Tree2GlobalTaskA", TEXT("ExitState21"));
+					AITEST_TRUE("4th Tick should exit Tree2GlobalTaskA", LogOrder);
+					LogOrder = LogOrder.Then("Tree1StateBTaskA", TEXT("EnterState14"));
+					AITEST_TRUE("4th Tick should enter Tree1StateBTaskA", LogOrder);
+				}
+				else
+				{
+					FTestStateTreeExecutionContext::FLogOrder LogOrder = Exec.Expect("Tree1GlobalTaskA", TEXT("Tick11"));
+					AITEST_TRUE("4th Tick should tick Tree1GlobalTaskA", LogOrder);
+					LogOrder = LogOrder.Then("Tree2GlobalTaskA", TEXT("Tick21"));
+					AITEST_TRUE("4th Tick should tick Tree2GlobalTaskA", LogOrder);
+					LogOrder = LogOrder.Then("Tree1StateRootTaskA", TEXT("Tick12"));
+					AITEST_TRUE("4th Tick should tick Tree1StateRootTaskA", LogOrder);
+					LogOrder = LogOrder.Then("Tree2StateRootTaskA", TEXT("Tick22"));
+					AITEST_TRUE("4th Tick should tick Tree2StateRootTaskA", LogOrder);
+					LogOrder = LogOrder.Then("Tree2StateRootTaskA", TEXT("ExitState22"));
+					AITEST_TRUE("4th Tick should exit Tree2StateRootTaskA", LogOrder);
+					LogOrder = LogOrder.Then("Tree2GlobalTaskA", TEXT("ExitState21"));
+					AITEST_TRUE("4th Tick should exit Tree2GlobalTaskA", LogOrder);
+					LogOrder = LogOrder.Then("Tree1StateBTaskA", TEXT("EnterState14"));
+					AITEST_TRUE("4th Tick should enter Tree1StateBTaskA", LogOrder);
+				}
 				AITEST_TRUE("4th Tick should be in the correct states", Exec.ExpectInActiveStates("Tree1StateRoot", "Tree1StateB"));
 				Exec.LogClear();
 			}
