Disable false positif static analysis introduced with CL 40350961.

#rnx

diff --git a/Engine/Plugins/Runtime/StateTree/Source/StateTreeModule/Private/StateTreeExecutionContext.cpp b/Engine/Plugins/Runtime/StateTree/Source/StateTreeModule/Private/StateTreeExecutionContext.cpp
--- a/Engine/Plugins/Runtime/StateTree/Source/StateTreeModule/Private/StateTreeExecutionContext.cpp
+++ b/Engine/Plugins/Runtime/StateTree/Source/StateTreeModule/Private/StateTreeExecutionContext.cpp
@@ -1115,7 +1115,7 @@ void FStateTreeExecutionContext::TickUpdateTasksInternal(float DeltaTime)
 
 	auto LogRequestStop = [&Exec, this]()
 		{
-			if (Exec.RequestedStop != EStateTreeRunStatus::Unset)
+			if (Exec.RequestedStop != EStateTreeRunStatus::Unset) // -V547
 			{
 				UE_STATETREE_DEBUG_LOG_EVENT(this, Log, TEXT("Global tasks completed (%s), stopping the tree"), *UEnum::GetDisplayValueAsText(Exec.RequestedStop).ToString());
 				STATETREE_LOG(Log, TEXT("Global tasks completed (%s), stopping the tree"), *UEnum::GetDisplayValueAsText(Exec.RequestedStop).ToString());
