[StateTree] fix compilation errors when using ForEachEvent

diff --git a/Engine/Plugins/Runtime/StateTree/Source/StateTreeModule/Public/StateTreeEvents.h b/Engine/Plugins/Runtime/StateTree/Source/StateTreeModule/Public/StateTreeEvents.h
--- a/Engine/Plugins/Runtime/StateTree/Source/StateTreeModule/Public/StateTreeEvents.h
+++ b/Engine/Plugins/Runtime/StateTree/Source/StateTreeModule/Public/StateTreeEvents.h
@@ -193,9 +193,9 @@ struct STATETREEMODULE_API FStateTreeEventQueue
 	 * @param Function a lambda which takes const FStateTreeSharedEvent& Event, and returns EStateTreeLoopEvents.
 	 */
 	template<typename TFunc>
-	void ForEachEvent(TFunc&& Function) const
+	void ForEachEvent(TFunc&& Function)
 	{
-		for (TArray<FStateTreeSharedEvent>::TConstIterator It(SharedEvents); It; ++It)
+		for (TArray<FStateTreeSharedEvent>::TIterator It(SharedEvents); It; ++It)
 		{
 			const EStateTreeLoopEvents Result = Function(*It);
 			if (Result == EStateTreeLoopEvents::Break)
diff --git a/Engine/Plugins/Runtime/StateTree/Source/StateTreeModule/Public/StateTreeExecutionContext.h b/Engine/Plugins/Runtime/StateTree/Source/StateTreeModule/Public/StateTreeExecutionContext.h
--- a/Engine/Plugins/Runtime/StateTree/Source/StateTreeModule/Public/StateTreeExecutionContext.h
+++ b/Engine/Plugins/Runtime/StateTree/Source/StateTreeModule/Public/StateTreeExecutionContext.h
@@ -390,12 +390,12 @@ public:
 		{
 			return;
 		}
-		EventQueue->template ForEachEvent(Function);
+		EventQueue->ForEachEvent(Function);
 	}
 
 	/**
 	 * Iterates over all events.
-	 * @param Function a lambda which takes const FStateTreeSharedEvent& Event, and returns EStateTreeLoopEvents.
+	 * @param Function a lambda which takes const FStateTreeEvent& Event, and returns EStateTreeLoopEvents.
 	 * Less preferable than FStateTreeSharedEvent version.
 	 */
 	template<typename TFunc>
@@ -405,7 +405,7 @@ public:
 		{
 			return;
 		}
-		EventQueue->template ForEachEvent([Function](const FStateTreeSharedEvent& Event)
+		EventQueue->ForEachEvent([Function](const FStateTreeSharedEvent& Event)
 		{
 			return Function(*Event);
 		});
