Fix StateTreeDebugStateTextTask modifying the state tree flow when disabled vs only running when being enabled
#rnx

diff --git a/Engine/Plugins/Runtime/StateTree/Source/StateTreeModule/Private/Tasks/StateTreeDebugTextTask.cpp b/Engine/Plugins/Runtime/StateTree/Source/StateTreeModule/Private/Tasks/StateTreeDebugTextTask.cpp
--- a/Engine/Plugins/Runtime/StateTree/Source/StateTreeModule/Private/Tasks/StateTreeDebugTextTask.cpp
+++ b/Engine/Plugins/Runtime/StateTree/Source/StateTreeModule/Private/Tasks/StateTreeDebugTextTask.cpp
@@ -21,7 +21,7 @@ EStateTreeRunStatus FStateTreeDebugTextTask::EnterState(FStateTreeExecutionConte
 {
 	if (!bEnabled)
 	{
-		return EStateTreeRunStatus::Succeeded;
+		return EStateTreeRunStatus::Running;
 	}
 
 	const FInstanceDataType& InstanceData = Context.GetInstanceData(*this);
