StateTree: Separates the workspace tab to initialize the tab for the editor toolkit and the editor mode.

diff --git a/Engine/Plugins/Experimental/AnimNextStateTree/Source/AnimNextStateTreeEditor/Private/AnimNextStateTreeEditorHost.cpp b/Engine/Plugins/Experimental/AnimNextStateTree/Source/AnimNextStateTreeEditor/Private/AnimNextStateTreeEditorHost.cpp
--- a/Engine/Plugins/Experimental/AnimNextStateTree/Source/AnimNextStateTreeEditor/Private/AnimNextStateTreeEditorHost.cpp
+++ b/Engine/Plugins/Experimental/AnimNextStateTree/Source/AnimNextStateTreeEditor/Private/AnimNextStateTreeEditorHost.cpp
@@ -5,12 +5,14 @@
 #include "AnimNextStateTree.h"
 #include "IAnimNextEditorModule.h"
 #include "StateTree.h"
+#include "StateTreeEditorWorkspaceTabHost.h"
 
 #include "IWorkspaceEditor.h"
 
 void FAnimNextStateTreeEditorHost::Init(const TWeakPtr<UE::Workspace::IWorkspaceEditor>& InWeakWorkspaceEditor)
 {
 	WeakWorkspaceEditor = InWeakWorkspaceEditor;
+	TabHost = MakeShared<UE::StateTreeEditor::FWorkspaceTabHost>();
 	
 	TSharedPtr<UE::Workspace::IWorkspaceEditor> SharedEditor = InWeakWorkspaceEditor.Pin();
 	check(SharedEditor.IsValid());
@@ -43,6 +45,11 @@ bool FAnimNextStateTreeEditorHost::ShouldShowCompileButton() const
 	return false;
 }
 
+bool FAnimNextStateTreeEditorHost::CanToolkitSpawnWorkspaceTab() const
+{
+	return true;
+}
+
 void FAnimNextStateTreeEditorHost::OnWorkspaceFocussedDocumentChanged(TObjectPtr<UObject> InObject) const
 {
 	TObjectPtr<UAnimNextStateTree> InStateTree = Cast<UAnimNextStateTree>(InObject);
@@ -87,3 +94,8 @@ TSharedPtr<IDetailsView> FAnimNextStateTreeEditorHost::GetDetailsView()
 	
 	return nullptr;
 }
+
+TSharedPtr<UE::StateTreeEditor::FWorkspaceTabHost> FAnimNextStateTreeEditorHost::GetTabHost() const
+{
+	return TabHost;
+}
diff --git a/Engine/Plugins/Experimental/AnimNextStateTree/Source/AnimNextStateTreeEditor/Private/AnimNextStateTreeEditorHost.h b/Engine/Plugins/Experimental/AnimNextStateTree/Source/AnimNextStateTreeEditor/Private/AnimNextStateTreeEditorHost.h
--- a/Engine/Plugins/Experimental/AnimNextStateTree/Source/AnimNextStateTreeEditor/Private/AnimNextStateTreeEditorHost.h
+++ b/Engine/Plugins/Experimental/AnimNextStateTree/Source/AnimNextStateTreeEditor/Private/AnimNextStateTreeEditorHost.h
@@ -21,15 +21,18 @@ public:
 	virtual FName GetCompilerLogName() const override;
 	virtual FName GetCompilerTabName() const override;
 	virtual bool ShouldShowCompileButton() const override;
+	virtual bool CanToolkitSpawnWorkspaceTab() const override;
 	virtual FSimpleMulticastDelegate& OnStateTreeChanged() override;
 	virtual TSharedPtr<IDetailsView> GetAssetDetailsView() override;
 	virtual TSharedPtr<IDetailsView> GetDetailsView() override;
+	virtual TSharedPtr<UE::StateTreeEditor::FWorkspaceTabHost> GetTabHost() const override;
 
 protected:
 	void OnWorkspaceFocussedDocumentChanged(TObjectPtr<UObject> InObject) const;
 
 	TWeakPtr<UE::Workspace::IWorkspaceEditor> WeakWorkspaceEditor;	
 	FSimpleMulticastDelegate OnStateTreeChangedDelegate;
+	TSharedPtr<UE::StateTreeEditor::FWorkspaceTabHost> TabHost;
 
 	/** Last state tree, cached to avoid changing debugger & losing context more than needed */
 	mutable TWeakObjectPtr<UAnimNextStateTree> WeakLastStateTree;
diff --git a/Engine/Plugins/Runtime/StateTree/Source/StateTreeEditorModule/Private/StandaloneStateTreeEditorHost.cpp b/Engine/Plugins/Runtime/StateTree/Source/StateTreeEditorModule/Private/StandaloneStateTreeEditorHost.cpp
--- a/Engine/Plugins/Runtime/StateTree/Source/StateTreeEditorModule/Private/StandaloneStateTreeEditorHost.cpp
+++ b/Engine/Plugins/Runtime/StateTree/Source/StateTreeEditorModule/Private/StandaloneStateTreeEditorHost.cpp
@@ -2,10 +2,12 @@
 
 #include "StandaloneStateTreeEditorHost.h"
 #include "StateTreeEditor.h"
+#include "StateTreeEditorWorkspaceTabHost.h"
 
 void FStandaloneStateTreeEditorHost::Init(const TWeakPtr<FStateTreeEditor>& InWeakStateTreeEditor)
 {
 	WeakStateTreeEditor = InWeakStateTreeEditor;
+	TabHost = MakeShared<UE::StateTreeEditor::FWorkspaceTabHost>();
 }
 
 UStateTree* FStandaloneStateTreeEditorHost::GetStateTree() const
@@ -32,6 +34,11 @@ bool FStandaloneStateTreeEditorHost::ShouldShowCompileButton() const
 	return true;
 }
 
+bool FStandaloneStateTreeEditorHost::CanToolkitSpawnWorkspaceTab() const
+{
+	return false;
+}
+
 FSimpleMulticastDelegate& FStandaloneStateTreeEditorHost::OnStateTreeChanged()
 {
 	return OnStateTreeChangedDelegate;
@@ -55,4 +62,9 @@ TSharedPtr<IDetailsView> FStandaloneStateTreeEditorHost::GetDetailsView()
 	}
 	
 	return nullptr;	
+}
+
+TSharedPtr<UE::StateTreeEditor::FWorkspaceTabHost> FStandaloneStateTreeEditorHost::GetTabHost() const
+{
+	return TabHost;
 }
\ No newline at end of file
diff --git a/Engine/Plugins/Runtime/StateTree/Source/StateTreeEditorModule/Private/StandaloneStateTreeEditorHost.h b/Engine/Plugins/Runtime/StateTree/Source/StateTreeEditorModule/Private/StandaloneStateTreeEditorHost.h
--- a/Engine/Plugins/Runtime/StateTree/Source/StateTreeEditorModule/Private/StandaloneStateTreeEditorHost.h
+++ b/Engine/Plugins/Runtime/StateTree/Source/StateTreeEditorModule/Private/StandaloneStateTreeEditorHost.h
@@ -11,16 +11,19 @@ class FStandaloneStateTreeEditorHost : public IStateTreeEditorHost
 public:
 	void Init(const TWeakPtr<FStateTreeEditor>& InWeakStateTreeEditor);
 	
-	// IStateTreeEditorHost overrides
+	//~ IStateTreeEditorHost overrides
 	virtual UStateTree* GetStateTree() const override;
 	virtual FSimpleMulticastDelegate& OnStateTreeChanged() override;
 	virtual TSharedPtr<IDetailsView> GetAssetDetailsView() override;
 	virtual TSharedPtr<IDetailsView> GetDetailsView() override;
+	virtual TSharedPtr<UE::StateTreeEditor::FWorkspaceTabHost> GetTabHost() const override;
 	virtual bool ShouldShowCompileButton() const override;
+	virtual bool CanToolkitSpawnWorkspaceTab() const override;
 	virtual FName GetCompilerLogName() const override;
 	virtual FName GetCompilerTabName() const override;
 
 protected:
 	TWeakPtr<FStateTreeEditor> WeakStateTreeEditor;
 	FSimpleMulticastDelegate OnStateTreeChangedDelegate;
+	TSharedPtr<UE::StateTreeEditor::FWorkspaceTabHost> TabHost;
 };
\ No newline at end of file
diff --git a/Engine/Plugins/Runtime/StateTree/Source/StateTreeEditorModule/Private/StateTreeEditor.cpp b/Engine/Plugins/Runtime/StateTree/Source/StateTreeEditorModule/Private/StateTreeEditor.cpp
--- a/Engine/Plugins/Runtime/StateTree/Source/StateTreeEditorModule/Private/StateTreeEditor.cpp
+++ b/Engine/Plugins/Runtime/StateTree/Source/StateTreeEditorModule/Private/StateTreeEditor.cpp
@@ -33,6 +33,7 @@
 #include "StateTreeEditorData.h"
 #include "StateTreeEditorModule.h"
 #include "StateTreeEditorSettings.h"
+#include "StateTreeEditorWorkspaceTabHost.h"
 #include "StateTreeObjectHash.h"
 #include "StateTreeTaskBase.h"
 #include "StateTreeViewModel.h"
@@ -108,6 +109,21 @@ void FStateTreeEditor::RegisterTabSpawners(const TSharedRef<class FTabManager>&
 		.SetDisplayName(NSLOCTEXT("StateTreeEditor", "CompilerResultsTab", "Compiler Results"))
 		.SetGroup(WorkspaceMenuCategoryRef)
 		.SetIcon(FSlateIcon(FAppStyle::GetAppStyleSetName(), "Log.TabIcon"));
+
+	if (EditorHost)
+	{
+		using namespace UE::StateTreeEditor;
+		TSharedPtr<FWorkspaceTabHost> TabHost = EditorHost->GetTabHost();
+		for (const FMinorWorkspaceTabConfig& Config : TabHost->GetTabConfigs())
+		{
+			FOnSpawnTab Delegate = TabHost->CreateSpawnDelegate(Config.ID);			
+			InTabManager->RegisterTabSpawner(Config.ID, Delegate)
+				.SetDisplayName(Config.Label)
+				.SetGroup(WorkspaceMenuCategoryRef)
+				.SetTooltipText(Config.Tooltip)
+				.SetIcon(Config.Icon);
+		}
+	}
 }
 
 
@@ -119,6 +135,16 @@ void FStateTreeEditor::UnregisterTabSpawners(const TSharedRef<class FTabManager>
 	InTabManager->UnregisterTabSpawner(AssetDetailsTabId);
 	InTabManager->UnregisterTabSpawner(StateTreeViewTabId);
 	InTabManager->UnregisterTabSpawner(CompilerResultsTabId);
+
+	if (EditorHost)
+	{
+		using namespace UE::StateTreeEditor;
+		TSharedPtr<FWorkspaceTabHost> TabHost = EditorHost->GetTabHost();
+		for (const FMinorWorkspaceTabConfig& Config : TabHost->GetTabConfigs())
+		{
+			InTabManager->UnregisterTabSpawner(Config.ID);
+		}
+	}
 }
 
 void FStateTreeEditor::InitEditor( const EToolkitMode::Type Mode, const TSharedPtr< class IToolkitHost >& InitToolkitHost, UStateTree* InStateTree)
@@ -129,6 +155,9 @@ void FStateTreeEditor::InitEditor( const EToolkitMode::Type Mode, const TSharedP
 	UStateTreeEditingSubsystem* StateTreeEditingSubsystem = GEditor->GetEditorSubsystem<UStateTreeEditingSubsystem>();
 	check(StateTreeEditingSubsystem);
 
+	EditorHost = MakeShared<FStandaloneStateTreeEditorHost>();
+	EditorHost->Init(StaticCastSharedRef<FStateTreeEditor>(AsShared()));
+
 	StateTreeViewModel = StateTreeEditingSubsystem->FindOrAddViewModel(InStateTree);
 
 	FMessageLogModule& MessageLogModule = FModuleManager::LoadModuleChecked<FMessageLogModule>("MessageLog");
@@ -143,19 +172,21 @@ void FStateTreeEditor::InitEditor( const EToolkitMode::Type Mode, const TSharedP
 	CompilerResultsListing = MessageLogModule.GetLogListing(CompilerLogListingName);
 	CompilerResults = MessageLogModule.CreateLogListingWidget(CompilerResultsListing.ToSharedRef());
 
-	TSharedRef<FTabManager::FLayout> StandaloneDefaultLayout = FTabManager::NewLayout("Standalone_StateTree_Layout_v4")
+	TSharedRef<FTabManager::FLayout> StandaloneDefaultLayout = FTabManager::NewLayout("Standalone_StateTree_Layout_v5")
 	->AddArea
 	(
-		FTabManager::NewPrimaryArea() ->SetOrientation(Orient_Vertical)
+		FTabManager::NewPrimaryArea()->SetOrientation(Orient_Vertical)
 		->Split
 		(
-			FTabManager::NewSplitter() ->SetOrientation(Orient_Horizontal)
+			FTabManager::NewSplitter()->SetOrientation(Orient_Horizontal)
 			->Split
 			(
 				FTabManager::NewStack()
 				->SetSizeCoefficient(0.2f)
 				->SetExtensionId(LayoutLeftStackId)
 				->AddTab(AssetDetailsTabId, ETabState::OpenedTab)
+				->AddTab(UE::StateTreeEditor::FWorkspaceTabHost::OutlinerTabId.Resolve(), ETabState::ClosedTab)
+				->AddTab(UE::StateTreeEditor::FWorkspaceTabHost::StatisticsTabId.Resolve(), ETabState::ClosedTab)
 				->SetForegroundTab(AssetDetailsTabId)
 			)
 			->Split
@@ -174,6 +205,8 @@ void FStateTreeEditor::InitEditor( const EToolkitMode::Type Mode, const TSharedP
 					->SetSizeCoefficient(0.25f)
 					->SetExtensionId(LayoutBottomMiddleStackId)
 					->AddTab(CompilerResultsTabId, ETabState::ClosedTab)
+					->AddTab(UE::StateTreeEditor::FWorkspaceTabHost::SearchTabId.Resolve(), ETabState::ClosedTab)
+					->AddTab(UE::StateTreeEditor::FWorkspaceTabHost::DebuggerTabId.Resolve(), ETabState::ClosedTab)
 				)
 			)
 			->Split
@@ -213,17 +246,14 @@ void FStateTreeEditor::PostInitAssetEditor()
 	ModeUILayer->SetModeMenuCategory(WorkspaceMenuCategory);
 	ModeUILayer->SetSecondaryModeToolbarName(GetToolMenuToolbarName());
 	ToolkitCommands->Append(ModeUILayer->GetModeCommands());
-	
+
 	if (UContextObjectStore* ContextStore = EditorModeManager->GetInteractiveToolsContext()->ContextObjectStore)
 	{
 		UStateTreeEditorContext* StateTreeEditorContext = ContextStore->FindContext<UStateTreeEditorContext>();
 		if (!StateTreeEditorContext)
 		{
-			TSharedPtr<FStandaloneStateTreeEditorHost> Host = MakeShared<FStandaloneStateTreeEditorHost>();
-			Host->Init(StaticCastSharedRef<FStateTreeEditor>(AsShared()));
-
 			StateTreeEditorContext = NewObject<UStateTreeEditorContext>();
-			StateTreeEditorContext->EditorHostInterface = Host;
+			StateTreeEditorContext->EditorHostInterface = EditorHost;
 			ContextStore->AddContextObject(StateTreeEditorContext);
 		}
 	}
diff --git a/Engine/Plugins/Runtime/StateTree/Source/StateTreeEditorModule/Private/StateTreeEditor.h b/Engine/Plugins/Runtime/StateTree/Source/StateTreeEditorModule/Private/StateTreeEditor.h
--- a/Engine/Plugins/Runtime/StateTree/Source/StateTreeEditorModule/Private/StateTreeEditor.h
+++ b/Engine/Plugins/Runtime/StateTree/Source/StateTreeEditorModule/Private/StateTreeEditor.h
@@ -14,6 +14,7 @@ class IMessageLogListing;
 class IDetailsView;
 class UStateTree;
 class FStateTreeViewModel;
+class FStandaloneStateTreeEditorHost;
 
 namespace UE::StateTree::Editor
 {
@@ -99,6 +100,7 @@ private:
 	TSharedPtr<FStateTreeEditorModeUILayer> ModeUILayer;
 	TSharedPtr<IToolkit> HostedToolkit;
 	TSharedPtr<FWorkspaceItem> WorkspaceMenuCategory;
+	TSharedPtr<FStandaloneStateTreeEditorHost> EditorHost;
 
 	static const FName StateTreeViewTabId;
 	static const FName SelectionDetailsTabId;
diff --git a/Engine/Plugins/Runtime/StateTree/Source/StateTreeEditorModule/Private/StateTreeEditorLayoutHost.cpp b/Engine/Plugins/Runtime/StateTree/Source/StateTreeEditorModule/Private/StateTreeEditorLayoutHost.cpp
new file mode 100644
--- /dev/null
+++ b/Engine/Plugins/Runtime/StateTree/Source/StateTreeEditorModule/Private/StateTreeEditorLayoutHost.cpp
@@ -0,0 +1,102 @@
+// Copyright Epic Games, Inc. All Rights Reserved.
+
+#include "StateTreeEditorWorkspaceTabHost.h"
+#include "Toolkits/AssetEditorModeUILayer.h"
+
+#include "Widgets/Docking/SDockTab.h"
+
+#define LOCTEXT_NAMESPACE "WorkspaceTabHost"
+
+class SDockTab;
+
+namespace UE::StateTreeEditor
+{
+const FLazyName FWorkspaceTabHost::OutlinerTabId = "StateTreeEditor_StateTreeOutliner";
+const FLazyName FWorkspaceTabHost::SearchTabId = "StateTreeEditor_StateTreeSearch";
+const FLazyName FWorkspaceTabHost::StatisticsTabId = "StateTreeEditor_StateTreeStatistics";
+const FLazyName FWorkspaceTabHost::DebuggerTabId = "StateTreeEditor_Debugger";
+
+TConstArrayView<FMinorWorkspaceTabConfig> FWorkspaceTabHost::GetTabConfigs() const
+{
+	auto BuildTabConfigs = []()
+		{
+			TArray<FMinorWorkspaceTabConfig> TabConfigs;
+			TabConfigs.Emplace(OutlinerTabId.Resolve(), LOCTEXT("StateTreeOutlinerTab", "Outliner"), FText(), FSlateIcon(FAppStyle::GetAppStyleSetName(), "LevelEditor.Tabs.Outliner"), UAssetEditorUISubsystem::TopLeftTabID);
+			TabConfigs.Emplace(SearchTabId.Resolve(), LOCTEXT("StateTreeSearchTab", "Find"), FText(), FSlateIcon(FAppStyle::GetAppStyleSetName(), "BlueprintEditor.FindInBlueprints.MenuIcon"), UAssetEditorUISubsystem::BottomLeftTabID);
+			TabConfigs.Emplace(StatisticsTabId.Resolve(), LOCTEXT("StatisticsTitle", "Statistics"), FText(), FSlateIcon(FAppStyle::GetAppStyleSetName(), "AssetEditor.ToggleStats"), UAssetEditorUISubsystem::BottomRightTabID);
+#if WITH_STATETREE_TRACE_DEBUGGER
+			TabConfigs.Emplace(DebuggerTabId.Resolve(), LOCTEXT("DebuggerTab", "Debugger"), FText(), FSlateIcon(FAppStyle::GetAppStyleSetName(), "Debug"), UAssetEditorUISubsystem::TopRightTabID);
+#endif
+			return TabConfigs;
+		};
+	static TArray<FMinorWorkspaceTabConfig> TabConfigs = BuildTabConfigs();
+	return TabConfigs;
+}
+
+FOnSpawnTab FWorkspaceTabHost::CreateSpawnDelegate(FName TabID)
+{
+	return ::FOnSpawnTab::CreateSP(this, &FWorkspaceTabHost::HandleSpawnDelegate, TabID);
+}
+
+TSharedRef<SDockTab> FWorkspaceTabHost::HandleSpawnDelegate(const FSpawnTabArgs& Args, FName TabID)
+{
+	return Spawn(TabID);
+}
+
+TSharedRef<SDockTab> FWorkspaceTabHost::Spawn(FName TabID)
+{
+	const FMinorWorkspaceTabConfig* Found = GetTabConfigs().FindByPredicate([TabID](const FMinorWorkspaceTabConfig& Config)
+		{
+			return Config.ID == TabID;
+		});
+
+	if (Found)
+	{
+		TSharedRef<SDockTab> Result = SNew(SDockTab)
+			.Label(Found->Label);
+		Result->SetOnTabClosed(SDockTab::FOnTabClosedCallback::CreateSP(this, &FWorkspaceTabHost::HandleTabClosed));
+
+		FSpawnedWorkspaceTab* SpawnedTab = SpawnedTabs.FindByPredicate([TabID](const FSpawnedWorkspaceTab& SpawnedTab)
+			{
+				return SpawnedTab.TabID == TabID;
+			});
+		if (SpawnedTab)
+		{
+			ensureMsgf(false, TEXT("The tab is already spawned."));
+			SpawnedTab->DockTab = Result;
+		}
+		else
+		{
+			SpawnedTabs.Emplace(TabID, TWeakPtr<SDockTab>(Result));
+		}
+
+		OnTabSpawned.Broadcast(FSpawnedWorkspaceTab{ TabID, Result });
+		return Result;
+	}
+	else
+	{
+		return SNew(SDockTab)
+			.Label(LOCTEXT("StateTreeErrorNotFoundDockTab", "<NOT FOUND>"));
+	}
+}
+
+void FWorkspaceTabHost::HandleTabClosed(TSharedRef<SDockTab> Tab)
+{
+	// Broadcast before removing
+	for (const FSpawnedWorkspaceTab& SpawnedTab : SpawnedTabs)
+	{
+		if (SpawnedTab.DockTab == Tab)
+		{
+			OnTabClosed.Broadcast(SpawnedTab);
+		}
+	}
+
+	SpawnedTabs.RemoveAllSwap([Tab](const FSpawnedWorkspaceTab& SpawnedTab)
+		{
+			return SpawnedTab.DockTab == Tab;
+		});
+}
+
+} // namespace UE::StateTreeEditor
+
+#undef LOCTEXT_NAMESPACE
diff --git a/Engine/Plugins/Runtime/StateTree/Source/StateTreeEditorModule/Private/StateTreeEditorModeToolkit.cpp b/Engine/Plugins/Runtime/StateTree/Source/StateTreeEditorModule/Private/StateTreeEditorModeToolkit.cpp
--- a/Engine/Plugins/Runtime/StateTree/Source/StateTreeEditorModule/Private/StateTreeEditorModeToolkit.cpp
+++ b/Engine/Plugins/Runtime/StateTree/Source/StateTreeEditorModule/Private/StateTreeEditorModeToolkit.cpp
@@ -6,6 +6,7 @@
 #include "ContextObjectStore.h"
 #include "InteractiveToolManager.h"
 #include "StateTreeEditorCommands.h"
+#include "StateTreeEditorWorkspaceTabHost.h"
 #include "Blueprint/StateTreeConditionBlueprintBase.h"
 #include "Blueprint/StateTreeConsiderationBlueprintBase.h"
 #include "Blueprint/StateTreeNodeBlueprintBase.h"
@@ -33,130 +34,59 @@
 
 #define LOCTEXT_NAMESPACE "StateTreeModeToolkit"
 
-const FLazyName FStateTreeEditorModeToolkit::StateTreeOutlinerTabId = "StateTreeEditor_StateTreeOutliner";
-const FLazyName FStateTreeEditorModeToolkit::StateTreeSearchTabId = "StateTreeEditor_StateTreeSearch";
-const FLazyName FStateTreeEditorModeToolkit::StateTreeStatisticsTabId = "StateTreeEditor_StateTreeStatistics";
-#if WITH_STATETREE_TRACE_DEBUGGER
-const FLazyName FStateTreeEditorModeToolkit::DebuggerTabId = "StateTreeEditor_Debugger";
-#endif // WITH_STATETREE_TRACE_DEBUGGER
-
-FStateTreeEditorModeToolkit::FStateTreeEditorModeToolkit(UStateTreeEditorMode* InEditorMode) : WeakEditorMode(InEditorMode)
+FStateTreeEditorModeToolkit::FStateTreeEditorModeToolkit(UStateTreeEditorMode* InEditorMode)
+	: WeakEditorMode(InEditorMode)
 {
 }
 
 void FStateTreeEditorModeToolkit::RequestModeUITabs()
 {
-	FModeToolkit::RequestModeUITabs();
-	
-	if (ModeUILayer.IsValid())
+	using namespace UE::StateTreeEditor;
+	TSharedPtr<FWorkspaceTabHost> TabHost = EditorHost->GetTabHost();
+	check(TabHost);
+
+	if (EditorHost->CanToolkitSpawnWorkspaceTab())
 	{
-		TSharedPtr<FAssetEditorModeUILayer> ModeUILayerPtr = ModeUILayer.Pin();
-		TSharedPtr<FWorkspaceItem> MenuGroup = ModeUILayerPtr->GetModeMenuCategory();		
-		if(!MenuGroup)
-		{ 
-			return;
-		}
+		FModeToolkit::RequestModeUITabs();
+	}
 
+	if (ModeUILayer.IsValid())
+	{
+		if (EditorHost && EditorHost->CanToolkitSpawnWorkspaceTab())
 		{
-			FMinorTabConfig OutlinerTabInfo;
-			OutlinerTabInfo.OnSpawnTab = FOnSpawnTab::CreateLambda([this](const FSpawnTabArgs& Args)-> TSharedRef<SDockTab>
+			TSharedPtr<FAssetEditorModeUILayer> ModeUILayerPtr = ModeUILayer.Pin();
+			TSharedPtr<FWorkspaceItem> MenuGroup = ModeUILayerPtr->GetModeMenuCategory();
+			if(!MenuGroup)
+			{ 
+				return;
+			}
+
+			for (const FMinorWorkspaceTabConfig& Config : TabHost->GetTabConfigs())
 			{
-				ensure(Args.GetTabId() == UAssetEditorUISubsystem::TopLeftTabID);
-				TSharedRef<SDockTab> Tab = SNew(SDockTab)
-				.Label(LOCTEXT("StateTreeOutlinerTab", "Outliner"))
-				.TabColorScale(GetTabColorScale());
-
-				WeakOutlinerTab = Tab;
-				UpdateStateTreeOutliner();
-
-			return Tab;
-		});
-		OutlinerTabInfo.TabLabel = LOCTEXT("StateTreeOutlinerTab", "Outliner");
-		
-		OutlinerTabInfo.TabId = StateTreeOutlinerTabId;
-		OutlinerTabInfo.TabIcon = FSlateIcon(FAppStyle::GetAppStyleSetName(), "LevelEditor.Tabs.Outliner");
-		OutlinerTabInfo.WorkspaceGroup = MenuGroup;
-		ModeUILayerPtr->SetModePanelInfo(UAssetEditorUISubsystem::TopLeftTabID, OutlinerTabInfo);
-		}
-		{
-			FMinorTabConfig SearchTabInfo;
-			SearchTabInfo.OnSpawnTab = FOnSpawnTab::CreateLambda([this](const FSpawnTabArgs& Args)-> TSharedRef<SDockTab>
-				{
-					ensure(Args.GetTabId() == UAssetEditorUISubsystem::BottomLeftTabID);
-					return SNew(SDockTab)
-						.Label(LOCTEXT("StateTreeSearchDockTab", "Find Results"))
-						[
-							SNew(UE::StateTreeEditor::SFindInAsset, EditorHost)
-							.bShowSearchBar(true)
-						];
-				});
-			SearchTabInfo.TabLabel = LOCTEXT("StateTreeSearchTab", "Find");
-
-			SearchTabInfo.TabId = StateTreeSearchTabId;
-			SearchTabInfo.TabIcon = FSlateIcon(FAppStyle::GetAppStyleSetName(), "BlueprintEditor.FindInBlueprints.MenuIcon");
-			SearchTabInfo.WorkspaceGroup = MenuGroup;
-			ModeUILayerPtr->SetModePanelInfo(UAssetEditorUISubsystem::BottomLeftTabID, SearchTabInfo);
-		}
-		{
-			FMinorTabConfig StatisticsTabInfo;
-			StatisticsTabInfo.OnSpawnTab = FOnSpawnTab::CreateLambda([this](const FSpawnTabArgs& Args)-> TSharedRef<SDockTab>
-				{
-					ensure(Args.GetTabId() == UAssetEditorUISubsystem::BottomRightTabID);
-					return SNew(SDockTab)
-						.Label(LOCTEXT("StatisticsTitle", "Statistics"))
-						[
-							SNew(SMultiLineEditableTextBox)
-								.Padding(10.0f)
-								.Style(FAppStyle::Get(), "Log.TextBox")
-								.Font(FCoreStyle::GetDefaultFontStyle("Mono", 9))
-								.ForegroundColor(FLinearColor::Gray)
-								.IsReadOnly(true)
-								.Text(this, &FStateTreeEditorModeToolkit::GetStatisticsText)
-						];
-				});
-			StatisticsTabInfo.TabLabel = LOCTEXT("StatisticsTitle", "Statistics");
-			StatisticsTabInfo.TabIcon = FSlateIcon(FAppStyle::GetAppStyleSetName(), "AssetEditor.ToggleStats");
-			StatisticsTabInfo.TabId = StateTreeStatisticsTabId;
-			StatisticsTabInfo.WorkspaceGroup = MenuGroup;
-
-			ModeUILayerPtr->SetModePanelInfo(UAssetEditorUISubsystem::BottomRightTabID, StatisticsTabInfo);
+				FMinorTabConfig TabInfo;
+				TabInfo.TabId = Config.ID;
+				TabInfo.TabLabel = Config.Label;
+				TabInfo.TabTooltip = Config.Tooltip;
+				TabInfo.TabIcon = Config.Icon;
+				TabInfo.WorkspaceGroup = MenuGroup;
+				TabInfo.OnSpawnTab = TabHost->CreateSpawnDelegate(Config.ID);
+				ModeUILayerPtr->SetModePanelInfo(Config.UISystemID, TabInfo);
+			}
 		}
 
-
-#if WITH_STATETREE_TRACE_DEBUGGER
+		for (const FSpawnedWorkspaceTab& SpawnedTab : TabHost->GetSpawnedTabs())
 		{
-			FMinorTabConfig DebuggerTabInfo;
-			DebuggerTabInfo.OnSpawnTab = FOnSpawnTab::CreateLambda([this](const FSpawnTabArgs& Args)-> TSharedRef<SDockTab>
-			{
-				ensure(Args.GetTabId() == UAssetEditorUISubsystem::TopRightTabID);
-				TSharedRef<SDockTab> Tab = SNew(SDockTab)
-				.Label(LOCTEXT("DebuggerTitle", "Debugger"))
-				.TabColorScale(GetTabColorScale())
-				.OnTabClosed(SDockTab::FOnTabClosedCallback::CreateLambda([](TSharedRef<SDockTab> Tab)
-					{
-						// Destroy the SStateTreeDebuggerView
-						Tab->SetContent(SNullWidget::NullWidget);
-					}));
-
-				WeakDebuggerTab = Tab;
-				UpdateDebuggerView();
-
-				return Tab;
-			});
-			DebuggerTabInfo.TabLabel = LOCTEXT("DebuggerTab", "Debugger");		
-			DebuggerTabInfo.TabId = DebuggerTabId;
-			DebuggerTabInfo.TabIcon = FSlateIcon(FAppStyle::GetAppStyleSetName(), "Debug");
-			DebuggerTabInfo.WorkspaceGroup = MenuGroup;
-			ModeUILayerPtr->SetModePanelInfo(UAssetEditorUISubsystem::TopRightTabID, DebuggerTabInfo);
+			HandleTabSpawned(SpawnedTab);
 		}
-#endif // WITH_STATETREE_TRACE_DEBUGGER
+		TabHost->OnTabSpawned.AddSP(this, &FStateTreeEditorModeToolkit::HandleTabSpawned);
+		TabHost->OnTabClosed.AddSP(this, &FStateTreeEditorModeToolkit::HandleTabClosed);
 	}
 }
 
 void FStateTreeEditorModeToolkit::Init(const TSharedPtr<IToolkitHost>& InitToolkitHost, TWeakObjectPtr<UEdMode> InOwningMode)
 {
 	FModeToolkit::Init(InitToolkitHost, InOwningMode);
-	
+
 	if (UContextObjectStore* ContextObjectStore = InOwningMode->GetToolManager()->GetContextObjectStore())
 	{
 		if (UStateTreeEditorContext* Context = ContextObjectStore->FindContext<UStateTreeEditorContext>())
@@ -164,12 +94,11 @@ void FStateTreeEditorModeToolkit::Init(const TSharedPtr<IToolkitHost>& InitToolk
 			EditorHost = Context->EditorHostInterface;
 		}
 	}
-	
 }
 
 void FStateTreeEditorModeToolkit::InvokeUI()
 {
-	if (ModeUILayer.IsValid())
+	if (ModeUILayer.IsValid() && EditorHost && EditorHost->CanToolkitSpawnWorkspaceTab())
 	{
 		TSharedPtr<FAssetEditorModeUILayer> ModeUILayerPtr = ModeUILayer.Pin();
 		ModeUILayerPtr->GetTabManager()->TryInvokeTab(UAssetEditorUISubsystem::TopLeftTabID);
@@ -177,6 +106,61 @@ void FStateTreeEditorModeToolkit::InvokeUI()
 	}
 }
 
+void FStateTreeEditorModeToolkit::HandleTabSpawned(UE::StateTreeEditor::FSpawnedWorkspaceTab SpawnedTab)
+{
+	if (SpawnedTab.TabID == UE::StateTreeEditor::FWorkspaceTabHost::OutlinerTabId)
+	{
+		WeakOutlinerTab = SpawnedTab.DockTab;
+		UpdateStateTreeOutliner();
+	}
+	else if (SpawnedTab.TabID == UE::StateTreeEditor::FWorkspaceTabHost::SearchTabId)
+	{
+		if (TSharedPtr<SDockTab> DockTab = SpawnedTab.DockTab.Pin())
+		{
+			DockTab->SetContent(
+				SNew(UE::StateTreeEditor::SFindInAsset, EditorHost)
+				.bShowSearchBar(true)
+			);
+		}
+	}
+	else if (SpawnedTab.TabID == UE::StateTreeEditor::FWorkspaceTabHost::StatisticsTabId)
+	{
+		if (TSharedPtr<SDockTab> DockTab = SpawnedTab.DockTab.Pin())
+		{
+			DockTab->SetContent(
+				SNew(SMultiLineEditableTextBox)
+				.Padding(10.0f)
+				.Style(FAppStyle::Get(), "Log.TextBox")
+				.Font(FCoreStyle::GetDefaultFontStyle("Mono", 9))
+				.ForegroundColor(FLinearColor::Gray)
+				.IsReadOnly(true)
+				.Text(this, &FStateTreeEditorModeToolkit::GetStatisticsText)
+			);
+		}
+	}
+	else if (SpawnedTab.TabID == UE::StateTreeEditor::FWorkspaceTabHost::DebuggerTabId)
+	{
+#if WITH_STATETREE_TRACE_DEBUGGER
+		WeakDebuggerTab = SpawnedTab.DockTab;
+		UpdateDebuggerView();
+#endif // WITH_STATETREE_TRACE_DEBUGGER
+	}
+}
+
+void FStateTreeEditorModeToolkit::HandleTabClosed(UE::StateTreeEditor::FSpawnedWorkspaceTab SpawnedTab)
+{
+	if (SpawnedTab.TabID == UE::StateTreeEditor::FWorkspaceTabHost::DebuggerTabId)
+	{
+#if WITH_STATETREE_TRACE_DEBUGGER
+		if (TSharedPtr<SDockTab> DockTab = WeakDebuggerTab.Pin())
+		{
+			// Destroy the SStateTreeDebuggerView
+			DockTab->SetContent(SNullWidget::NullWidget);
+		}
+#endif // WITH_STATETREE_TRACE_DEBUGGER
+	}
+}
+
 FName FStateTreeEditorModeToolkit::GetToolkitFName() const
 {
 	return FName("StateTreeMode");
@@ -615,14 +599,14 @@ void FStateTreeEditorModeToolkit::UpdateDebuggerView()
 		// We don't want it to clear newly created bindings from our new view.
 		SharedDebuggerTab->SetContent(DebuggerView);
 	}
-	
+
 	if (UStateTree* StateTree = WeakEditorMode->GetStateTree())
-    {
-    	if (UStateTreeEditingSubsystem* Subsystem = GEditor->GetEditorSubsystem<UStateTreeEditingSubsystem>())
-    	{
-    		DebuggerView = SNew(SStateTreeDebuggerView, *StateTree, Subsystem->FindOrAddViewModel(StateTree), GetToolkitCommands());
-    	}		
-    }
+	{
+		if (UStateTreeEditingSubsystem* Subsystem = GEditor->GetEditorSubsystem<UStateTreeEditingSubsystem>())
+		{
+			DebuggerView = SNew(SStateTreeDebuggerView, *StateTree, Subsystem->FindOrAddViewModel(StateTree), GetToolkitCommands());
+		}
+	}
 
 	if (SharedDebuggerTab)
 	{
diff --git a/Engine/Plugins/Runtime/StateTree/Source/StateTreeEditorModule/Public/IStateTreeEditorHost.h b/Engine/Plugins/Runtime/StateTree/Source/StateTreeEditorModule/Public/IStateTreeEditorHost.h
--- a/Engine/Plugins/Runtime/StateTree/Source/StateTreeEditorModule/Public/IStateTreeEditorHost.h
+++ b/Engine/Plugins/Runtime/StateTree/Source/StateTreeEditorModule/Public/IStateTreeEditorHost.h
@@ -7,6 +7,11 @@
 
 #include "IStateTreeEditorHost.generated.h"
 
+namespace UE::StateTreeEditor
+{
+	class FWorkspaceTabHost;
+}
+
 class UStateTree;
 class IMessageLogListing;
 class IDetailsView;
@@ -15,14 +20,20 @@ class IDetailsView;
 class IStateTreeEditorHost : public TSharedFromThis<IStateTreeEditorHost>
 {
 public:
+	IStateTreeEditorHost() = default;
 	virtual ~IStateTreeEditorHost() = default;
+
 	virtual FName GetCompilerLogName() const = 0;
 	virtual FName GetCompilerTabName() const = 0;
 	virtual bool ShouldShowCompileButton() const = 0;
+	virtual bool CanToolkitSpawnWorkspaceTab() const = 0;
+
 	virtual UStateTree* GetStateTree() const = 0;
 	virtual FSimpleMulticastDelegate& OnStateTreeChanged() = 0;
+
 	virtual TSharedPtr<IDetailsView> GetAssetDetailsView() = 0;
 	virtual TSharedPtr<IDetailsView> GetDetailsView() = 0;
+	virtual TSharedPtr<UE::StateTreeEditor::FWorkspaceTabHost> GetTabHost() const = 0;
 };
 
 UCLASS()
diff --git a/Engine/Plugins/Runtime/StateTree/Source/StateTreeEditorModule/Public/StateTreeEditorModeToolkit.h b/Engine/Plugins/Runtime/StateTree/Source/StateTreeEditorModule/Public/StateTreeEditorModeToolkit.h
--- a/Engine/Plugins/Runtime/StateTree/Source/StateTreeEditorModule/Public/StateTreeEditorModeToolkit.h
+++ b/Engine/Plugins/Runtime/StateTree/Source/StateTreeEditorModule/Public/StateTreeEditorModeToolkit.h
@@ -7,6 +7,10 @@
 #include "Toolkits/BaseToolkit.h"
 
 class UStateTreeEditorMode;
+namespace UE::StateTreeEditor
+{
+	struct FSpawnedWorkspaceTab;
+}
 
 class STATETREEEDITORMODULE_API FStateTreeEditorModeToolkit : public FModeToolkit
 {
@@ -25,7 +29,8 @@ public:
 	virtual void ExtendSecondaryModeToolbar(UToolMenu* InModeToolbarMenu) override;
 		
 	void OnStateTreeChanged();
-protected:	
+
+protected:
 	FSlateIcon GetCompileStatusImage() const;
 
 	static FSlateIcon GetNewTaskButtonImage();
@@ -42,6 +47,10 @@ protected:
 	FText GetStatisticsText() const;
 
 	void UpdateStateTreeOutliner();
+
+	void HandleTabSpawned(UE::StateTreeEditor::FSpawnedWorkspaceTab SpawnedTab);
+	void HandleTabClosed(UE::StateTreeEditor::FSpawnedWorkspaceTab SpawnedTab);
+
 protected:
 	TWeakObjectPtr<UStateTreeEditorMode> WeakEditorMode;
 	TSharedPtr<IStateTreeEditorHost> EditorHost;
@@ -53,13 +62,8 @@ protected:
 #if WITH_STATETREE_TRACE_DEBUGGER
 	void UpdateDebuggerView();
 	TWeakPtr<SDockTab> WeakDebuggerTab = nullptr;
-public:
-	static const FLazyName DebuggerTabId;
 #endif // WITH_STATETREE_TRACE_DEBUGGER
-	
-	static const FLazyName StateTreeOutlinerTabId;
-	static const FLazyName StateTreeSearchTabId;
-	static const FLazyName StateTreeStatisticsTabId;
+
 };
 
 
diff --git a/Engine/Plugins/Runtime/StateTree/Source/StateTreeEditorModule/Public/StateTreeEditorWorkspaceTabHost.h b/Engine/Plugins/Runtime/StateTree/Source/StateTreeEditorModule/Public/StateTreeEditorWorkspaceTabHost.h
new file mode 100644
--- /dev/null
+++ b/Engine/Plugins/Runtime/StateTree/Source/StateTreeEditorModule/Public/StateTreeEditorWorkspaceTabHost.h
@@ -0,0 +1,59 @@
+// Copyright Epic Games, Inc. All Rights Reserved.
+
+#pragma once
+
+#include "CoreMinimal.h"
+#include "Textures/SlateIcon.h"
+#include "Framework/Docking/TabManager.h"
+
+class SDockTab;
+
+namespace UE::StateTreeEditor
+{
+	struct FMinorWorkspaceTabConfig
+	{
+		FName ID;
+		FText Label;
+		FText Tooltip;
+		FSlateIcon Icon;
+		FName UISystemID;
+	};
+
+	struct FSpawnedWorkspaceTab
+	{
+		FName TabID;
+		TWeakPtr<SDockTab> DockTab;
+	};
+
+	// Interface required for re-using the same tab management across different AssetEditors
+	class STATETREEEDITORMODULE_API FWorkspaceTabHost : public TSharedFromThis<FWorkspaceTabHost>
+	{
+	public:
+		static const FLazyName OutlinerTabId;
+		static const FLazyName SearchTabId;
+		static const FLazyName StatisticsTabId;
+		static const FLazyName DebuggerTabId;
+
+	public:
+		virtual ~FWorkspaceTabHost() = default;
+
+		FOnSpawnTab CreateSpawnDelegate(FName TabID);
+		TConstArrayView<FMinorWorkspaceTabConfig> GetTabConfigs() const;
+
+		TConstArrayView<FSpawnedWorkspaceTab> GetSpawnedTabs() const
+		{
+			return SpawnedTabs;
+		}
+
+		TSharedRef<SDockTab> Spawn(FName TabID);
+
+		DECLARE_MULTICAST_DELEGATE_OneParam(FOnWorkspaceSpawnTab, FSpawnedWorkspaceTab);
+		FOnWorkspaceSpawnTab OnTabSpawned;
+		FOnWorkspaceSpawnTab OnTabClosed;
+
+	private:
+		TSharedRef<SDockTab> HandleSpawnDelegate(const FSpawnTabArgs& Args, FName TabID);
+		void HandleTabClosed(TSharedRef<SDockTab>);
+		TArray<FSpawnedWorkspaceTab> SpawnedTabs;
+	};
+}
