StateTree: Test the weak object pointer before using it.

#rnx

diff --git a/Engine/Plugins/Runtime/StateTree/Source/StateTreeEditorModule/Private/StateTreeEditorModeToolkit.cpp b/Engine/Plugins/Runtime/StateTree/Source/StateTreeEditorModule/Private/StateTreeEditorModeToolkit.cpp
--- a/Engine/Plugins/Runtime/StateTree/Source/StateTreeEditorModule/Private/StateTreeEditorModeToolkit.cpp
+++ b/Engine/Plugins/Runtime/StateTree/Source/StateTreeEditorModule/Private/StateTreeEditorModeToolkit.cpp
@@ -179,24 +179,25 @@ FSlateIcon FStateTreeEditorModeToolkit::GetCompileStatusImage() const
 	static const FName CompileStatusError("Blueprint.CompileStatus.Overlay.Error");
 	static const FName CompileStatusGood("Blueprint.CompileStatus.Overlay.Good");
 
-	UStateTree* StateTree = WeakEditorMode.Pin()->GetStateTree();
-	UStateTreeEditorMode* SharedEditorMode = WeakEditorMode.Get();
-	
-	if (StateTree && SharedEditorMode)
+	if (UStateTreeEditorMode* EditorMode = WeakEditorMode.Get())
 	{
-		const bool bCompiledDataResetDuringLoad = StateTree->LastCompiledEditorDataHash == SharedEditorMode->EditorDataHash && !StateTree->IsReadyToRun();
-
-		if (!SharedEditorMode->bLastCompileSucceeded || bCompiledDataResetDuringLoad)
-		{
-			return FSlateIcon(FAppStyle::GetAppStyleSetName(), CompileStatusBackground, NAME_None, CompileStatusError);
-		}
-	
-		if (StateTree->LastCompiledEditorDataHash != SharedEditorMode->EditorDataHash)
+		UStateTree* StateTree = EditorMode->GetStateTree();
+		if (StateTree)
 		{
-			return FSlateIcon(FAppStyle::GetAppStyleSetName(), CompileStatusBackground, NAME_None, CompileStatusUnknown);
+			const bool bCompiledDataResetDuringLoad = StateTree->LastCompiledEditorDataHash == EditorMode->EditorDataHash && !StateTree->IsReadyToRun();
+
+			if (!EditorMode->bLastCompileSucceeded || bCompiledDataResetDuringLoad)
+			{
+				return FSlateIcon(FAppStyle::GetAppStyleSetName(), CompileStatusBackground, NAME_None, CompileStatusError);
+			}
+
+			if (StateTree->LastCompiledEditorDataHash != EditorMode->EditorDataHash)
+			{
+				return FSlateIcon(FAppStyle::GetAppStyleSetName(), CompileStatusBackground, NAME_None, CompileStatusUnknown);
+			}
+
+			return FSlateIcon(FAppStyle::GetAppStyleSetName(), CompileStatusBackground, NAME_None, CompileStatusGood);
 		}
-	
-		return FSlateIcon(FAppStyle::GetAppStyleSetName(), CompileStatusBackground, NAME_None, CompileStatusGood);
 	}
 	
 	return FSlateIcon(FAppStyle::GetAppStyleSetName(), CompileStatusBackground, NAME_None, CompileStatusUnknown);
@@ -499,53 +500,64 @@ void FStateTreeEditorModeToolkit::OnNodeBPBaseClassPicked(UClass* NodeClass) con
 {
 	check(NodeClass);
 
-	UStateTree* StateTree = WeakEditorMode->GetStateTree();
-    if (!StateTree)
-    {
-    	return;
-    }
-
-    const FString ClassName = FBlueprintEditorUtils::GetClassNameWithoutSuffix(NodeClass);
-    const FString PathName = FPaths::GetPath(StateTree->GetOutermost()->GetPathName());
-
-    // Now that we've generated some reasonable default locations/names for the package, allow the user to have the final say
-    // before we create the package and initialize the blueprint inside of it.
-    FSaveAssetDialogConfig SaveAssetDialogConfig;
-    SaveAssetDialogConfig.DialogTitleOverride = LOCTEXT("SaveAssetDialogTitle", "Save Asset As");
-    SaveAssetDialogConfig.DefaultPath = PathName;
-    SaveAssetDialogConfig.DefaultAssetName = ClassName + TEXT("_New");
-    SaveAssetDialogConfig.ExistingAssetPolicy = ESaveAssetDialogExistingAssetPolicy::Disallow;
-
-    const FContentBrowserModule& ContentBrowserModule = FModuleManager::LoadModuleChecked<FContentBrowserModule>("ContentBrowser");
-    const FString SaveObjectPath = ContentBrowserModule.Get().CreateModalSaveAssetDialog(SaveAssetDialogConfig);
-    if (!SaveObjectPath.IsEmpty())
-    {
-    	const FString SavePackageName = FPackageName::ObjectPathToPackageName(SaveObjectPath);
-    	const FString SavePackagePath = FPaths::GetPath(SavePackageName);
-    	const FString SaveAssetName = FPaths::GetBaseFilename(SavePackageName);
-
-    	if (UPackage* Package = CreatePackage(*SavePackageName))
-    	{
-    		// Create and init a new Blueprint
-    		if (UBlueprint* NewBP = FKismetEditorUtilities::CreateBlueprint(NodeClass, Package, FName(*SaveAssetName), BPTYPE_Normal, UBlueprint::StaticClass(), UBlueprintGeneratedClass::StaticClass()))
-    		{
-    			GEditor->GetEditorSubsystem<UAssetEditorSubsystem>()->OpenEditorForAsset(NewBP);
-
-    			// Notify the asset registry
-    			FAssetRegistryModule::AssetCreated(NewBP);
-
-    			Package->MarkPackageDirty();
-    		}
-    	}
-    }
-
-    FSlateApplication::Get().DismissAllMenus();
+	UStateTreeEditorMode* EditorMode = WeakEditorMode.Get();
+	if (EditorMode == nullptr)
+	{
+		return;
+	}
+	UStateTree* StateTree = EditorMode->GetStateTree();
+	if (StateTree == nullptr)
+	{
+		return;
+	}
+
+	const FString ClassName = FBlueprintEditorUtils::GetClassNameWithoutSuffix(NodeClass);
+	const FString PathName = FPaths::GetPath(StateTree->GetOutermost()->GetPathName());
+
+	// Now that we've generated some reasonable default locations/names for the package, allow the user to have the final say
+	// before we create the package and initialize the blueprint inside of it.
+	FSaveAssetDialogConfig SaveAssetDialogConfig;
+	SaveAssetDialogConfig.DialogTitleOverride = LOCTEXT("SaveAssetDialogTitle", "Save Asset As");
+	SaveAssetDialogConfig.DefaultPath = PathName;
+	SaveAssetDialogConfig.DefaultAssetName = ClassName + TEXT("_New");
+	SaveAssetDialogConfig.ExistingAssetPolicy = ESaveAssetDialogExistingAssetPolicy::Disallow;
+
+	const FContentBrowserModule& ContentBrowserModule = FModuleManager::LoadModuleChecked<FContentBrowserModule>("ContentBrowser");
+	const FString SaveObjectPath = ContentBrowserModule.Get().CreateModalSaveAssetDialog(SaveAssetDialogConfig);
+	if (!SaveObjectPath.IsEmpty())
+	{
+		const FString SavePackageName = FPackageName::ObjectPathToPackageName(SaveObjectPath);
+		const FString SavePackagePath = FPaths::GetPath(SavePackageName);
+		const FString SaveAssetName = FPaths::GetBaseFilename(SavePackageName);
+
+		if (UPackage* Package = CreatePackage(*SavePackageName))
+		{
+			// Create and init a new Blueprint
+			if (UBlueprint* NewBP = FKismetEditorUtilities::CreateBlueprint(NodeClass, Package, FName(*SaveAssetName), BPTYPE_Normal, UBlueprint::StaticClass(), UBlueprintGeneratedClass::StaticClass()))
+			{
+				GEditor->GetEditorSubsystem<UAssetEditorSubsystem>()->OpenEditorForAsset(NewBP);
+
+				// Notify the asset registry
+				FAssetRegistryModule::AssetCreated(NewBP);
+
+				Package->MarkPackageDirty();
+			}
+		}
+	}
+
+	FSlateApplication::Get().DismissAllMenus();
 }
 
 FText FStateTreeEditorModeToolkit::GetStatisticsText() const
 {
-	UStateTree* StateTree = WeakEditorMode->GetStateTree(); 
-	if (!StateTree)
+	UStateTreeEditorMode* EditorMode = WeakEditorMode.Get();
+	if (EditorMode == nullptr)
+	{
+		return FText::GetEmpty();
+	}
+
+	UStateTree* StateTree = EditorMode->GetStateTree();
+	if (StateTree == nullptr)
 	{
 		return FText::GetEmpty();
 	}
@@ -558,7 +570,6 @@ FText FStateTreeEditorModeToolkit::GetStatisticsText() const
    }
 
 	TArray<FText> Rows;
-
 	for (const FStateTreeMemoryUsage& Usage : MemoryUsages)
 	{
 		const FText SizeText = FText::AsMemory(Usage.EstimatedMemoryUsage);
@@ -572,12 +583,15 @@ FText FStateTreeEditorModeToolkit::GetStatisticsText() const
 void FStateTreeEditorModeToolkit::UpdateStateTreeOutliner()
 {
 	StateTreeOutliner = SNullWidget::NullWidget;
-	if (WeakEditorMode->GetStateTree())
+	if (UStateTreeEditorMode* EditorMode = WeakEditorMode.Get())
 	{
-		if (UStateTreeEditingSubsystem* Subsystem = GEditor->GetEditorSubsystem<UStateTreeEditingSubsystem>())
+		if (EditorMode && EditorMode->GetStateTree())
 		{
-			StateTreeOutliner = SNew(SStateTreeOutliner, Subsystem->FindOrAddViewModel(WeakEditorMode->GetStateTree()), GetToolkitCommands());
-		}		
+			if (UStateTreeEditingSubsystem* Subsystem = GEditor->GetEditorSubsystem<UStateTreeEditingSubsystem>())
+			{
+				StateTreeOutliner = SNew(SStateTreeOutliner, Subsystem->FindOrAddViewModel(EditorMode->GetStateTree()), GetToolkitCommands());
+			}		
+		}
 	}
 	
 	if(TSharedPtr<SDockTab> SharedOutlinerTab = WeakOutlinerTab.Pin())
@@ -600,11 +614,14 @@ void FStateTreeEditorModeToolkit::UpdateDebuggerView()
 		SharedDebuggerTab->SetContent(DebuggerView);
 	}
 
-	if (UStateTree* StateTree = WeakEditorMode->GetStateTree())
+	if (UStateTreeEditorMode* EditorMode = WeakEditorMode.Get())
 	{
-		if (UStateTreeEditingSubsystem* Subsystem = GEditor->GetEditorSubsystem<UStateTreeEditingSubsystem>())
+		if (UStateTree* StateTree = EditorMode->GetStateTree())
 		{
-			DebuggerView = SNew(SStateTreeDebuggerView, *StateTree, Subsystem->FindOrAddViewModel(StateTree), GetToolkitCommands());
+			if (UStateTreeEditingSubsystem* Subsystem = GEditor->GetEditorSubsystem<UStateTreeEditingSubsystem>())
+			{
+				DebuggerView = SNew(SStateTreeDebuggerView, *StateTree, Subsystem->FindOrAddViewModel(StateTree), GetToolkitCommands());
+			}
 		}
 	}
 
