StateTree: Add FStateTreeReadOnlyEecutionContext. An execution context that allows concurrent read access. Useful for debugging on multiple threads. Move read-only functions from the FStateTreeExecutionContext (of FStateTreeMinimalExecutionContext) to the new context.
Deprecate the variable InstanceDataStorage from FStateTreeExecutionContext, Storage duplicated it.
Deprecate StateTreeEvaluatorBase::AppendDebugInfoString and StateTreeTaskBase to make them usable by FStateTreeReadOnlyEecutionContext.

diff --git a/Engine/Plugins/AI/MassAI/Source/MassAIBehavior/Private/MassStateTreeExecutionContext.cpp b/Engine/Plugins/AI/MassAI/Source/MassAIBehavior/Private/MassStateTreeExecutionContext.cpp
--- a/Engine/Plugins/AI/MassAI/Source/MassAIBehavior/Private/MassStateTreeExecutionContext.cpp
+++ b/Engine/Plugins/AI/MassAI/Source/MassAIBehavior/Private/MassStateTreeExecutionContext.cpp
@@ -154,14 +154,11 @@ void FMassStateTreeExecutionContext::SetEntity(const FMassEntityHandle InEntity)
 	Entity = InEntity;
 
 #if WITH_STATETREE_DEBUG
-	if (InstanceDataStorage)
+	const TInstancedStruct<FStateTreeExecutionExtension>& ExecutionExtension = Storage.GetExecutionState().ExecutionExtension;
+	if (ExecutionExtension.IsValid())
 	{
-		const TInstancedStruct<FStateTreeExecutionExtension>& ExecutionExtension = InstanceDataStorage->GetExecutionState().ExecutionExtension;
-		if (ExecutionExtension.IsValid())
-		{
-			const FMassExecutionExtension* Extension = ExecutionExtension.GetPtr<const FMassExecutionExtension>();
-			ensureMsgf(Extension && Extension->Entity == Entity, TEXT("Incorrect extension type or the entity handle is different from previous execution."));
-		}
+		const FMassExecutionExtension* Extension = ExecutionExtension.GetPtr<const FMassExecutionExtension>();
+		ensureMsgf(Extension && Extension->Entity == Entity, TEXT("Incorrect extension type or the entity handle is different from previous execution."));
 	}
 #endif
 }
diff --git a/Engine/Plugins/AI/MassAI/Source/MassAIDebug/Private/GameplayDebuggerCategory_Mass.cpp b/Engine/Plugins/AI/MassAI/Source/MassAIDebug/Private/GameplayDebuggerCategory_Mass.cpp
--- a/Engine/Plugins/AI/MassAI/Source/MassAIDebug/Private/GameplayDebuggerCategory_Mass.cpp
+++ b/Engine/Plugins/AI/MassAI/Source/MassAIDebug/Private/GameplayDebuggerCategory_Mass.cpp
@@ -749,8 +749,7 @@ void FGameplayDebuggerCategory_Mass::CollectData(APlayerController* OwnerPC, AAc
 						{
 							if (FStateTreeInstanceData* InstanceData = MassStateTreeSubsystem->GetInstanceData(StateTreeInstance.InstanceHandle))
 							{
-								FMassStateTreeExecutionContext StateTreeContext(*MassStateTreeSubsystem, *StateTree, *InstanceData, Context);
-								StateTreeContext.SetEntity(Entity);
+								FStateTreeReadOnlyExecutionContext StateTreeContext(MassStateTreeSubsystem, StateTree, *InstanceData);
 
 								Status += StateTreeContext.GetActiveStateName();
 								Status += FString::Printf(TEXT("  {yellow}%d{white}\n"), StateTreeContext.GetStateChangeCount());
diff --git a/Engine/Plugins/AI/MassAI/Source/MassAIDebug/Private/MassDebugStateTreeProcessor.cpp b/Engine/Plugins/AI/MassAI/Source/MassAIDebug/Private/MassDebugStateTreeProcessor.cpp
--- a/Engine/Plugins/AI/MassAI/Source/MassAIDebug/Private/MassDebugStateTreeProcessor.cpp
+++ b/Engine/Plugins/AI/MassAI/Source/MassAIDebug/Private/MassDebugStateTreeProcessor.cpp
@@ -89,8 +89,7 @@ void UMassDebugStateTreeProcessor::Execute(FMassEntityManager& EntityManager, FM
 			
 			if (Entity == SelectedEntity)
 			{
-				FMassStateTreeExecutionContext StateTreeContext(*MassStateTreeSubsystem, *StateTree, *InstanceData, Context);
-				StateTreeContext.SetEntity(Entity);
+				FStateTreeReadOnlyExecutionContext StateTreeContext(MassStateTreeSubsystem, StateTree, *InstanceData);
 				
 #if WITH_GAMEPLAY_DEBUGGER
 				Debugger->AppendSelectedEntityInfo(StateTreeContext.GetDebugInfoString());
@@ -106,8 +105,7 @@ void UMassDebugStateTreeProcessor::Execute(FMassEntityManager& EntityManager, FM
 				const FVector ZOffset(0,0,50);
 				const FVector Position = Transform.GetTransform().GetLocation() + ZOffset;
 
-				FMassStateTreeExecutionContext StateTreeContext(*MassStateTreeSubsystem, *StateTree, *InstanceData, Context);
-				StateTreeContext.SetEntity(Entity);
+				FStateTreeReadOnlyExecutionContext StateTreeContext(MassStateTreeSubsystem, StateTree, *InstanceData);
 
 				// State
 				UE_VLOG_SEGMENT_THICK(this, LogStateTree, Log, Position, Position + FVector(0,0,50), EntityColor, /*Thickness*/ 2, TEXT("%s %s"),
diff --git a/Engine/Plugins/Runtime/StateTree/Source/StateTreeModule/Private/StateTreeEvaluatorBase.cpp b/Engine/Plugins/Runtime/StateTree/Source/StateTreeModule/Private/StateTreeEvaluatorBase.cpp
--- a/Engine/Plugins/Runtime/StateTree/Source/StateTreeModule/Private/StateTreeEvaluatorBase.cpp
+++ b/Engine/Plugins/Runtime/StateTree/Source/StateTreeModule/Private/StateTreeEvaluatorBase.cpp
@@ -5,6 +5,16 @@
 #include UE_INLINE_GENERATED_CPP_BY_NAME(StateTreeEvaluatorBase)
 
 #if WITH_GAMEPLAY_DEBUGGER
+FString FStateTreeEvaluatorBase::GetDebugInfo(const FStateTreeReadOnlyExecutionContext& Context) const
+{
+	TStringBuilder<256> Builder;
+	Builder << TEXT('[');
+	Builder << Name;
+	Builder << TEXT("]\n");
+	return Builder.ToString();
+}
+
+// Deprecated
 void FStateTreeEvaluatorBase::AppendDebugInfoString(FString& DebugString, const FStateTreeExecutionContext& Context) const
 {
 	DebugString += FString::Printf(TEXT("[%s]\n"), *Name.ToString());
diff --git a/Engine/Plugins/Runtime/StateTree/Source/StateTreeModule/Private/StateTreeExecutionContext.cpp b/Engine/Plugins/Runtime/StateTree/Source/StateTreeModule/Private/StateTreeExecutionContext.cpp
--- a/Engine/Plugins/Runtime/StateTree/Source/StateTreeModule/Private/StateTreeExecutionContext.cpp
+++ b/Engine/Plugins/Runtime/StateTree/Source/StateTreeModule/Private/StateTreeExecutionContext.cpp
@@ -26,7 +26,7 @@
 {
 	constexpr int32 IndentSize = 2;	// Debug printing indent for hierarchical data.
 	static std::atomic<uint32> InstanceSerialNumber = 0;
-}; // UE::StateTree
+}; // namespace
 
 namespace UE::StateTree::ExecutionContext::Private
 {
@@ -116,51 +116,34 @@
 				return Task.bTickProcessed;
 			});
 	}
-}
+} // namespace
 
-FStateTreeMinimalExecutionContext::FStateTreeMinimalExecutionContext(UObject& InOwner, const UStateTree& InStateTree, FStateTreeInstanceData& InInstanceData)
-	: FStateTreeMinimalExecutionContext(InOwner, InStateTree, InInstanceData.GetMutableStorage())
+/**
+ * FStateTreeReadOnlyExecutionContext implementation
+ */
+FStateTreeReadOnlyExecutionContext::FStateTreeReadOnlyExecutionContext(TNotNull<UObject*> InOwner, TNotNull<const UStateTree*> InStateTree, FStateTreeInstanceData& InInstanceData)
+	: FStateTreeReadOnlyExecutionContext(InOwner, InStateTree, InInstanceData.GetMutableStorage())
 {
 }
 
-FStateTreeMinimalExecutionContext::FStateTreeMinimalExecutionContext(UObject& InOwner, const UStateTree& InStateTree, FStateTreeInstanceStorage& InStorage)
-	: Owner(InOwner)
-	, RootStateTree(InStateTree)
+FStateTreeReadOnlyExecutionContext::FStateTreeReadOnlyExecutionContext(TNotNull<UObject*> InOwner, TNotNull<const UStateTree*> InStateTree, FStateTreeInstanceStorage& InStorage)
+	: Owner(*InOwner)
+	, RootStateTree(*InStateTree)
 	, Storage(InStorage)
 {
-	Storage.AcquireWriteAccess();
-	if (InStateTree.IsReadyToRun())
+	Storage.AcquireReadAccess();
+	if (IsValid())
 	{
 		Storage.GetRuntimeValidation().SetContext(&Owner, &RootStateTree);
 	}
 }
 
-FStateTreeMinimalExecutionContext::~FStateTreeMinimalExecutionContext()
+FStateTreeReadOnlyExecutionContext::~FStateTreeReadOnlyExecutionContext()
 {
-	Storage.ReleaseWriteAccess();
+	Storage.ReleaseReadAccess();
 }
 
-FString FStateTreeMinimalExecutionContext::GetInstanceDescriptionInternal() const
-{
-	const TInstancedStruct<FStateTreeExecutionExtension>& ExecutionExtension = Storage.GetExecutionState().ExecutionExtension;
-	return ExecutionExtension.IsValid()
-		? ExecutionExtension.Get().GetInstanceDescription(FStateTreeExecutionExtension::FContextParameters(Owner, RootStateTree, Storage))
-		: Owner.GetName();
-}
-
-EStateTreeRunStatus FStateTreeMinimalExecutionContext::GetStateTreeRunStatus() const
+FStateTreeScheduledTick FStateTreeReadOnlyExecutionContext::GetNextScheduledTick() const
 {
 	if (!IsValid())
 	{
-		STATETREE_LOG(Warning, TEXT("%hs: StateTree context is not initialized properly ('%s' using StateTree '%s')"),
-			__FUNCTION__, *GetNameSafe(&Owner), *GetFullNameSafe(&RootStateTree));
-		return EStateTreeRunStatus::Failed;
-	}
-
-	return Storage.GetExecutionState().TreeRunStatus;
-}
-
-FStateTreeScheduledTick FStateTreeMinimalExecutionContext::GetNextScheduledTick() const
-{
-	if (!IsValid())
-	{
@@ -294,4 +277,283 @@
 	return FStateTreeScheduledTick::MakeSleep();
 }
 
+EStateTreeRunStatus FStateTreeReadOnlyExecutionContext::GetStateTreeRunStatus() const
+{
+	if (!IsValid())
+	{
+		STATETREE_LOG(Warning, TEXT("%hs: StateTree context is not initialized properly ('%s' using StateTree '%s')"),
+			__FUNCTION__, *GetNameSafe(&Owner), *GetFullNameSafe(&RootStateTree));
+		return EStateTreeRunStatus::Failed;
+	}
+
+	return Storage.GetExecutionState().TreeRunStatus;
+}
+
+EStateTreeRunStatus FStateTreeReadOnlyExecutionContext::GetLastTickStatus() const
+{
+	if (!IsValid())
+	{
+		STATETREE_LOG(Warning, TEXT("%hs: StateTree context is not initialized properly ('%s' using StateTree '%s')"),
+			__FUNCTION__, *GetNameSafe(&Owner), *GetFullNameSafe(&RootStateTree));
+		return EStateTreeRunStatus::Failed;
+	}
+
+	const FStateTreeExecutionState& Exec = Storage.GetExecutionState();
+	return Exec.LastTickStatus;
+}
+
+TConstArrayView<FStateTreeExecutionFrame> FStateTreeReadOnlyExecutionContext::GetActiveFrames() const
+{
+	if (!IsValid())
+	{
+		STATETREE_LOG(Warning, TEXT("%hs: StateTree context is not initialized properly ('%s' using StateTree '%s')"),
+			__FUNCTION__, *GetNameSafe(&Owner), *GetFullNameSafe(&RootStateTree));
+		return TConstArrayView<FStateTreeExecutionFrame>();
+	}
+
+	const FStateTreeExecutionState& Exec = Storage.GetExecutionState();
+	return Exec.ActiveFrames;
+}
+
+FString FStateTreeReadOnlyExecutionContext::GetActiveStateName() const
+{
+	if (!IsValid())
+	{
+		STATETREE_LOG(Warning, TEXT("%hs: StateTree context is not initialized properly ('%s' using StateTree '%s')"),
+			__FUNCTION__, *GetNameSafe(&Owner), *GetFullNameSafe(&RootStateTree));
+		return FString();
+	}
+
+	const FStateTreeExecutionState& Exec = Storage.GetExecutionState();
+
+	TStringBuilder<1024> FullStateName;
+
+	const UStateTree* LastStateTree = &RootStateTree;
+	int32 Indent = 0;
+
+	for (int32 FrameIndex = 0; FrameIndex < Exec.ActiveFrames.Num(); FrameIndex++)
+	{
+		const FStateTreeExecutionFrame& CurrentFrame = Exec.ActiveFrames[FrameIndex];
+		const UStateTree* CurrentStateTree = CurrentFrame.StateTree;
+
+		// Append linked state marker at the end of the previous line.
+		if (Indent > 0)
+		{
+			FullStateName << TEXT(" >");
+		}
+		// If tree has changed, append that too.
+		if (CurrentFrame.StateTree != LastStateTree)
+		{
+			FullStateName << TEXT(" [");
+			FullStateName << CurrentFrame.StateTree.GetFName();
+			FullStateName << TEXT(']');
+
+			LastStateTree = CurrentFrame.StateTree;
+		}
+
+		for (int32 Index = 0; Index < CurrentFrame.ActiveStates.Num(); Index++)
+		{
+			const FStateTreeStateHandle Handle = CurrentFrame.ActiveStates[Index];
+			if (Handle.IsValid())
+			{
+				const FCompactStateTreeState& State = CurrentStateTree->States[Handle.Index];
+				if (Indent > 0)
+				{
+					FullStateName += TEXT("\n");
+				}
+				FullStateName.Appendf(TEXT("%*s-"), Indent * 3, TEXT("")); // Indent
+				FullStateName << State.Name;
+				Indent++;
+			}
+		}
+	}
+
+	switch (Exec.TreeRunStatus)
+	{
+	case EStateTreeRunStatus::Failed:
+		FullStateName << TEXT(" FAILED\n");
+		break;
+	case EStateTreeRunStatus::Succeeded:
+		FullStateName << TEXT(" SUCCEEDED\n");
+		break;
+	case EStateTreeRunStatus::Running:
+		// Empty
+		break;
+	default:
+		FullStateName << TEXT("--\n");
+	}
+
+	return FullStateName.ToString();
+}
+
+TArray<FName> FStateTreeReadOnlyExecutionContext::GetActiveStateNames() const
+{
+	if (!IsValid())
+	{
+		STATETREE_LOG(Warning, TEXT("%hs: StateTree context is not initialized properly ('%s' using StateTree '%s')"),
+			__FUNCTION__, *GetNameSafe(&Owner), *GetFullNameSafe(&RootStateTree));
+		return TArray<FName>();
+	}
+
+	TArray<FName> Result;
+	const FStateTreeExecutionState& Exec = Storage.GetExecutionState();
+
+	// Active States
+	for (const FStateTreeExecutionFrame& CurrentFrame : Exec.ActiveFrames)
+	{
+		const UStateTree* CurrentStateTree = CurrentFrame.StateTree;
+		for (int32 Index = 0; Index < CurrentFrame.ActiveStates.Num(); Index++)
+		{
+			const FStateTreeStateHandle Handle = CurrentFrame.ActiveStates[Index];
+			if (Handle.IsValid())
+			{
+				const FCompactStateTreeState& State = CurrentStateTree->States[Handle.Index];
+				Result.Add(State.Name);
+			}
+		}
+	}
+
+	return Result;
+}
+
+#if WITH_GAMEPLAY_DEBUGGER
+FString FStateTreeReadOnlyExecutionContext::GetDebugInfoString() const
+{
+	TStringBuilder<2048> DebugString;
+	DebugString << TEXT("StateTree (asset: '");
+	RootStateTree.GetFullName(DebugString);
+	DebugString << TEXT("')");
+
+	if (IsValid())
+	{
+		const FStateTreeExecutionState& Exec = Storage.GetExecutionState();
+
+		DebugString << TEXT("Status: ");
+		DebugString << UEnum::GetDisplayValueAsText(Exec.TreeRunStatus).ToString();
+		DebugString << TEXT("\n");
+
+		// Active States
+		DebugString << TEXT("Current State:\n");
+		for (const FStateTreeExecutionFrame& CurrentFrame : Exec.ActiveFrames)
+		{
+			const UStateTree* CurrentStateTree = CurrentFrame.StateTree;
+
+			if (CurrentFrame.bIsGlobalFrame)
+			{
+				DebugString.Appendf(TEXT("\nEvaluators\n  [ %-30s | %8s | %15s ]\n"),
+					TEXT("Name"), TEXT("Bindings"), TEXT("Data Handle"));
+				for (int32 EvalIndex = CurrentStateTree->EvaluatorsBegin; EvalIndex < (CurrentStateTree->EvaluatorsBegin + CurrentStateTree->EvaluatorsNum); EvalIndex++)
+				{
+					const FStateTreeEvaluatorBase& Eval = CurrentStateTree->Nodes[EvalIndex].Get<const FStateTreeEvaluatorBase>();
+					DebugString.Appendf(TEXT("| %-30s | %8d | %15s |\n"),
+						*Eval.Name.ToString(), Eval.BindingsBatch.Get(), *Eval.InstanceDataHandle.Describe());
+				}
+
+				DebugString << TEXT("\nGlobal Tasks\n");
+				for (int32 TaskIndex = CurrentStateTree->GlobalTasksBegin; TaskIndex < (CurrentStateTree->GlobalTasksBegin + CurrentStateTree->GlobalTasksNum); TaskIndex++)
+				{
+					const FStateTreeTaskBase& Task = CurrentStateTree->Nodes[TaskIndex].Get<const FStateTreeTaskBase>();
+					if (Task.bTaskEnabled)
+					{
+						DebugString << Task.GetDebugInfo(*this);
+
+					}
+				}
+			}
+
+			for (int32 Index = 0; Index < CurrentFrame.ActiveStates.Num(); Index++)
+			{
+				FStateTreeStateHandle Handle = CurrentFrame.ActiveStates[Index];
+				if (Handle.IsValid())
+				{
+					const FCompactStateTreeState& State = RootStateTree.States[Handle.Index];
+					DebugString << TEXT('[');
+					DebugString << State.Name;
+					DebugString << TEXT("]\n");
+
+					if (State.TasksNum > 0)
+					{
+						DebugString += TEXT("\nTasks:\n");
+						for (int32 TaskIndex = State.TasksBegin; TaskIndex < (State.TasksBegin + State.TasksNum); TaskIndex++)
+						{
+							const FStateTreeTaskBase& Task = RootStateTree.Nodes[TaskIndex].Get<const FStateTreeTaskBase>();
+							if (Task.bTaskEnabled)
+							{
+								DebugString << Task.GetDebugInfo(*this);
+							}
+						}
+					}
+				}
+			}
+		}
+	}
+	else
+	{
+		DebugString << TEXT("StateTree context is not initialized properly.");
+	}
+
+	return DebugString.ToString();
+}
+#endif // WITH_GAMEPLAY_DEBUGGER
+
+#if WITH_STATETREE_DEBUG
+int32 FStateTreeReadOnlyExecutionContext::GetStateChangeCount() const
+{
+	if (!IsValid())
+	{
+		STATETREE_LOG(Warning, TEXT("%hs: StateTree context is not initialized properly ('%s' using StateTree '%s')"),
+			__FUNCTION__, *GetNameSafe(&Owner), *GetFullNameSafe(&RootStateTree));
+		return 0;
+	}
+
+	const FStateTreeExecutionState& Exec = Storage.GetExecutionState();
+	return Exec.StateChangeCount;
+}
+
+void FStateTreeReadOnlyExecutionContext::DebugPrintInternalLayout()
+{
+	LOG_SCOPE_VERBOSITY_OVERRIDE(LogStateTree, ELogVerbosity::Log);
+	UE_LOG(LogStateTree, Log, TEXT("%s"), *RootStateTree.DebugInternalLayoutAsString());
+}
+#endif // WITH_STATETREE_DEBUG
+
+FString FStateTreeReadOnlyExecutionContext::GetInstanceDescriptionInternal() const
+{
+	const TInstancedStruct<FStateTreeExecutionExtension>& ExecutionExtension = Storage.GetExecutionState().ExecutionExtension;
+	return ExecutionExtension.IsValid()
+		? ExecutionExtension.Get().GetInstanceDescription(FStateTreeExecutionExtension::FContextParameters(Owner, RootStateTree, Storage))
+		: Owner.GetName();
+}
+
+#if WITH_STATETREE_TRACE
+FStateTreeInstanceDebugId FStateTreeReadOnlyExecutionContext::GetInstanceDebugId() const
+{
+	FStateTreeInstanceDebugId& InstanceDebugId = Storage.GetMutableExecutionState().InstanceDebugId;
+	if (!InstanceDebugId.IsValid())
+	{
+		InstanceDebugId = FStateTreeInstanceDebugId(GetTypeHash(GetInstanceDescriptionInternal()), ++UE::StateTree::Debug::InstanceSerialNumber);
+	}
+	return InstanceDebugId;
+}
+#endif // WITH_STATETREE_TRACE
+
+/**
+ * FStateTreeMinimalExecutionContext implementation
+ */
+FStateTreeMinimalExecutionContext::FStateTreeMinimalExecutionContext(UObject& InOwner, const UStateTree& InStateTree, FStateTreeInstanceData& InInstanceData)
+	: FStateTreeMinimalExecutionContext(InOwner, InStateTree, InInstanceData.GetMutableStorage())
+{
+}
+
+FStateTreeMinimalExecutionContext::FStateTreeMinimalExecutionContext(UObject& InOwner, const UStateTree& InStateTree, FStateTreeInstanceStorage& InStorage)
+	: FStateTreeReadOnlyExecutionContext(&InOwner, &InStateTree, InStorage)
+{
+	Storage.AcquireWriteAccess();
+}
+
+FStateTreeMinimalExecutionContext::~FStateTreeMinimalExecutionContext()
+{
+	Storage.ReleaseWriteAccess();
+}
+
 UE::StateTree::FScheduledTickHandle FStateTreeMinimalExecutionContext::AddScheduledTickRequest(FStateTreeScheduledTick ScheduledTick)
@@ -368,18 +630,9 @@
 	}
 }
 
-#if WITH_STATETREE_TRACE
-FStateTreeInstanceDebugId FStateTreeMinimalExecutionContext::GetInstanceDebugId() const
-{
-	FStateTreeInstanceDebugId& InstanceDebugId = Storage.GetMutableExecutionState().InstanceDebugId;
-	if (!InstanceDebugId.IsValid())
-	{
-		InstanceDebugId = FStateTreeInstanceDebugId(GetTypeHash(GetInstanceDescriptionInternal()), ++UE::StateTree::Debug::InstanceSerialNumber);
-	}
-	return InstanceDebugId;
-}
-#endif // WITH_STATETREE_TRACE
-
+/**
+ * FStateTreeExecutionContext::FCurrentlyProcessedFrameScope implementation
+ */
 FStateTreeExecutionContext::FCurrentlyProcessedFrameScope::FCurrentlyProcessedFrameScope(FStateTreeExecutionContext& InContext, const FStateTreeExecutionFrame* CurrentParentFrame, const FStateTreeExecutionFrame& CurrentFrame): Context(InContext)
 {
 	check(CurrentFrame.StateTree);
@@ -407,6 +660,9 @@
 	}
 }
 
+/**
+ * FStateTreeExecutionContext::FNodeInstanceDataScope implementation
+ */
 FStateTreeExecutionContext::FNodeInstanceDataScope::FNodeInstanceDataScope(FStateTreeExecutionContext& InContext, const FStateTreeNodeBase* InNode, const int32 InNodeIndex, const FStateTreeDataHandle InNodeDataHandle, const FStateTreeDataView InNodeInstanceData)
 	: Context(InContext)
 {
@@ -428,21 +684,19 @@
 	Context.CurrentNode = SavedNode;
 }
 
+/**
+ * FStateTreeExecutionContext implementation
+ */
 FStateTreeExecutionContext::FStateTreeExecutionContext(UObject& InOwner, const UStateTree& InStateTree, FStateTreeInstanceData& InInstanceData, const FOnCollectStateTreeExternalData& InCollectExternalDataDelegate, const EStateTreeRecordTransitions RecordTransitions)
 	: FStateTreeMinimalExecutionContext(InOwner, InStateTree, InInstanceData)
 	, InstanceData(InInstanceData)
 	, CollectExternalDataDelegate(InCollectExternalDataDelegate)
 {
-	if (InStateTree.IsReadyToRun())
+	if (IsValid())
 	{
 		// Initialize data views for all possible items.
 		ContextAndExternalDataViews.SetNum(RootStateTree.GetNumContextDataViews());
-
-		InstanceDataStorage = &InstanceData.GetMutableStorage();
-		check(InstanceDataStorage);
-		
 		EventQueue = InstanceData.GetSharedMutableEventQueue();
-
 		bRecordTransitions = RecordTransitions == EStateTreeRecordTransitions::Yes;
 	}
 	else
@@ -464,8 +718,6 @@
 			const int32 TargetIndex = TargetDataDesc.Handle.DataHandle.GetIndex();
 			ContextAndExternalDataViews[TargetIndex] = InContextToCopy.ContextAndExternalDataViews[TargetIndex];
 		}
-		
-		EventQueue = InstanceData.GetSharedMutableEventQueue();
 	}
 	else
 	{
@@ -751,7 +1003,7 @@
 	// Initialize instance data. No active states yet, so we'll initialize the evals and global tasks.
 	InstanceData.Reset();
 
-	InstanceDataStorage->GetRuntimeValidation().SetContext(&Owner, &RootStateTree);
+	Storage.GetRuntimeValidation().SetContext(&Owner, &RootStateTree);
 	Exec.ExecutionExtension = MoveTemp(Parameters.ExecutionExtension);
 	if (Parameters.SharedEventQueue)
 	{
@@ -773,7 +1025,7 @@
 	// Initialize for the init frame.
 	ensure(Exec.ActiveFrames.Num() == 0);
 	FStateTreeExecutionFrame& InitFrame = Exec.ActiveFrames.AddDefaulted_GetRef();
-	InitFrame.FrameID = UE::StateTree::FActiveFrameID(InstanceDataStorage->GenerateUniqueId());
+	InitFrame.FrameID = UE::StateTree::FActiveFrameID(Storage.GenerateUniqueId());
 	InitFrame.StateTree = &RootStateTree;
 	InitFrame.RootState = FStateTreeStateHandle::Root;
 	InitFrame.ActiveStates = {};
@@ -1251,6 +1503,7 @@
 	}
 }
 
+// Deprecated
 bool FStateTreeExecutionContext::AddDelegateListener(const FStateTreeDelegateListener& Listener, FSimpleDelegate Delegate)
 {
 	BindDelegate(Listener, MoveTemp(Delegate));
@@ -1283,6 +1536,7 @@
 	GetExecState().DelegateActiveListeners.Add(Listener, MoveTemp(Delegate), CurrentFrame->FrameID, CurrentlyStateID, FStateTreeIndex16(CurrentNodeDataHandle.GetIndex()));
 }
 
+// Deprecated
 void FStateTreeExecutionContext::RemoveDelegateListener(const FStateTreeDelegateListener& Listener)
 {
 	UnbindDelegate(Listener);
@@ -1642,7 +1896,7 @@
 
 	TArray<FCompactStateTreeParameters, TFixedAllocator<FStateSelectionResult::MaxExecutionFrames>> TempParams;
 
-	TArrayView<FStateTreeTemporaryInstanceData> TempInstances = InstanceDataStorage->GetMutableTemporaryInstances();
+	TArrayView<FStateTreeTemporaryInstanceData> TempInstances = Storage.GetMutableTemporaryInstances();
 	auto FindInstanceTempData = [&TempInstances](const FStateTreeExecutionFrame& Frame, FStateTreeDataHandle DataHandle)
 	{
 		FStateTreeTemporaryInstanceData* TempData = TempInstances.FindByPredicate([&Frame, &DataHandle](const FStateTreeTemporaryInstanceData& Data)
@@ -1915,11 +2169,11 @@
 				return {};
 			}
 
-			return GetDataViewFromInstanceStorage(*InstanceDataStorage, CurrentlyProcessedSharedInstanceStorage, ParentFrame, CurrentFrame, Handle);
+			return GetDataViewFromInstanceStorage(Storage, CurrentlyProcessedSharedInstanceStorage, ParentFrame, CurrentFrame, Handle);
 		}
 
 	default:
-		return GetDataViewFromInstanceStorage(*InstanceDataStorage, CurrentlyProcessedSharedInstanceStorage, ParentFrame, CurrentFrame, Handle);
+		return GetDataViewFromInstanceStorage(Storage, CurrentlyProcessedSharedInstanceStorage, ParentFrame, CurrentFrame, Handle);
 	}
 }
 
@@ -2055,13 +2309,13 @@
 	case EStateTreeDataSourceType::GlobalInstanceData:
 	case EStateTreeDataSourceType::GlobalInstanceDataObject:
 		return CurrentFrame.GlobalInstanceIndexBase.IsValid()
-			&& InstanceDataStorage->IsValidIndex(CurrentFrame.GlobalInstanceIndexBase.Get() + Handle.GetIndex());
+			&& Storage.IsValidIndex(CurrentFrame.GlobalInstanceIndexBase.Get() + Handle.GetIndex());
 
 	case EStateTreeDataSourceType::ActiveInstanceData:
 	case EStateTreeDataSourceType::ActiveInstanceDataObject:
 		return CurrentFrame.ActiveInstanceIndexBase.IsValid()
 			&& CurrentFrame.ActiveStates.Contains(Handle.GetState(), CurrentFrame.NumCurrentlyActiveStates)
-			&& InstanceDataStorage->IsValidIndex(CurrentFrame.ActiveInstanceIndexBase.Get() + Handle.GetIndex());
+			&& Storage.IsValidIndex(CurrentFrame.ActiveInstanceIndexBase.Get() + Handle.GetIndex());
 		
 	case EStateTreeDataSourceType::SharedInstanceData:
 	case EStateTreeDataSourceType::SharedInstanceDataObject:
@@ -2093,12 +2347,12 @@
 		// Standalone subtree, params define as state params.
 		return CurrentFrame.ActiveInstanceIndexBase.IsValid()
 			&& CurrentFrame.ActiveStates.Contains(Handle.GetState(), CurrentFrame.NumCurrentlyActiveStates)
-			&& InstanceDataStorage->IsValidIndex(CurrentFrame.ActiveInstanceIndexBase.Get() + Handle.GetIndex());
+			&& Storage.IsValidIndex(CurrentFrame.ActiveInstanceIndexBase.Get() + Handle.GetIndex());
 
 	case EStateTreeDataSourceType::StateParameterData:
 		return CurrentFrame.ActiveInstanceIndexBase.IsValid()
 			&& CurrentFrame.ActiveStates.Contains(Handle.GetState(), CurrentFrame.NumCurrentlyActiveStates)
-			&& InstanceDataStorage->IsValidIndex(CurrentFrame.ActiveInstanceIndexBase.Get() + Handle.GetIndex());
+			&& Storage.IsValidIndex(CurrentFrame.ActiveInstanceIndexBase.Get() + Handle.GetIndex());
 
 	case EStateTreeDataSourceType::TransitionEvent:
 		return CurrentlyProcessedTransitionEvent != nullptr;
@@ -2107,7 +2361,7 @@
 		return CurrentlyProcessedStateSelectionEvents != nullptr
 			|| (CurrentFrame.ActiveInstanceIndexBase.IsValid()
 			&& CurrentFrame.ActiveStates.Contains(Handle.GetState())
-			&& InstanceDataStorage->IsValidIndex(CurrentFrame.ActiveInstanceIndexBase.Get() + Handle.GetIndex()));
+			&& Storage.IsValidIndex(CurrentFrame.ActiveInstanceIndexBase.Get() + Handle.GetIndex()));
 
 	default:
 		checkf(false, TEXT("Unhandle case %s"), *UEnum::GetValueAsString(Handle.GetSource()));
@@ -2122,23 +2376,21 @@
 	{
 		return GetDataView(ParentFrame, CurrentFrame, Handle);
 	}
-
-	check(InstanceDataStorage);
 	
 	switch (Handle.GetSource())
 	{
 	case EStateTreeDataSourceType::GlobalInstanceData:
 	case EStateTreeDataSourceType::ActiveInstanceData:
-		return InstanceDataStorage->GetMutableTemporaryStruct(CurrentFrame, Handle);
+		return Storage.GetMutableTemporaryStruct(CurrentFrame, Handle);
 		
 	case EStateTreeDataSourceType::GlobalInstanceDataObject:
 	case EStateTreeDataSourceType::ActiveInstanceDataObject:
-		return InstanceDataStorage->GetMutableTemporaryObject(CurrentFrame, Handle);
+		return Storage.GetMutableTemporaryObject(CurrentFrame, Handle);
 
 	case EStateTreeDataSourceType::GlobalParameterData:
 		if (ParentFrame)
 		{
-			if (FCompactStateTreeParameters* Params = InstanceDataStorage->GetMutableTemporaryStruct(*ParentFrame, CurrentFrame.GlobalParameterDataHandle).GetPtr<FCompactStateTreeParameters>())
+			if (FCompactStateTreeParameters* Params = Storage.GetMutableTemporaryStruct(*ParentFrame, CurrentFrame.GlobalParameterDataHandle).GetPtr<FCompactStateTreeParameters>())
 			{
 				return Params->Parameters.GetMutableValue();
 			}
@@ -2149,20 +2401,20 @@
 		if (ParentFrame)
 		{
 			// Linked subtree, params defined in parent scope.
-			if (FCompactStateTreeParameters* Params = InstanceDataStorage->GetMutableTemporaryStruct(*ParentFrame, CurrentFrame.StateParameterDataHandle).GetPtr<FCompactStateTreeParameters>())
+			if (FCompactStateTreeParameters* Params = Storage.GetMutableTemporaryStruct(*ParentFrame, CurrentFrame.StateParameterDataHandle).GetPtr<FCompactStateTreeParameters>())
 			{
 				return Params->Parameters.GetMutableValue();
 			}
 		}
 		// Standalone subtree, params define as state params.
-		if (FCompactStateTreeParameters* Params = InstanceDataStorage->GetMutableTemporaryStruct(CurrentFrame, Handle).GetPtr<FCompactStateTreeParameters>())
+		if (FCompactStateTreeParameters* Params = Storage.GetMutableTemporaryStruct(CurrentFrame, Handle).GetPtr<FCompactStateTreeParameters>())
 		{
 			return Params->Parameters.GetMutableValue();
 		}
 		break;
 
 	case EStateTreeDataSourceType::StateParameterData:
-		if (FCompactStateTreeParameters* Params = InstanceDataStorage->GetMutableTemporaryStruct(CurrentFrame, Handle).GetPtr<FCompactStateTreeParameters>())
+		if (FCompactStateTreeParameters* Params = Storage.GetMutableTemporaryStruct(CurrentFrame, Handle).GetPtr<FCompactStateTreeParameters>())
 		{
 			return Params->Parameters.GetMutableValue();
 		}
@@ -2177,8 +2429,7 @@
 
 FStateTreeDataView FStateTreeExecutionContext::AddTemporaryInstance(const FStateTreeExecutionFrame& Frame, const FStateTreeIndex16 OwnerNodeIndex, const FStateTreeDataHandle DataHandle, FConstStructView NewInstanceData)
 {
-	check(InstanceDataStorage);
-	const FStructView NewInstance = InstanceDataStorage->AddTemporaryInstance(Owner, Frame, OwnerNodeIndex, DataHandle, NewInstanceData);
+	const FStructView NewInstance = Storage.AddTemporaryInstance(Owner, Frame, OwnerNodeIndex, DataHandle, NewInstanceData);
 	if (FStateTreeInstanceObjectWrapper* Wrapper = NewInstance.GetPtr<FStateTreeInstanceObjectWrapper>())
 	{
 		return FStateTreeDataView(Wrapper->InstanceObject);
@@ -2350,7 +2601,7 @@
 	if (ensureMsgf(RootStateTree.GetDefaultParameters().HasSameLayout(Parameters),
 		TEXT("Parameters must be of the same struct type. Make sure to migrate the provided parameters to the same type as the StateTree default parameters.")))
 	{
-		InstanceDataStorage->SetGlobalParameters(Parameters);
+		Storage.SetGlobalParameters(Parameters);
 		return true;
 	}
 
@@ -2401,7 +2652,7 @@
 			{
 				if (State->EventDataIndex.IsValid())
 				{
-					FStateTreeSharedEvent& StateTreeEvent = InstanceDataStorage->GetMutableStruct(NewFrame.ActiveInstanceIndexBase.Get() + State->EventDataIndex.Get()).Get<FStateTreeSharedEvent>();
+					FStateTreeSharedEvent& StateTreeEvent = Storage.GetMutableStruct(NewFrame.ActiveInstanceIndexBase.Get() + State->EventDataIndex.Get()).Get<FStateTreeSharedEvent>();
 					
 					const FStateTreeSharedEvent& EventToCapture = FramesStateSelectionEvents[FrameIndex].Events[State->Depth];
 					if (EventsInUse.Contains(EventToCapture))
@@ -3542,7 +3793,7 @@
 	Transition.TargetState = FStateTreeStateHandle::FromCompletionStatus(CompletionStatus);
 	Transition.CurrentRunStatus = CompletionStatus;
 
-	TArrayView<FStateTreeTemporaryInstanceData> TempInstances = InstanceDataStorage->GetMutableTemporaryInstances();
+	TArrayView<FStateTreeTemporaryInstanceData> TempInstances = Storage.GetMutableTemporaryInstances();
 	for (int32 Index = TempInstances.Num() - 1; Index >= 0; Index--)
 	{
 		FStateTreeTemporaryInstanceData& TempInstance = TempInstances[Index];
@@ -4325,7 +4576,7 @@
 	{
 		// Re-cache bHasEvents, RequestTransition above can create new events.
 		const bool bHasEvents = EventQueue && EventQueue->HasEvents();
-		const bool bHasBroadcastedDelegates = InstanceDataStorage->HasBroadcastedDelegates();
+		const bool bHasBroadcastedDelegates = Storage.HasBroadcastedDelegates();
 
 		// Transition() can TriggerTransitions() in a loop when a sub-frame completes.
 		//We do not want to evaluate the transition from that sub-frame.
@@ -4529,7 +4780,7 @@
 				}
 				else if (EnumHasAnyFlags(Transition.Trigger, EStateTreeTransitionTrigger::OnDelegate))
 				{
-					if (InstanceDataStorage->IsDelegateBroadcasted(Transition.RequiredDelegateDispatcher))
+					if (Storage.IsDelegateBroadcasted(Transition.RequiredDelegateDispatcher))
 					{
 						// Dummy event to make sure we iterate to loop below once.
 						TransitionEvents.Emplace(nullptr);
@@ -4623,12 +4874,12 @@
 	}
 
 	// All events have had the change to be reacted to, clear the event queue (if this instance owns it).
-	if (InstanceData.IsOwningEventQueue())
+	if (InstanceData.IsOwningEventQueue() && EventQueue)
 	{
 		EventQueue->Reset();
 	}
 
-	InstanceDataStorage->ResetBroadcastedDelegates();
+	Storage.ResetBroadcastedDelegates();
 
 	//
 	// Check state completion transitions.
@@ -5127,7 +5378,7 @@
 	}
 
 	// Activate/Push a new state
-	if (!CurrentFrame.ActiveStates.Push(NextStateHandle, UE::StateTree::FActiveStateID(InstanceDataStorage->GenerateUniqueId())))
+	if (!CurrentFrame.ActiveStates.Push(NextStateHandle, UE::StateTree::FActiveStateID(Storage.GenerateUniqueId())))
 	{
 		STATETREE_LOG(Error, TEXT("%hs: Reached max execution depth when trying to select state %s from '%s'.  '%s' using StateTree '%s'."),
 			__FUNCTION__, *GetSafeStateName(CurrentFrame, NextStateHandle), *GetStateStatusString(Exec), *GetNameSafe(&Owner), *GetFullNameSafe(CurrentFrame.StateTree));
@@ -5193,7 +5444,7 @@
 				FStateTreeExecutionFrame NewFrame;
 				NewFrame.StateTree = CurrentFrame.StateTree;
 				NewFrame.RootState = NextState.LinkedState;
-				NewFrame.FrameID = UE::StateTree::FActiveFrameID(InstanceDataStorage->GenerateUniqueId());
+				NewFrame.FrameID = UE::StateTree::FActiveFrameID(Storage.GenerateUniqueId());
 				NewFrame.ExternalDataBaseIndex = CurrentFrame.ExternalDataBaseIndex;
 
 				// Check and prevent recursion.
@@ -5271,7 +5522,7 @@
 			}
 				
 			FStateTreeExecutionFrame NewFrame;
-			NewFrame.FrameID = UE::StateTree::FActiveFrameID(InstanceDataStorage->GenerateUniqueId());
+			NewFrame.FrameID = UE::StateTree::FActiveFrameID(Storage.GenerateUniqueId());
 			NewFrame.StateTree = NextLinkedStateAsset;
 			NewFrame.RootState = FStateTreeStateHandle::Root;
 			NewFrame.bIsGlobalFrame = true;
@@ -5438,7 +5689,7 @@
 				}
 				else if (EnumHasAnyFlags(Transition.Trigger, EStateTreeTransitionTrigger::OnDelegate))
 				{
-					if (InstanceDataStorage->IsDelegateBroadcasted(Transition.RequiredDelegateDispatcher))
+					if (Storage.IsDelegateBroadcasted(Transition.RequiredDelegateDispatcher))
 					{
 						SelectedStateTransitionEvents.Emplace();
 					}
@@ -5757,246 +6008,11 @@
 	return GetSafeStateName(ExecState.ActiveFrames.Last(), ExecState.ActiveFrames.Last().ActiveStates.Last()) + TEXT(":") + UEnum::GetDisplayValueAsText(ExecState.LastTickStatus).ToString();
 }
 
-EStateTreeRunStatus FStateTreeExecutionContext::GetLastTickStatus() const
-{
-	if (!IsValid())
-	{
-		STATETREE_LOG(Warning, TEXT("%hs: StateTree context is not initialized properly ('%s' using StateTree '%s')"),
-			__FUNCTION__, *GetNameSafe(&Owner), *GetFullNameSafe(&RootStateTree));
-		return EStateTreeRunStatus::Failed;
-	}
-
-	const FStateTreeExecutionState& Exec = GetExecState();
-	return Exec.LastTickStatus;
-}
-
+// Deprecated
 FString FStateTreeExecutionContext::GetInstanceDescription() const
 {
 	return GetInstanceDescriptionInternal();
 }
-
-TConstArrayView<FStateTreeExecutionFrame> FStateTreeExecutionContext::GetActiveFrames() const
-{
-	if (!IsValid())
-	{
-		STATETREE_LOG(Warning, TEXT("%hs: StateTree context is not initialized properly ('%s' using StateTree '%s')"),
-			__FUNCTION__, *GetNameSafe(&Owner), *GetFullNameSafe(&RootStateTree));
-		return TConstArrayView<FStateTreeExecutionFrame>();
-	}
-
-	const FStateTreeExecutionState& Exec = GetExecState();
-	return Exec.ActiveFrames;
-}
-
-
-#if WITH_GAMEPLAY_DEBUGGER
-
-FString FStateTreeExecutionContext::GetDebugInfoString() const
-{
-	TStringBuilder<2048> DebugString;
-	DebugString << TEXT("StateTree (asset: '");
-	RootStateTree.GetFullName(DebugString);
-	DebugString << TEXT("')");
-
-	if (IsValid())
-	{
-		const FStateTreeExecutionState& Exec = GetExecState();
-
-		DebugString << TEXT("Status: ");
-		DebugString << UEnum::GetDisplayValueAsText(Exec.TreeRunStatus).ToString();
-		DebugString << TEXT("\n");
-
-		// Active States
-		DebugString << TEXT("Current State:\n");
-		for (const FStateTreeExecutionFrame& CurrentFrame : Exec.ActiveFrames)
-		{
-			const UStateTree* CurrentStateTree = CurrentFrame.StateTree;
-
-			if (CurrentFrame.bIsGlobalFrame)
-			{
-				DebugString.Appendf(TEXT("\nEvaluators\n  [ %-30s | %8s | %15s ]\n"),
-					TEXT("Name"), TEXT("Bindings"), TEXT("Data Handle"));
-				for (int32 EvalIndex = CurrentStateTree->EvaluatorsBegin; EvalIndex < (CurrentStateTree->EvaluatorsBegin + CurrentStateTree->EvaluatorsNum); EvalIndex++)
-				{
-					const FStateTreeEvaluatorBase& Eval = CurrentStateTree->Nodes[EvalIndex].Get<const FStateTreeEvaluatorBase>();
-					DebugString.Appendf(TEXT("| %-30s | %8d | %15s |\n"),
-						*Eval.Name.ToString(), Eval.BindingsBatch.Get(), *Eval.InstanceDataHandle.Describe());
-				}
-
-				DebugString << TEXT("\nGlobal Tasks\n");
-				for (int32 TaskIndex = CurrentStateTree->GlobalTasksBegin; TaskIndex < (CurrentStateTree->GlobalTasksBegin + CurrentStateTree->GlobalTasksNum); TaskIndex++)
-				{
-					const FStateTreeTaskBase& Task = CurrentStateTree->Nodes[TaskIndex].Get<const FStateTreeTaskBase>();
-					if (Task.bTaskEnabled)
-					{
-						FString TempString;
-						Task.AppendDebugInfoString(TempString, *this);
-						DebugString << TempString;
-
-					}
-				}
-			}
-
-			for (int32 Index = 0; Index < CurrentFrame.ActiveStates.Num(); Index++)
-			{
-				FStateTreeStateHandle Handle = CurrentFrame.ActiveStates[Index];
-				if (Handle.IsValid())
-				{
-					const FCompactStateTreeState& State = RootStateTree.States[Handle.Index];
-					DebugString << TEXT('[');
-					DebugString << State.Name;
-					DebugString << TEXT("]\n");
-
-					if (State.TasksNum > 0)
-					{
-						DebugString += TEXT("\nTasks:\n");
-						for (int32 TaskIndex = State.TasksBegin; TaskIndex < (State.TasksBegin + State.TasksNum); TaskIndex++)
-						{
-							const FStateTreeTaskBase& Task = RootStateTree.Nodes[TaskIndex].Get<const FStateTreeTaskBase>();
-							if (Task.bTaskEnabled)
-							{
-								FString TempString;
-								Task.AppendDebugInfoString(TempString, *this);
-								DebugString << TempString;
-							}
-						}
-					}
-				}
-			}
-		}
-	}
-	else
-	{
-		DebugString << TEXT("StateTree context is not initialized properly.");
-	}
-
-	return DebugString.ToString();
-}
-#endif // WITH_GAMEPLAY_DEBUGGER
-
-#if WITH_STATETREE_DEBUG
-void FStateTreeExecutionContext::DebugPrintInternalLayout()
-{
-	LOG_SCOPE_VERBOSITY_OVERRIDE(LogStateTree, ELogVerbosity::Log);
-	UE_LOG(LogStateTree, Log, TEXT("%s"), *RootStateTree.DebugInternalLayoutAsString());
-}
-
-int32 FStateTreeExecutionContext::GetStateChangeCount() const
-{
-	if (!IsValid())
-	{
-		STATETREE_LOG(Warning, TEXT("%hs: StateTree context is not initialized properly ('%s' using StateTree '%s')"),
-			__FUNCTION__, *GetNameSafe(&Owner), *GetFullNameSafe(&RootStateTree));
-		return 0;
-	}
-
-	const FStateTreeExecutionState& Exec = GetExecState();
-	return Exec.StateChangeCount;
-}
-
-#endif // WITH_STATETREE_DEBUG
-
-FString FStateTreeExecutionContext::GetActiveStateName() const
-{
-	if (!IsValid())
-	{
-		STATETREE_LOG(Warning, TEXT("%hs: StateTree context is not initialized properly ('%s' using StateTree '%s')"),
-			__FUNCTION__, *GetNameSafe(&Owner), *GetFullNameSafe(&RootStateTree));
-		return FString();
-	}
-
-	const FStateTreeExecutionState& Exec = GetExecState();
-
-	TStringBuilder<1024> FullStateName;
-	
-	const UStateTree* LastStateTree = &RootStateTree;
-	int32 Indent = 0;
-
-	for (int32 FrameIndex = 0; FrameIndex < Exec.ActiveFrames.Num(); FrameIndex++)
-	{
-		const FStateTreeExecutionFrame& CurrentFrame = Exec.ActiveFrames[FrameIndex];
-		const UStateTree* CurrentStateTree = CurrentFrame.StateTree;
-
-		// Append linked state marker at the end of the previous line.
-		if (Indent > 0)
-		{
-			FullStateName << TEXT(" >");
-		}
-		// If tree has changed, append that too.
-		if (CurrentFrame.StateTree != LastStateTree)
-		{
-			FullStateName << TEXT(" [");
-			FullStateName << CurrentFrame.StateTree.GetFName();
-			FullStateName << TEXT(']');
-
-			LastStateTree = CurrentFrame.StateTree;
-		}
-
-		for (int32 Index = 0; Index < CurrentFrame.ActiveStates.Num(); Index++)
-		{
-			const FStateTreeStateHandle Handle = CurrentFrame.ActiveStates[Index];
-			if (Handle.IsValid())
-			{
-				const FCompactStateTreeState& State = CurrentStateTree->States[Handle.Index];
-				if (Indent > 0)
-				{
-					FullStateName += TEXT("\n");
-				}
-				FullStateName.Appendf(TEXT("%*s-"), Indent * 3, TEXT("")); // Indent
-				FullStateName << State.Name;
-				Indent++;
-			}
-		}
-	}
-
-	switch (Exec.TreeRunStatus)
-	{
-	case EStateTreeRunStatus::Failed:
-		FullStateName << TEXT(" FAILED\n");
-		break;
-	case EStateTreeRunStatus::Succeeded:
-		FullStateName << TEXT(" SUCCEEDED\n");
-		break;
-	case EStateTreeRunStatus::Running:
-		// Empty
-		break;
-	default:
-		FullStateName << TEXT("--\n");
-	}
-
-	return FullStateName.ToString();
-}
-
-TArray<FName> FStateTreeExecutionContext::GetActiveStateNames() const
-{
-	if (!IsValid())
-	{
-		STATETREE_LOG(Warning, TEXT("%hs: StateTree context is not initialized properly ('%s' using StateTree '%s')"),
-			__FUNCTION__, *GetNameSafe(&Owner), *GetFullNameSafe(&RootStateTree));
-		return TArray<FName>();
-	}
-
-	TArray<FName> Result;
-	const FStateTreeExecutionState& Exec = GetExecState();
-
-	// Active States
-	for (const FStateTreeExecutionFrame& CurrentFrame : Exec.ActiveFrames)
-	{
-		const UStateTree* CurrentStateTree = CurrentFrame.StateTree;
-		for (int32 Index = 0; Index < CurrentFrame.ActiveStates.Num(); Index++)
-		{
-			const FStateTreeStateHandle Handle = CurrentFrame.ActiveStates[Index];
-			if (Handle.IsValid())
-			{
-				const FCompactStateTreeState& State = CurrentStateTree->States[Handle.Index];
-				Result.Add(State.Name);
-			}
-		}
-	}
-
-	return Result;
-}
-
 FStateTreeWeakExecutionContext::FStateTreeWeakExecutionContext(const FStateTreeExecutionContext& Context)
 	: Owner(Context.GetOwner())
 	, StateTree(Context.GetStateTree())
diff --git a/Engine/Plugins/Runtime/StateTree/Source/StateTreeModule/Private/StateTreeInstanceData.cpp b/Engine/Plugins/Runtime/StateTree/Source/StateTreeModule/Private/StateTreeInstanceData.cpp
--- a/Engine/Plugins/Runtime/StateTree/Source/StateTreeModule/Private/StateTreeInstanceData.cpp
+++ b/Engine/Plugins/Runtime/StateTree/Source/StateTreeModule/Private/StateTreeInstanceData.cpp
@@ -365,6 +365,22 @@ void FStateTreeInstanceStorage::Reset()
 #endif
 }
 
+void FStateTreeInstanceStorage::AcquireReadAccess()
+{
+	if (UE::StateTree::Debug::Private::RuntimeValidationThreadingAccessDetector)
+	{
+		UE_MT_ACQUIRE_READ_ACCESS(AccessDetector)
+	}
+}
+
+void FStateTreeInstanceStorage::ReleaseReadAccess()
+{
+	if (UE::StateTree::Debug::Private::RuntimeValidationThreadingAccessDetector)
+	{
+		UE_MT_RELEASE_READ_ACCESS(AccessDetector)
+	}
+}
+
 void FStateTreeInstanceStorage::AcquireWriteAccess()
 {
 	if (UE::StateTree::Debug::Private::RuntimeValidationThreadingAccessDetector)
diff --git a/Engine/Plugins/Runtime/StateTree/Source/StateTreeModule/Private/StateTreeTaskBase.cpp b/Engine/Plugins/Runtime/StateTree/Source/StateTreeModule/Private/StateTreeTaskBase.cpp
--- a/Engine/Plugins/Runtime/StateTree/Source/StateTreeModule/Private/StateTreeTaskBase.cpp
+++ b/Engine/Plugins/Runtime/StateTree/Source/StateTreeModule/Private/StateTreeTaskBase.cpp
@@ -6,6 +6,16 @@
 
 
 #if WITH_GAMEPLAY_DEBUGGER
+FString FStateTreeTaskBase::GetDebugInfo(const FStateTreeReadOnlyExecutionContext& Context) const
+{
+	TStringBuilder<256> Builder;
+	Builder << TEXT('[');
+	Builder << Name;
+	Builder << TEXT("]\n");
+	return Builder.ToString();
+}
+
+//Deprecated
 void FStateTreeTaskBase::AppendDebugInfoString(FString& DebugString, const FStateTreeExecutionContext& Context) const
 {
 	DebugString += FString::Printf(TEXT("[%s]\n"), *Name.ToString());
diff --git a/Engine/Plugins/Runtime/StateTree/Source/StateTreeModule/Public/StateTree.h b/Engine/Plugins/Runtime/StateTree/Source/StateTreeModule/Public/StateTree.h
--- a/Engine/Plugins/Runtime/StateTree/Source/StateTreeModule/Public/StateTree.h
+++ b/Engine/Plugins/Runtime/StateTree/Source/StateTreeModule/Public/StateTree.h
@@ -389,6 +389,7 @@ private:
 	friend struct FStateTreeInstance;
 	friend struct FStateTreeExecutionContext;
 	friend struct FStateTreeMinimalExecutionContext;
+	friend struct FStateTreeReadOnlyExecutionContext;
 	friend struct FStateTreeWeakExecutionContext;
 #if WITH_EDITORONLY_DATA
 	friend struct FStateTreeCompiler;
diff --git a/Engine/Plugins/Runtime/StateTree/Source/StateTreeModule/Public/StateTreeEvaluatorBase.h b/Engine/Plugins/Runtime/StateTree/Source/StateTreeModule/Public/StateTreeEvaluatorBase.h
--- a/Engine/Plugins/Runtime/StateTree/Source/StateTreeModule/Public/StateTreeEvaluatorBase.h
+++ b/Engine/Plugins/Runtime/StateTree/Source/StateTreeModule/Public/StateTreeEvaluatorBase.h
@@ -6,6 +6,7 @@
 #include "StateTreeEvaluatorBase.generated.h"
 
 struct FStateTreeExecutionContext;
+struct FStateTreeReadOnlyExecutionContext;
 
 /**
  * Base struct of StateTree Evaluators.
@@ -36,6 +37,9 @@ struct STATETREEMODULE_API FStateTreeEvaluatorBase : public FStateTreeNodeBase
 	virtual void Tick(FStateTreeExecutionContext& Context, const float DeltaTime) const {}
 
 #if WITH_GAMEPLAY_DEBUGGER
+	virtual FString GetDebugInfo(const FStateTreeReadOnlyExecutionContext& Context) const;
+
+	UE_DEPRECATED(5.6, "Use the version with the FStateTreeReadOnlyExecutionContext.")
 	virtual void AppendDebugInfoString(FString& DebugString, const FStateTreeExecutionContext& Context) const;
 #endif // WITH_GAMEPLAY_DEBUGGER
 };
diff --git a/Engine/Plugins/Runtime/StateTree/Source/StateTreeModule/Public/StateTreeExecutionContext.h b/Engine/Plugins/Runtime/StateTree/Source/StateTreeModule/Public/StateTreeExecutionContext.h
--- a/Engine/Plugins/Runtime/StateTree/Source/StateTreeModule/Public/StateTreeExecutionContext.h
+++ b/Engine/Plugins/Runtime/StateTree/Source/StateTreeModule/Public/StateTreeExecutionContext.h
@@ -40,20 +40,20 @@
 DECLARE_DELEGATE_RetVal_FourParams(bool, FOnCollectStateTreeExternalData, const FStateTreeExecutionContext& /*Context*/, const UStateTree* /*StateTree*/, TArrayView<const FStateTreeExternalDataDesc> /*ExternalDataDescs*/, TArrayView<FStateTreeDataView> /*OutDataViews*/);
 
 /**
- * Minimal execution context to interact with the state tree instance data.
- * A regular execution context requires the context data and external data to be valid to execute all possible operations.
- * The minimal execution context doesn't requires those data but supports only a subset of operations. GetStateTreeRunStatus, SendEvent.
+ * Read-only execution context to interact with the state tree instance data. Only const and read accesses are available.
+ * Multiple FStateTreeReadOnlyExecutionContext can coexist on different threads as long no other (minimal, weak, regular) execution context exists.
+ * The user is responsible for preventing invalid multithread access.
  */
-struct STATETREEMODULE_API FStateTreeMinimalExecutionContext
+struct STATETREEMODULE_API FStateTreeReadOnlyExecutionContext
 {
 public:
-	explicit FStateTreeMinimalExecutionContext(UObject& Owner, const UStateTree& StateTree, FStateTreeInstanceData& InInstanceData);
-	explicit FStateTreeMinimalExecutionContext(UObject& Owner, const UStateTree& StateTree, FStateTreeInstanceStorage& Storage);
-	virtual ~FStateTreeMinimalExecutionContext();
+	explicit FStateTreeReadOnlyExecutionContext(TNotNull<UObject*> Owner, TNotNull<const UStateTree*> StateTree, FStateTreeInstanceData& InInstanceData);
+	explicit FStateTreeReadOnlyExecutionContext(TNotNull<UObject*> Owner, TNotNull<const UStateTree*> StateTree, FStateTreeInstanceStorage& Storage);
+	virtual ~FStateTreeReadOnlyExecutionContext();
 
 private:
-	FStateTreeMinimalExecutionContext(const FStateTreeMinimalExecutionContext&) = delete;
-	FStateTreeMinimalExecutionContext& operator=(const FStateTreeMinimalExecutionContext&) = delete;
+	FStateTreeReadOnlyExecutionContext(const FStateTreeReadOnlyExecutionContext&) = delete;
+	FStateTreeReadOnlyExecutionContext& operator=(const FStateTreeReadOnlyExecutionContext&) = delete;
 
 public:
 	/**
@@ -84,12 +84,84 @@
 		return &RootStateTree;
 	}
 
-	/** @return the tree run status. */
-	EStateTreeRunStatus GetStateTreeRunStatus() const;
+	/** @return true if there is a pending event with specified tag. */
+	bool HasEventToProcess(const FGameplayTag Tag) const
+	{
+		return Storage.GetEventQueue().GetEventsView().ContainsByPredicate([Tag](const FStateTreeSharedEvent& Event)
+			{
+				check(Event.IsValid());
+				return Event->Tag.MatchesTag(Tag);
+			});
+	}
 
+	/** @return Pointer to a State or null if state not found */
+	const FCompactStateTreeState* GetStateFromHandle(const FStateTreeStateHandle StateHandle) const
+	{
+		return RootStateTree.GetStateFromHandle(StateHandle);
+	}
+
 	/** @return the delta time for the next execution context tick. */
 	FStateTreeScheduledTick GetNextScheduledTick() const;
 
+	/** @return the tree run status. */
+	EStateTreeRunStatus GetStateTreeRunStatus() const;
+
+	/** @return the status of the last tick function */
+	EStateTreeRunStatus GetLastTickStatus() const;
+
+	/** @return reference to the list of currently active frames and states. */
+	TConstArrayView<FStateTreeExecutionFrame> GetActiveFrames() const;
+
+	/** @return the name of the active state. */
+	FString GetActiveStateName() const;
+
+	/** @return the names of all the active state. */
+	TArray<FName> GetActiveStateNames() const;
+
+#if WITH_GAMEPLAY_DEBUGGER
+	/** @return Debug string describing the current state of the execution */
+	FString GetDebugInfoString() const;
+#endif // WITH_GAMEPLAY_DEBUGGER
+
+#if WITH_STATETREE_DEBUG
+	int32 GetStateChangeCount() const;
+
+	void DebugPrintInternalLayout();
+#endif
+
+protected:
+	/** @return Prefix that will be used by STATETREE_LOG and STATETREE_CLOG, Owner name by default. */
+	FString GetInstanceDescriptionInternal() const;
+
+#if WITH_STATETREE_TRACE
+	/** A unique Id used by debugging tools to identify the instance. */
+	FStateTreeInstanceDebugId GetInstanceDebugId() const;
+#endif
+
+protected:
+	/** Owner of the instance data. */
+	UObject& Owner;
+
+	/** The StateTree asset the context is initialized for */
+	const UStateTree& RootStateTree;
+
+	/** Data storage of the instance data. */
+	FStateTreeInstanceStorage& Storage;
+};
+
+/**
+ * Minimal execution context to interact with the state tree instance data.
+ * A regular execution context requires the context data and external data to be valid to execute all possible operations.
+ * The minimal execution context doesn't requires those data but supports only a subset of operations.
+ */
+struct STATETREEMODULE_API FStateTreeMinimalExecutionContext : public FStateTreeReadOnlyExecutionContext
+{
+public:
+	explicit FStateTreeMinimalExecutionContext(UObject& Owner, const UStateTree& StateTree, FStateTreeInstanceData& InInstanceData);
+	explicit FStateTreeMinimalExecutionContext(UObject& Owner, const UStateTree& StateTree, FStateTreeInstanceStorage& Storage);
+	virtual ~FStateTreeMinimalExecutionContext();
+
+public:
 	/** 
 	 * Adds a scheduled tick request.
 	 * The result of GetNextScheduledTick is affected by the request.
@@ -117,13 +189,7 @@
 #endif //WITH_STATETREE_TRACE
 
 protected:
-	/** @return Prefix that will be used by STATETREE_LOG and STATETREE_CLOG, Owner name by default. */
-	FString GetInstanceDescriptionInternal() const;
-
 #if WITH_STATETREE_TRACE
-	/** A unique Id used by debugging tools to identify the instance. */
-	FStateTreeInstanceDebugId GetInstanceDebugId() const;
-
 	UE::StateTreeTrace::FNodeCustomDebugData StealNodeCustomDebugTraceData()
 	{
 		return MoveTemp(NodeCustomDebugTraceData);
@@ -136,15 +202,6 @@
 	void ScheduleNextTick();
 
 protected:
-	/** Owner of the instance data. */
-	UObject& Owner;
-
-	/** The StateTree asset the context is initialized for */
-	const UStateTree& RootStateTree;
-
-	/** Data storage of the instance data. */
-	FStateTreeInstanceStorage& Storage;
-
 	/** The context is processing the tree. We do not need to inform the owner that something changed. */
 	bool bAllowedToScheduleNextTick = true;
 };
@@ -232,16 +289,28 @@
 	const FStateTreeReference* GetLinkedStateTreeOverrideForTag(const FGameplayTag StateTag) const;
 
 	/** @return const references to the instance data in use, or nullptr if the context is not valid. */
-	const FStateTreeInstanceData* GetInstanceData() const { return &InstanceData; }
+	const FStateTreeInstanceData* GetInstanceData() const
+	{
+		return &InstanceData;
+	}
 
 	/** @return mutable references to the instance data in use, or nullptr if the context is not valid. */
-	FStateTreeInstanceData* GetMutableInstanceData() const { return &InstanceData; }
+	FStateTreeInstanceData* GetMutableInstanceData() const
+	{
+		return &InstanceData;
+	}
 
 	/** @return mutable references to the instance data in use. */
-	const FStateTreeEventQueue& GetEventQueue() const { return InstanceData.GetEventQueue(); }
+	const FStateTreeEventQueue& GetEventQueue() const
+	{
+		return InstanceData.GetEventQueue();
+	}
 
 	/** @return mutable references to the instance data in use. */
-	FStateTreeEventQueue& GetMutableEventQueue() const { return InstanceData.GetMutableEventQueue(); }
+	FStateTreeEventQueue& GetMutableEventQueue() const
+	{
+		return InstanceData.GetMutableEventQueue();
+	}
 
 	/** @return a weak context to interact with the state tree instance data that can be stored for later uses. */
 	FStateTreeWeakExecutionContext MakeWeakExecutionContext() const;
@@ -332,29 +401,6 @@
 	 */
 	EStateTreeRunStatus TickTriggerTransitions();
 
-	/** @return the status of the last tick function */
-	EStateTreeRunStatus GetLastTickStatus() const;
-
-	/** @return reference to the list of currently active frames and states. */
-	TConstArrayView<FStateTreeExecutionFrame> GetActiveFrames() const;
-
-#if WITH_GAMEPLAY_DEBUGGER
-	/** @return Debug string describing the current state of the execution */
-	FString GetDebugInfoString() const;
-#endif // WITH_GAMEPLAY_DEBUGGER
-
-#if WITH_STATETREE_DEBUG
-	int32 GetStateChangeCount() const;
-
-	void DebugPrintInternalLayout();
-#endif
-
-	/** @return the name of the active state. */
-	FString GetActiveStateName() const;
-	
-	/** @return the names of all the active state. */
-	TArray<FName> GetActiveStateNames() const;
-
 	/**
 	 * Broadcasts the delegate.
 	 * It executes bound delegates immediately and triggers bound transitions (when transitions are evaluated).
@@ -435,16 +481,6 @@
 	UE_DEPRECATED(5.5, "Use GetEventsToProcessView() instead.")
 	TConstArrayView<FStateTreeEvent> GetEventsToProcess() const { return {}; }
 
-	/** @return true if there is a pending event with specified tag. */
-	bool HasEventToProcess(const FGameplayTag Tag) const
-	{
-		return GetEventsToProcessView().ContainsByPredicate([Tag](const FStateTreeSharedEvent& Event)
-		{
-			check(Event.IsValid());
-			return Event->Tag.MatchesTag(Tag);
-		});
-	}
-
 	/** @return the currently processed node if applicable. */
 	FStateTreeDataHandle GetCurrentlyProcessedNode() const
 	{
@@ -469,12 +505,6 @@
 		return CurrentlyProcessedParentFrame;
 	}
 
-	/** @return Pointer to a State or null if state not found */ 
-	const FCompactStateTreeState* GetStateFromHandle(const FStateTreeStateHandle StateHandle) const
-	{
-		return RootStateTree.GetStateFromHandle(StateHandle);
-	}
-
 	/** @return Array view to named external data descriptors associated with this context. Note: Init() must be called before calling this method. */
 	TConstArrayView<FStateTreeExternalDataDesc> GetContextDataDescs() const
 	{
@@ -897,15 +927,13 @@
 	/** @return StateTree execution state from the instance storage. */
 	FStateTreeExecutionState& GetExecState()
 	{
-		check(InstanceDataStorage);
-		return InstanceDataStorage->GetMutableExecutionState();
+		return Storage.GetMutableExecutionState();
 	}
 
 	/** @return const StateTree execution state from the instance storage. */
 	const FStateTreeExecutionState& GetExecState() const
 	{
-		check(InstanceDataStorage);
-		return InstanceDataStorage->GetExecutionState();
+		return Storage.GetExecutionState();
 	}
 
 	/** Updates the update phase of the statetree execution state. */
@@ -973,6 +1001,7 @@
 	/** Instance data used during current tick. */
 	FStateTreeInstanceData& InstanceData;
 
+	UE_DEPRECATED(5.6, "Use Storage instead.")
 	/** Data storage of the instance data, cached for less indirections. */
 	FStateTreeInstanceStorage* InstanceDataStorage = nullptr;
 
diff --git a/Engine/Plugins/Runtime/StateTree/Source/StateTreeModule/Public/StateTreeInstanceData.h b/Engine/Plugins/Runtime/StateTree/Source/StateTreeModule/Public/StateTreeInstanceData.h
--- a/Engine/Plugins/Runtime/StateTree/Source/StateTreeModule/Public/StateTreeInstanceData.h
+++ b/Engine/Plugins/Runtime/StateTree/Source/StateTreeModule/Public/StateTreeInstanceData.h
@@ -281,10 +281,16 @@ struct STATETREEMODULE_API FStateTreeInstanceStorage
 	/** Resets the storage to initial state. */
 	void Reset();
 
-	/** Start the invalid multithreading access detection. */
+	/** Start the invalid multithreading read-only access detection. */
+	void AcquireReadAccess();
+
+	/** Stop the multitheading read-only access detection. */
+	void ReleaseReadAccess();
+
+	/** Start the invalid multithreading write access detection. */
 	void AcquireWriteAccess();
 
-	/** Stop the multitheading access detection. */
+	/** Stop the multitheading write access detection. */
 	void ReleaseWriteAccess();
 
 	/** Get the data used at runtime to confirm the inner working of the StateTree. */
@@ -321,8 +327,12 @@ protected:
 	UPROPERTY(Transient)
 	uint32 UniqueIdGenerator = 0;
 
-	/** Used to detect if we are using the instance data on multiple threads at the same time. */
-	UE_MT_DECLARE_RW_RECURSIVE_ACCESS_DETECTOR(AccessDetector)
+	/**
+	 * Used to detect if we are using the instance data on multiple threads in a safe way.
+	 * The instance data supports multiple reader threads or a single writer thread.
+	 * The detector supports recursive access.
+	 */
+	UE_MT_DECLARE_MRSW_RECURSIVE_ACCESS_DETECTOR(AccessDetector)
 
 	/* True if the storage owns the event queue. */
 	bool bIsOwningEventQueue = true;
diff --git a/Engine/Plugins/Runtime/StateTree/Source/StateTreeModule/Public/StateTreeTaskBase.h b/Engine/Plugins/Runtime/StateTree/Source/StateTreeModule/Public/StateTreeTaskBase.h
--- a/Engine/Plugins/Runtime/StateTree/Source/StateTreeModule/Public/StateTreeTaskBase.h
+++ b/Engine/Plugins/Runtime/StateTree/Source/StateTreeModule/Public/StateTreeTaskBase.h
@@ -7,6 +7,7 @@
 #include "StateTreeTaskBase.generated.h"
 
 struct FStateTreeExecutionContext;
+struct FStateTreeReadOnlyExecutionContext;
 
 /**
  * Base struct for StateTree Tasks.
@@ -95,7 +96,10 @@ struct STATETREEMODULE_API FStateTreeTaskBase : public FStateTreeNodeBase
 #endif
 
 #if WITH_GAMEPLAY_DEBUGGER
-	virtual void AppendDebugInfoString(FString& DebugString, const FStateTreeExecutionContext& Context) const;
+	virtual FString GetDebugInfo(const FStateTreeReadOnlyExecutionContext& Context) const;
+
+	UE_DEPRECATED(5.6, "Use the version with the FStateTreeReadOnlyExecutionContext.")
+	virtual void AppendDebugInfoString(FString& DebugString, const FStateTreeExecutionContext& Context) const final;
 #endif
 
 	/**
