StateTree: Rename objects that will be duplicated on compiles to get predictable object names.

diff --git a/Engine/Plugins/Runtime/StateTree/Source/StateTreeEditorModule/Private/StateTreeCompiler.cpp b/Engine/Plugins/Runtime/StateTree/Source/StateTreeEditorModule/Private/StateTreeCompiler.cpp
--- a/Engine/Plugins/Runtime/StateTree/Source/StateTreeEditorModule/Private/StateTreeCompiler.cpp
+++ b/Engine/Plugins/Runtime/StateTree/Source/StateTreeEditorModule/Private/StateTreeCompiler.cpp
@@ -195,6 +195,27 @@ namespace UE::StateTree::Compiler
 		return Result;
 	}
 
+	UObject* DuplicateInstanceObject(FStateTreeCompilerLog& Log, const FStateTreeBindableStructDesc& NodeDesc, FGuid NodeID, TNotNull<const UObject*> InstanceObject, TNotNull<UObject*> Owner)
+	{
+		if (InstanceObject->GetClass()->HasAnyClassFlags(CLASS_NewerVersionExists))
+		{
+			const UStateTree* OuterStateTree = Owner->GetTypedOuter<UStateTree>();
+			Log.Reportf(EMessageSeverity::Warning, NodeDesc,
+				TEXT("Duplicating '%s' with an old class '%s' Please resave State Tree asset '%s'."),
+				*InstanceObject->GetName(), *InstanceObject->GetClass()->GetName(), *GetFullNameSafe(OuterStateTree));
+		}
+
+		// We want the object name to match between compilations.
+		//Use the class name and increase the counter internally. We do that to not be influenced by another object in a different outer.
+		//The objects from a previous compilation are rename in UStateTree::ResetCompiled.
+		FName NewObjectName = InstanceObject->GetClass()->GetFName();
+		while (StaticFindObjectFastInternal(nullptr, Owner, NewObjectName, true) != nullptr)
+		{
+			NewObjectName.SetNumber(NewObjectName.GetNumber() + 1);
+		}
+		return ::DuplicateObject(&(*InstanceObject), &(*Owner), NewObjectName);
+	}
+
 }; // UE::StateTree::Compiler
 
 bool FStateTreeCompiler::Compile(UStateTree& InStateTree)
@@ -314,8 +335,8 @@ bool FStateTreeCompiler::Compile(UStateTree& InStateTree)
 	}
 
 	StateTree->Nodes = Nodes;
-	StateTree->DefaultInstanceData.Init(*StateTree, InstanceStructs);
-	StateTree->SharedInstanceData.Init(*StateTree, SharedInstanceStructs);
+	StateTree->DefaultInstanceData.Init(*StateTree, InstanceStructs, FStateTreeInstanceData::FAddArgs{ .bDuplicateWrappedObject = false });
+	StateTree->SharedInstanceData.Init(*StateTree, SharedInstanceStructs, FStateTreeInstanceData::FAddArgs{ .bDuplicateWrappedObject = false });
 
 	// Store the new compiled dispatchers.
 	EditorData->CompiledDispatchers = BindingsCompiler.GetCompiledDelegateDispatchers();
@@ -1498,7 +1519,7 @@ T* FStateTreeCompiler::CreateNodeWithSharedInstanceData(UStateTreeState* State,
 		// Object Instance
 		check(EditorNode.InstanceObject != nullptr);
 
-		UObject* Instance = DuplicateObject(EditorNode.InstanceObject, StateTree);
+		UObject* Instance = UE::StateTree::Compiler::DuplicateInstanceObject(Log, StructDesc, EditorNode.ID, EditorNode.InstanceObject, StateTree);
 		FInstancedStruct Wrapper;
 		Wrapper.InitializeAs<FStateTreeInstanceObjectWrapper>(Instance);
 		const int32 InstanceIndex = SharedInstanceStructs.Add(MoveTemp(Wrapper));
@@ -1810,8 +1831,7 @@ bool FStateTreeCompiler::CreateTask(UStateTreeState* State, const FStateTreeEdit
 	{
 		// Object Instance
 		check(TaskNode.InstanceObject != nullptr);
-
-		UObject* Instance = DuplicateObject(TaskNode.InstanceObject, StateTree);
+		UObject* Instance = UE::StateTree::Compiler::DuplicateInstanceObject(Log, StructDesc, TaskNode.ID, TaskNode.InstanceObject, StateTree);
 		FInstancedStruct Wrapper;
 		Wrapper.InitializeAs<FStateTreeInstanceObjectWrapper>(Instance);
 		const int32 InstanceIndex = InstanceStructs.Add(MoveTemp(Wrapper));
@@ -1904,7 +1924,7 @@ bool FStateTreeCompiler::CreateEvaluator(const FStateTreeEditorNode& EvalNode, c
 		// Object Instance
 		check(EvalNode.InstanceObject != nullptr);
 
-		UObject* Instance = DuplicateObject(EvalNode.InstanceObject, StateTree);
+		UObject* Instance = UE::StateTree::Compiler::DuplicateInstanceObject(Log, StructDesc, EvalNode.ID, EvalNode.InstanceObject, StateTree);
 		FInstancedStruct Wrapper;
 		Wrapper.InitializeAs<FStateTreeInstanceObjectWrapper>(Instance);
 		const int32 InstanceIndex = InstanceStructs.Add(MoveTemp(Wrapper));
diff --git a/Engine/Plugins/Runtime/StateTree/Source/StateTreeModule/Private/StateTreeInstanceData.cpp b/Engine/Plugins/Runtime/StateTree/Source/StateTreeModule/Private/StateTreeInstanceData.cpp
--- a/Engine/Plugins/Runtime/StateTree/Source/StateTreeModule/Private/StateTreeInstanceData.cpp
+++ b/Engine/Plugins/Runtime/StateTree/Source/StateTreeModule/Private/StateTreeInstanceData.cpp
@@ -32,35 +32,38 @@ namespace UE::StateTree
 	/**
 	 * Duplicates object, and tries to covert old BP classes (REINST_*) to their newer version.
 	 */
-	UObject* DuplicateNodeInstance(const UObject& Instance, UObject& InOwner)
+	TNotNull<UObject*> CopyNodeInstance(TNotNull<UObject*> Instance, TNotNull<UObject*> InOwner, bool bDuplicate)
 	{
-		const UClass* InstanceClass = Instance.GetClass();
+		const UClass* InstanceClass = Instance->GetClass();
 		if (InstanceClass->HasAnyClassFlags(CLASS_NewerVersionExists))
 		{
 			const UClass* AuthoritativeClass = InstanceClass->GetAuthoritativeClass();
-			UObject* NewInstance = NewObject<UObject>(&InOwner, AuthoritativeClass);
+			UObject* NewInstance = NewObject<UObject>(InOwner, AuthoritativeClass);
 
 			// Try to copy the values over using serialization
 			// FObjectAndNameAsStringProxyArchive is used to store and restore names and objects as memory writer does not support UObject references at all.
 			TArray<uint8> Data;
 			FMemoryWriter Writer(Data);
 			FObjectAndNameAsStringProxyArchive WriterProxy(Writer, /*bInLoadIfFindFails*/true);
-			UObject& NonConstInstance = const_cast<UObject&>(Instance);
-			NonConstInstance.Serialize(WriterProxy);
+			Instance->Serialize(WriterProxy);
 
 			FMemoryReader Reader(Data);
 			FObjectAndNameAsStringProxyArchive ReaderProxy(Reader, /*bInLoadIfFindFails*/true);
 			NewInstance->Serialize(ReaderProxy);
 
-			const UStateTree* OuterStateTree = Instance.GetTypedOuter<UStateTree>();
-
+			const UStateTree* OuterStateTree = Instance->GetTypedOuter<UStateTree>();
 			UE_LOG(LogStateTree, Display, TEXT("FStateTreeInstanceData: Duplicating '%s' with old class '%s' as '%s', potential data loss. Please resave State Tree asset %s."),
-				*GetFullNameSafe(&Instance), *GetNameSafe(InstanceClass), *GetNameSafe(AuthoritativeClass), *GetFullNameSafe(OuterStateTree));
+				*GetFullNameSafe(Instance), *GetNameSafe(InstanceClass), *GetNameSafe(AuthoritativeClass), *GetFullNameSafe(OuterStateTree));
 
 			return NewInstance;
 		}
 
-		return DuplicateObject(&Instance, &InOwner);
+		if (bDuplicate)
+		{
+			return ::DuplicateObject(&(*Instance), &(*InOwner));
+		}
+
+		return Instance;
 	}
 
 #if WITH_EDITORONLY_DATA
@@ -282,7 +285,8 @@ FStructView FStateTreeInstanceStorage::AddTemporaryInstance(UObject& InOwner, co
 	{
 		if (Wrapper->InstanceObject)
 		{
-			Wrapper->InstanceObject = UE::StateTree::DuplicateNodeInstance(*Wrapper->InstanceObject, InOwner);
+			constexpr bool bDuplicate = true;
+			Wrapper->InstanceObject = UE::StateTree::CopyNodeInstance(Wrapper->InstanceObject, &InOwner, bDuplicate);
 		}
 	}
 
@@ -410,6 +414,8 @@ UE::StateTree::Debug::FRuntimeValidation FStateTreeInstanceStorage::GetRuntimeVa
 // FStateTreeInstanceData
 //----------------------------------------------------------------//
 
+FStateTreeInstanceData::FAddArgs FStateTreeInstanceData::FAddArgs::Default;
+
 FStateTreeInstanceData::FStateTreeInstanceData() = default;
 
 FStateTreeInstanceData::FStateTreeInstanceData(const FStateTreeInstanceData& Other)
@@ -679,25 +685,26 @@ void FStateTreeInstanceData::CopyFrom(UObject& InOwner, const FStateTreeInstance
 		{
 			if (Wrapper->InstanceObject)
 			{
-				Wrapper->InstanceObject = UE::StateTree::DuplicateNodeInstance(*Wrapper->InstanceObject, InOwner);
+				constexpr bool bDuplicate = true;
+				Wrapper->InstanceObject = UE::StateTree::CopyNodeInstance(Wrapper->InstanceObject, &InOwner, bDuplicate);
 			}
 		}
 	}
 }
 
-void FStateTreeInstanceData::Init(UObject& InOwner, TConstArrayView<FInstancedStruct> InStructs)
+void FStateTreeInstanceData::Init(UObject& InOwner, TConstArrayView<FInstancedStruct> InStructs, FAddArgs Args)
 {
 	Reset();
-	Append(InOwner, InStructs);
+	Append(InOwner, InStructs, Args);
 }
 
-void FStateTreeInstanceData::Init(UObject& InOwner, TConstArrayView<FConstStructView> InStructs)
+void FStateTreeInstanceData::Init(UObject& InOwner, TConstArrayView<FConstStructView> InStructs, FAddArgs Args)
 {
 	Reset();
-	Append(InOwner, InStructs);
+	Append(InOwner, InStructs, Args);
 }
 
-void FStateTreeInstanceData::Append(UObject& InOwner, TConstArrayView<FInstancedStruct> InStructs)
+void FStateTreeInstanceData::Append(UObject& InOwner, TConstArrayView<FInstancedStruct> InStructs, FAddArgs Args)
 {
 	FStateTreeInstanceStorage& Storage = GetMutableStorage();
 
@@ -710,13 +717,14 @@ void FStateTreeInstanceData::Append(UObject& InOwner, TConstArrayView<FInstanced
 		{
 			if (Wrapper->InstanceObject)
 			{
-				Wrapper->InstanceObject = UE::StateTree::DuplicateNodeInstance(*Wrapper->InstanceObject, InOwner);
+				const bool bDuplicate = Args.bDuplicateWrappedObject || &InOwner != Wrapper->InstanceObject->GetOuter();
+				Wrapper->InstanceObject = UE::StateTree::CopyNodeInstance(Wrapper->InstanceObject, &InOwner, bDuplicate);
 			}
 		}
 	}
 }
 
-void FStateTreeInstanceData::Append(UObject& InOwner, TConstArrayView<FConstStructView> InStructs)
+void FStateTreeInstanceData::Append(UObject& InOwner, TConstArrayView<FConstStructView> InStructs, FAddArgs Args)
 {
 	FStateTreeInstanceStorage& Storage = GetMutableStorage();
 
@@ -729,13 +737,14 @@ void FStateTreeInstanceData::Append(UObject& InOwner, TConstArrayView<FConstStru
 		{
 			if (Wrapper->InstanceObject)
 			{
-				Wrapper->InstanceObject = UE::StateTree::DuplicateNodeInstance(*Wrapper->InstanceObject, InOwner);
+				const bool bDuplicate = Args.bDuplicateWrappedObject || &InOwner != Wrapper->InstanceObject->GetOuter();
+				Wrapper->InstanceObject = UE::StateTree::CopyNodeInstance(Wrapper->InstanceObject, &InOwner, bDuplicate);
 			}
 		}
 	}
 }
 
-void FStateTreeInstanceData::Append(UObject& InOwner, TConstArrayView<FConstStructView> InStructs, TConstArrayView<FInstancedStruct*> InInstancesToMove)
+void FStateTreeInstanceData::Append(UObject& InOwner, TConstArrayView<FConstStructView> InStructs, TConstArrayView<FInstancedStruct*> InInstancesToMove, FAddArgs Args)
 {
 	check(InStructs.Num() == InInstancesToMove.Num());
 	
@@ -765,7 +774,8 @@ void FStateTreeInstanceData::Append(UObject& InOwner, TConstArrayView<FConstStru
 		{
 			if (Wrapper->InstanceObject)
 			{
-				Wrapper->InstanceObject = UE::StateTree::DuplicateNodeInstance(*Wrapper->InstanceObject, InOwner);
+				const bool bDuplicate = Args.bDuplicateWrappedObject || &InOwner != Wrapper->InstanceObject->GetOuter();
+				Wrapper->InstanceObject = UE::StateTree::CopyNodeInstance(Wrapper->InstanceObject, &InOwner, bDuplicate);
 			}
 		}
 	}
diff --git a/Engine/Plugins/Runtime/StateTree/Source/StateTreeModule/Public/StateTreeInstanceData.h b/Engine/Plugins/Runtime/StateTree/Source/StateTreeModule/Public/StateTreeInstanceData.h
--- a/Engine/Plugins/Runtime/StateTree/Source/StateTreeModule/Public/StateTreeInstanceData.h
+++ b/Engine/Plugins/Runtime/StateTree/Source/StateTreeModule/Public/StateTreeInstanceData.h
@@ -365,17 +365,24 @@ struct STATETREEMODULE_API FStateTreeInstanceData
 	FStateTreeInstanceData& operator=(FStateTreeInstanceData&& Other) noexcept;
 
 	~FStateTreeInstanceData();
-	
+
+	struct FAddArgs
+	{
+		STATETREEMODULE_API static FAddArgs Default;
+		/** Duplicate the object contained by object wrapper. */
+		bool bDuplicateWrappedObject = true;
+	};
+
 	/** Initializes the array with specified items. */
-	void Init(UObject& InOwner, TConstArrayView<FInstancedStruct> InStructs);
-	void Init(UObject& InOwner, TConstArrayView<FConstStructView> InStructs);
+	void Init(UObject& InOwner, TConstArrayView<FInstancedStruct> InStructs, FAddArgs Args = FAddArgs::Default);
+	void Init(UObject& InOwner, TConstArrayView<FConstStructView> InStructs, FAddArgs Args = FAddArgs::Default);
 
 	/** Appends new items to the instance. */
-	void Append(UObject& InOwner, TConstArrayView<FInstancedStruct> InStructs);
-	void Append(UObject& InOwner, TConstArrayView<FConstStructView> InStructs);
+	void Append(UObject& InOwner, TConstArrayView<FInstancedStruct> InStructs, FAddArgs Args = FAddArgs::Default);
+	void Append(UObject& InOwner, TConstArrayView<FConstStructView> InStructs, FAddArgs Args = FAddArgs::Default);
 
 	/** Appends new items to the instance, and moves existing data into the allocated instances. */
-	void Append(UObject& InOwner, TConstArrayView<FConstStructView> InStructs, TConstArrayView<FInstancedStruct*> InInstancesToMove);
+	void Append(UObject& InOwner, TConstArrayView<FConstStructView> InStructs, TConstArrayView<FInstancedStruct*> InInstancesToMove, FAddArgs Args = FAddArgs::Default);
 
 	/** Shrinks the array sizes to specified lengths. Sizes must be small or equal than current size. */
 	void ShrinkTo(const int32 Num);
