StateTree: Fixed incorrectly created ArrayView to ActiveStates.

diff --git a/Engine/Plugins/Runtime/StateTree/Source/StateTreeModule/Private/StateTreeExecutionContext.cpp b/Engine/Plugins/Runtime/StateTree/Source/StateTreeModule/Private/StateTreeExecutionContext.cpp
--- a/Engine/Plugins/Runtime/StateTree/Source/StateTreeModule/Private/StateTreeExecutionContext.cpp
+++ b/Engine/Plugins/Runtime/StateTree/Source/StateTreeModule/Private/StateTreeExecutionContext.cpp
@@ -2625,7 +2625,7 @@ void FStateTreeExecutionContext::CaptureNewStateEvents(TConstArrayView<FStateTre
 		const FStateTreeExecutionFrame& NewFrame = NewFrames[FrameIndex];
 
 		// Find states that are unique to the new frame.
-		TConstArrayView<FStateTreeStateHandle> UniqueStates = NewFrame.ActiveStates.States;
+		TConstArrayView<FStateTreeStateHandle> UniqueStates = NewFrame.ActiveStates;
 		if (PrevFrames.IsValidIndex(FrameIndex))
 		{
 			const FStateTreeExecutionFrame& PrevFrame = PrevFrames[FrameIndex];
diff --git a/Engine/Plugins/Runtime/StateTree/Source/StateTreeModule/Public/StateTreeExecutionTypes.h b/Engine/Plugins/Runtime/StateTree/Source/StateTreeModule/Public/StateTreeExecutionTypes.h
--- a/Engine/Plugins/Runtime/StateTree/Source/StateTreeModule/Public/StateTreeExecutionTypes.h
+++ b/Engine/Plugins/Runtime/StateTree/Source/StateTreeModule/Public/StateTreeExecutionTypes.h
@@ -500,6 +500,18 @@ struct STATETREEMODULE_API FStateTreeActiveStates
 		return States[Index];
 	}
 
+	/** Returns the active states from the States array. */
+	operator TArrayView<FStateTreeStateHandle>()
+	{
+		return TArrayView<FStateTreeStateHandle>(&States[0], Num());
+	}
+
+	/** Returns the active states from the States array. */
+	operator TConstArrayView<FStateTreeStateHandle>() const
+	{
+		return TConstArrayView<FStateTreeStateHandle>(&States[0], Num());
+	}
+
 	/** Returns a specified state in the array, or FStateTreeStateHandle::Invalid if Index is out of array bounds. */
 	FStateTreeStateHandle GetStateSafe(const int32 Index) const
 	{
