StateTree: Add specific StateTreeEditorData for each schema. A schema can register its desired EditorData type. Then it will upgrade the next time the editor asset opens or when the asset compiles. It is the default editor type when the asset is created.
Add same verification code when you compile the asset to when you open the editor asset (for consistency).

diff --git a/Engine/Plugins/Runtime/StateTree/Source/StateTreeEditorModule/Private/StateTreeEditingSubsystem.cpp b/Engine/Plugins/Runtime/StateTree/Source/StateTreeEditorModule/Private/StateTreeEditingSubsystem.cpp
--- a/Engine/Plugins/Runtime/StateTree/Source/StateTreeEditorModule/Private/StateTreeEditingSubsystem.cpp
+++ b/Engine/Plugins/Runtime/StateTree/Source/StateTreeEditorModule/Private/StateTreeEditingSubsystem.cpp
@@ -81,37 +81,10 @@ TSharedRef<FStateTreeViewModel> UStateTreeEditingSubsystem::FindOrAddViewModel(T
 		}
 	}
 
+	ValidateStateTree(InStateTree);
+
 	TSharedRef<FStateTreeViewModel> SharedModel = StateTreeViewModels.Add(StateTreeKey, MakeShared<FStateTreeViewModel>()).ToSharedRef();
 	UStateTreeEditorData* EditorData = Cast<UStateTreeEditorData>(InStateTree->EditorData);
-	if (EditorData == nullptr)
-	{
-		EditorData = NewObject<UStateTreeEditorData>(InStateTree, FName(), RF_Transactional);
-		EditorData->AddRootState();
-		InStateTree->EditorData = EditorData;
-
-		FStateTreeCompilerLog Log;
-		CompileStateTree(InStateTree, Log);
-	}
-	
-	for (UStateTreeState* SubTree : EditorData->SubTrees)
-	{
-		TArray<UStateTreeState*> Stack;
-
-		Stack.Add(SubTree);
-		while (!Stack.IsEmpty())
-		{
-			if (UStateTreeState* State = Stack.Pop())
-			{
-				State->SetFlags(RF_Transactional);
-				
-				for (UStateTreeState* ChildState : State->Children)
-				{
-					Stack.Add(ChildState);
-				}
-			}
-		}
-	}
-		
 	SharedModel->Init(EditorData);
 
 	return SharedModel;
@@ -122,7 +95,7 @@ TSharedRef<SWidget> UStateTreeEditingSubsystem::GetStateTreeView(TSharedRef<FSta
 	return SNew(SStateTreeView, InViewModel, TreeViewCommandList);
 }
 
-void UStateTreeEditingSubsystem::ValidateStateTree(TNotNull<const UStateTree*> InStateTree)
+void UStateTreeEditingSubsystem::ValidateStateTree(TNotNull<UStateTree*> InStateTree)
 {
 	auto FixChangedStateLinkName = [](FStateTreeStateLink& StateLink, const TMap<FGuid, FName>& IDToName) -> bool
 	{
@@ -152,9 +125,6 @@ void UStateTreeEditingSubsystem::ValidateStateTree(TNotNull<const UStateTree*> I
 			return;
 		}
 
-		constexpr bool bMarkDirty = false;
-		TreeData->Modify(bMarkDirty);
-
 		// Make sure all state links are valid and update the names if needed.
 
 		// Create ID to state name map.
@@ -185,6 +155,64 @@ void UStateTreeEditingSubsystem::ValidateStateTree(TNotNull<const UStateTree*> I
 		});
 	};
 
+	auto FixEditorData = [](TNotNull<UStateTree*> StateTree)
+	{
+		UStateTreeEditorData* EditorData = Cast<UStateTreeEditorData>(StateTree->EditorData);
+		// The schema is defined in the EditorData. If we can't find the editor data (probably because the class doesn't exist anymore), then try the compiled schema in the state tree asset.
+		TSubclassOf<const UStateTreeSchema> SchemaClass;
+		if (EditorData && EditorData->Schema)
+		{
+			SchemaClass = EditorData->Schema->GetClass();
+		}
+		else if (StateTree->GetSchema())
+		{
+			SchemaClass = StateTree->GetSchema()->GetClass();
+		}
+
+		if (SchemaClass.Get() == nullptr)
+		{
+			UE_LOG(LogStateTreeEditor, Error, TEXT("The state tree '%s' does not have a schema."), *StateTree->GetPathName());
+			return;
+		}
+
+		TNonNullSubclassOf<UStateTreeEditorData> EditorDataClass = FStateTreeEditorModule::GetModule().GetEditorDataClass(SchemaClass.Get());
+		if (EditorData == nullptr)
+		{
+			EditorData = NewObject<UStateTreeEditorData>(StateTree, EditorDataClass.Get(), FName(), RF_Transactional);
+			EditorData->AddRootState();
+			EditorData->Schema = NewObject<UStateTreeSchema>(EditorData, SchemaClass.Get());
+
+			constexpr bool bMarkDirty = false;
+			StateTree->Modify(bMarkDirty);
+			StateTree->EditorData = EditorData;
+		}
+		else if (!EditorData->IsA(EditorDataClass.Get()))
+		{
+			// The current EditorData is not of the correct type. The data needs to be patched by the schema desired editor data subclass.
+			UStateTreeEditorData* PreviousEditorData = EditorData;
+			EditorData = CastChecked<UStateTreeEditorData>(StaticDuplicateObject(EditorData, StateTree, FName(), RF_Transactional, EditorDataClass.Get()));
+			if (EditorData->SubTrees.Num() == 0)
+			{
+				EditorData->AddRootState();
+			}
+			if (EditorData->Schema == nullptr || !EditorData->Schema->IsA(SchemaClass.Get()))
+			{
+				EditorData->Schema = NewObject<UStateTreeSchema>(EditorData, SchemaClass.Get());
+			}
+
+			constexpr bool bMarkDirty = false;
+			StateTree->Modify(bMarkDirty);
+			StateTree->EditorData = EditorData;
+
+			// Trash the previous EditorData
+			FName TrashName = MakeUniqueObjectName(GetTransientPackage(), UStateTree::StaticClass(), *FString::Printf(TEXT("TRASH_%s"), *UStateTree::StaticClass()->GetName()));
+			UStateTree* TransientOuter = NewObject<UStateTree>(GetTransientPackage(), TrashName, RF_Transient);
+			ERenameFlags RenameFlags = REN_DoNotDirty | REN_DontCreateRedirectors;
+			PreviousEditorData->Rename(*TrashName.ToString(), TransientOuter, RenameFlags);
+			PreviousEditorData->SetFlags(RF_Transient);
+		}
+	};
+
 	auto UpdateParents = [](const UStateTree& StateTree)
 	{
 		UStateTreeEditorData* TreeData = Cast<UStateTreeEditorData>(StateTree.EditorData);
@@ -193,8 +221,6 @@ void UStateTreeEditingSubsystem::ValidateStateTree(TNotNull<const UStateTree*> I
 			return;
 		}
 
-		constexpr bool bMarkDirty = false;
-		TreeData->Modify(bMarkDirty);
 		TreeData->ReparentStates();
 	};
 
@@ -211,9 +237,6 @@ void UStateTreeEditingSubsystem::ValidateStateTree(TNotNull<const UStateTree*> I
 		{
 			return;
 		}
-
-		constexpr bool bMarkDirty = false;
-		TreeData->Modify(bMarkDirty);
 		
 		// Clear evaluators if not allowed.
 		if (Schema->AllowEvaluators() == false && TreeData->Evaluators.Num() > 0)
@@ -222,7 +245,6 @@ void UStateTreeEditingSubsystem::ValidateStateTree(TNotNull<const UStateTree*> I
 			TreeData->Evaluators.Reset();
 		}
 
-
 		TreeData->VisitHierarchy([&StateTree, Schema](UStateTreeState& State, UStateTreeState* /*ParentState*/)
 		{
 			constexpr bool bMarkDirty = false;
@@ -269,7 +291,6 @@ void UStateTreeEditingSubsystem::ValidateStateTree(TNotNull<const UStateTree*> I
 			return EStateTreeVisitor::Continue;
 		});
 	};
-	
 
 	auto RemoveUnusedBindings = [](const UStateTree& StateTree)
 	{
@@ -281,8 +302,6 @@ void UStateTreeEditingSubsystem::ValidateStateTree(TNotNull<const UStateTree*> I
 
 		TMap<FGuid, const FPropertyBindingDataView> AllStructValues;
 		TreeData->GetAllStructValues(AllStructValues);
-		constexpr bool bMarkDirty = false;
-		TreeData->Modify(bMarkDirty);
 		TreeData->GetPropertyEditorBindings()->RemoveInvalidBindings(AllStructValues);
 	};
 
@@ -294,9 +313,6 @@ void UStateTreeEditingSubsystem::ValidateStateTree(TNotNull<const UStateTree*> I
 			return;
 		}
 
-		constexpr bool bMarkDirty = false;
-		TreeData->Modify(bMarkDirty);
-
 		const EStateTreeVisitor Result = TreeData->VisitHierarchy([](UStateTreeState& State, UStateTreeState* /*ParentState*/)
 		{
 			if (State.Type == EStateTreeStateType::Linked
@@ -310,12 +326,46 @@ void UStateTreeEditingSubsystem::ValidateStateTree(TNotNull<const UStateTree*> I
 		});
 	};
 
+	auto UpdateTransactionalFlags = [](TNotNull<UStateTree*> StateTree)
+	{
+		UStateTreeEditorData* EditorData = Cast<UStateTreeEditorData>(StateTree->EditorData);
+		if (EditorData)
+		{
+			for (UStateTreeState* SubTree : EditorData->SubTrees)
+			{
+				TArray<UStateTreeState*> Stack;
+
+				Stack.Add(SubTree);
+				while (!Stack.IsEmpty())
+				{
+					if (UStateTreeState* State = Stack.Pop())
+					{
+						State->SetFlags(RF_Transactional);
+
+						for (UStateTreeState* ChildState : State->Children)
+						{
+							Stack.Add(ChildState);
+						}
+					}
+				}
+			}
+		}
+	};
+
+	FixEditorData(InStateTree);
+
+	if (InStateTree->EditorData)
+	{
+		constexpr bool bMarkDirty = false;
+		InStateTree->EditorData->Modify(bMarkDirty);
+	}
 
 	UpdateParents(*InStateTree);
 	ApplySchema(*InStateTree);
 	RemoveUnusedBindings(*InStateTree);
 	ValidateLinkedStates(*InStateTree);
 	UpdateLinkedStateParameters(*InStateTree);
+	UpdateTransactionalFlags(InStateTree);
 }
 
 uint32 UStateTreeEditingSubsystem::CalculateStateTreeHash(TNotNull<const UStateTree*> InStateTree)
diff --git a/Engine/Plugins/Runtime/StateTree/Source/StateTreeEditorModule/Private/StateTreeEditorModule.cpp b/Engine/Plugins/Runtime/StateTree/Source/StateTreeEditorModule/Private/StateTreeEditorModule.cpp
--- a/Engine/Plugins/Runtime/StateTree/Source/StateTreeEditorModule/Private/StateTreeEditorModule.cpp
+++ b/Engine/Plugins/Runtime/StateTree/Source/StateTreeEditorModule/Private/StateTreeEditorModule.cpp
@@ -25,6 +25,7 @@
 #include "StateTree.h"
 #include "StateTreeCompiler.h"
 #include "StateTreeCompilerLog.h"
+#include "StateTreeEditorData.h"
 #include "Debugger/StateTreeDebuggerCommands.h"
 #include "Debugger/StateTreeRewindDebuggerExtensions.h"
 #include "StateTreeDelegates.h"
@@ -52,6 +53,16 @@ namespace UE::StateTree::Editor
 
 }; // UE::StateTree::Editor
 
+FStateTreeEditorModule& FStateTreeEditorModule::GetModule()
+{
+	return FModuleManager::GetModuleChecked<FStateTreeEditorModule>("StateTreeEditorModule");
+}
+
+FStateTreeEditorModule* FStateTreeEditorModule::GetModulePtr()
+{
+	return FModuleManager::GetModulePtr<FStateTreeEditorModule>("StateTreeEditorModule");
+}
+
 void FStateTreeEditorModule::StartupModule()
 {
 	UE::StateTree::Delegates::OnRequestCompile.BindStatic(&UE::StateTree::Editor::CompileStateTree);
@@ -161,5 +172,40 @@ TSharedPtr<FStateTreeNodeClassCache> FStateTreeEditorModule::GetNodeClassCache()
 	return NodeClassCache;
 }
 
+void FStateTreeEditorModule::RegisterEditorDataClass(TNonNullSubclassOf<const UStateTreeSchema> Schema, TNonNullSubclassOf<const UStateTreeEditorData> EditorData)
+{
+	const bool bContains = EditorDataTypes.ContainsByPredicate(
+		[Schema](const FEditorDataType& Other)
+		{
+			return Other.Schema == Schema.Get();
+		});
+	ensureMsgf(bContains, TEXT("The type %s is already registered."), *Schema.Get()->GetName());
+	if (!bContains)
+	{
+		EditorDataTypes.Emplace(TWeakObjectPtr<const UClass>(Schema.Get()), TWeakObjectPtr<const UClass>(EditorData.Get()));
+	}
+}
+
+void FStateTreeEditorModule::UnregisterEditorDataClass(TNonNullSubclassOf<const UStateTreeSchema> Schema)
+{
+	EditorDataTypes.RemoveAllSwap(
+		[Schema](const FEditorDataType& Other)
+		{
+			return Other.Schema == Schema.Get();
+		});
+}
+
+TNonNullSubclassOf<UStateTreeEditorData> FStateTreeEditorModule::GetEditorDataClass(TNonNullSubclassOf<const UStateTreeSchema> Schema) const
+{
+	const FEditorDataType* FoundEditorData = EditorDataTypes.FindByPredicate(
+		[Schema](const FEditorDataType& Other)
+		{
+			const UClass* OtherSchema = Other.Schema.Get();
+			return OtherSchema && Schema.Get()->IsChildOf(OtherSchema);
+		});
+	const UClass* Result = FoundEditorData && FoundEditorData->EditorData.Get() ? FoundEditorData->EditorData.Get() : UStateTreeEditorData::StaticClass();
+	return const_cast<UClass*>(Result); // NewObject wants none const UClass
+}
+
 
 #undef LOCTEXT_NAMESPACE
diff --git a/Engine/Plugins/Runtime/StateTree/Source/StateTreeEditorModule/Private/StateTreeFactory.cpp b/Engine/Plugins/Runtime/StateTree/Source/StateTreeEditorModule/Private/StateTreeFactory.cpp
--- a/Engine/Plugins/Runtime/StateTree/Source/StateTreeEditorModule/Private/StateTreeFactory.cpp
+++ b/Engine/Plugins/Runtime/StateTree/Source/StateTreeEditorModule/Private/StateTreeFactory.cpp
@@ -120,7 +120,8 @@ UObject* UStateTreeFactory::FactoryCreateNew(UClass* Class, UObject* InParent, F
 	UStateTree* NewStateTree = NewObject<UStateTree>(InParent, Class, Name, Flags);
 
 	// Create and init new editor data.
-	UStateTreeEditorData* EditorData = NewObject<UStateTreeEditorData>(NewStateTree, FName(), RF_Transactional);
+	TNonNullSubclassOf<UStateTreeEditorData> EditorDataClass = FStateTreeEditorModule::GetModule().GetEditorDataClass(StateTreeSchemaClass.Get());
+	UStateTreeEditorData* EditorData = NewObject<UStateTreeEditorData>(NewStateTree, EditorDataClass, FName(), RF_Transactional);
 	EditorData->AddRootState();
 	EditorData->Schema = NewObject<UStateTreeSchema>(EditorData, StateTreeSchemaClass);
 	
diff --git a/Engine/Plugins/Runtime/StateTree/Source/StateTreeEditorModule/Public/StateTreeEditingSubsystem.h b/Engine/Plugins/Runtime/StateTree/Source/StateTreeEditorModule/Public/StateTreeEditingSubsystem.h
--- a/Engine/Plugins/Runtime/StateTree/Source/StateTreeEditorModule/Public/StateTreeEditingSubsystem.h
+++ b/Engine/Plugins/Runtime/StateTree/Source/StateTreeEditorModule/Public/StateTreeEditingSubsystem.h
@@ -33,7 +33,7 @@ public:
 	 * Validates and applies the schema restrictions on the StateTree.
 	 * Updates state's link, removes the unused node while validating the StateTree asset.
 	 */
-	static void ValidateStateTree(TNotNull<const UStateTree*> InStateTree);
+	static void ValidateStateTree(TNotNull<UStateTree*> InStateTree);
 
 	/** Calculates editor data hash of the asset. */
 	static uint32 CalculateStateTreeHash(TNotNull<const UStateTree*> InStateTree);
diff --git a/Engine/Plugins/Runtime/StateTree/Source/StateTreeEditorModule/Public/StateTreeEditorModule.h b/Engine/Plugins/Runtime/StateTree/Source/StateTreeEditorModule/Public/StateTreeEditorModule.h
--- a/Engine/Plugins/Runtime/StateTree/Source/StateTreeEditorModule/Public/StateTreeEditorModule.h
+++ b/Engine/Plugins/Runtime/StateTree/Source/StateTreeEditorModule/Public/StateTreeEditorModule.h
@@ -4,10 +4,13 @@
 
 #include "Modules/ModuleInterface.h"
 #include "Logging/LogMacros.h"
+#include "Templates/NonNullSubclassOf.h"
 #include "Templates/PimplPtr.h"
 #include "Toolkits/AssetEditorToolkit.h"
 
 class UStateTree;
+class UStateTreeEditorData;
+class UStateTreeSchema;
 class UUserDefinedStruct;
 class IStateTreeEditor;
 struct FStateTreeNodeClassCache;
@@ -26,9 +29,13 @@ STATETREEEDITORMODULE_API DECLARE_LOG_CATEGORY_EXTERN(LogStateTreeEditor, Log, A
 class STATETREEEDITORMODULE_API FStateTreeEditorModule : public IModuleInterface, public IHasMenuExtensibility, public IHasToolBarExtensibility
 {
 public:
-	// Begin IModuleInterface
+	//~Begin IModuleInterface
 	virtual void StartupModule() override;
 	virtual void ShutdownModule() override;
+	//~End IModuleInterface
+
+	static FStateTreeEditorModule& GetModule();
+	static FStateTreeEditorModule* GetModulePtr();
 
 	/** Creates an instance of StateTree editor. Only virtual so that it can be called across the DLL boundary. */
 	virtual TSharedRef<IStateTreeEditor> CreateStateTreeEditor(const EToolkitMode::Type Mode, const TSharedPtr< class IToolkitHost >& InitToolkitHost, UStateTree* StateTree);
@@ -44,10 +51,25 @@ public:
 	DECLARE_EVENT_OneParam(FStateTreeEditorModule, FOnRegisterLayoutExtensions, FLayoutExtender&);
 	FOnRegisterLayoutExtensions& OnRegisterLayoutExtensions() { return RegisterLayoutExtensions; }
 
+	/** Register the editor data type for a specific schema. */
+	void RegisterEditorDataClass(TNonNullSubclassOf<const UStateTreeSchema> Schema, TNonNullSubclassOf<const UStateTreeEditorData> EditorData);
+	/** Unregister the editor data type for a specific schema. */
+	void UnregisterEditorDataClass(TNonNullSubclassOf<const UStateTreeSchema> Schema);
+	/** Get the editor data type for a specific schema. */
+	TNonNullSubclassOf<UStateTreeEditorData> GetEditorDataClass(TNonNullSubclassOf<const UStateTreeSchema> Schema) const;
+
 protected:
 	TSharedPtr<FExtensibilityManager> MenuExtensibilityManager;
 	TSharedPtr<FExtensibilityManager> ToolBarExtensibilityManager;
 	TSharedPtr<FStateTreeNodeClassCache> NodeClassCache;
+
+	struct FEditorDataType
+	{
+		TWeakObjectPtr<const UClass> Schema;
+		TWeakObjectPtr<const UClass> EditorData;
+	};
+	TArray<FEditorDataType> EditorDataTypes;
+
 #if WITH_STATETREE_TRACE_DEBUGGER
 	TPimplPtr<UE::StateTreeDebugger::FRewindDebuggerExtension> RewindDebuggerExtension;
 	TPimplPtr<UE::StateTreeDebugger::FRewindDebuggerRuntimeExtension> RewindDebuggerRuntimeExtension;
