StateTree: Add specific StateTreeEditorData for each schema. A schema can register its desired EditorData type. Then it will upgrade the next time the editor asset opens or when the asset compiles. It is the default editor type when the asset is created.
Add same verification code when you compile the asset to when you open the editor asset (for consistency).

diff --git a/Engine/Plugins/Runtime/StateTree/Source/StateTreeEditorModule/Private/StateTreeEditingSubsystem.cpp b/Engine/Plugins/Runtime/StateTree/Source/StateTreeEditorModule/Private/StateTreeEditingSubsystem.cpp
--- a/Engine/Plugins/Runtime/StateTree/Source/StateTreeEditorModule/Private/StateTreeEditingSubsystem.cpp
+++ b/Engine/Plugins/Runtime/StateTree/Source/StateTreeEditorModule/Private/StateTreeEditingSubsystem.cpp
@@ -158,7 +158,6 @@ void UStateTreeEditingSubsystem::ValidateStateTree(TNotNull<UStateTree*> InState
 	auto FixEditorData = [](TNotNull<UStateTree*> StateTree)
 	{
 		UStateTreeEditorData* EditorData = Cast<UStateTreeEditorData>(StateTree->EditorData);
-		// The schema is defined in the EditorData. If we can't find the editor data (probably because the class doesn't exist anymore), then try the compiled schema in the state tree asset.
 		TSubclassOf<const UStateTreeSchema> SchemaClass;
 		if (EditorData && EditorData->Schema)
 		{
@@ -188,7 +187,7 @@ void UStateTreeEditingSubsystem::ValidateStateTree(TNotNull<UStateTree*> InState
 		}
 		else if (!EditorData->IsA(EditorDataClass.Get()))
 		{
-			// The current EditorData is not of the correct type. The data needs to be patched by the schema desired editor data subclass.
+			// The current EditorData is not of the correct type. The data need to be patched.
 			UStateTreeEditorData* PreviousEditorData = EditorData;
 			EditorData = CastChecked<UStateTreeEditorData>(StaticDuplicateObject(EditorData, StateTree, FName(), RF_Transactional, EditorDataClass.Get()));
 			if (EditorData->SubTrees.Num() == 0)
@@ -353,13 +352,14 @@ void UStateTreeEditingSubsystem::ValidateStateTree(TNotNull<UStateTree*> InState
 	};
 
 	FixEditorData(InStateTree);
-
-	if (InStateTree->EditorData)
 	{
-		constexpr bool bMarkDirty = false;
-		InStateTree->EditorData->Modify(bMarkDirty);
+		UStateTreeEditorData* TreeData = Cast<UStateTreeEditorData>(InStateTree->EditorData);
+		if (TreeData)
+		{
+			constexpr bool bMarkDirty = false;
+			TreeData->Modify(bMarkDirty);
+		}
 	}
-
 	UpdateParents(*InStateTree);
 	ApplySchema(*InStateTree);
 	RemoveUnusedBindings(*InStateTree);
diff --git a/Engine/Plugins/Runtime/StateTree/Source/StateTreeEditorModule/Private/StateTreeEditorModule.cpp b/Engine/Plugins/Runtime/StateTree/Source/StateTreeEditorModule/Private/StateTreeEditorModule.cpp
--- a/Engine/Plugins/Runtime/StateTree/Source/StateTreeEditorModule/Private/StateTreeEditorModule.cpp
+++ b/Engine/Plugins/Runtime/StateTree/Source/StateTreeEditorModule/Private/StateTreeEditorModule.cpp
@@ -179,8 +179,7 @@ void FStateTreeEditorModule::RegisterEditorDataClass(TNonNullSubclassOf<const US
 		{
 			return Other.Schema == Schema.Get();
 		});
-	ensureMsgf(bContains, TEXT("The type %s is already registered."), *Schema.Get()->GetName());
-	if (!bContains)
+	if (ensureMsgf(!bContains, TEXT("The type %s is already registered."), *Schema.Get()->GetName()))
 	{
 		EditorDataTypes.Emplace(TWeakObjectPtr<const UClass>(Schema.Get()), TWeakObjectPtr<const UClass>(EditorData.Get()));
 	}
