[StateTreeTrace] fixed trace recording check when autostarted in PIE

diff --git a/Engine/Plugins/Runtime/StateTree/Source/StateTreeModule/Private/Debugger/StateTreeTrace.cpp b/Engine/Plugins/Runtime/StateTree/Source/StateTreeModule/Private/Debugger/StateTreeTrace.cpp
--- a/Engine/Plugins/Runtime/StateTree/Source/StateTreeModule/Private/Debugger/StateTreeTrace.cpp
+++ b/Engine/Plugins/Runtime/StateTree/Source/StateTreeModule/Private/Debugger/StateTreeTrace.cpp
@@ -412,6 +412,13 @@ void FBufferedDataList::OnTracesStarted_GameThread()
 		}
 	}
 
+	// Delegate can be received before first call to SetWorldTime_GameThread()
+	// In that case we set the time to 0 for any pending lifetime events
+	if (RecordingWorldTime == UNINITIALIZED_WORLD_TIME)
+	{
+		SetWorldTime_GameThread(0);
+	}
+
 	TraceWorldTimeIfNeeded_AnyThread();
 
 	// Then trace instance lifetime events since they are required for other event types.
