Make StateTree Editor GetModule load the module.

diff --git a/Engine/Plugins/Runtime/StateTree/Source/StateTreeEditorModule/Private/StateTreeEditorModule.cpp b/Engine/Plugins/Runtime/StateTree/Source/StateTreeEditorModule/Private/StateTreeEditorModule.cpp
--- a/Engine/Plugins/Runtime/StateTree/Source/StateTreeEditorModule/Private/StateTreeEditorModule.cpp
+++ b/Engine/Plugins/Runtime/StateTree/Source/StateTreeEditorModule/Private/StateTreeEditorModule.cpp
@@ -55,7 +55,7 @@ namespace UE::StateTree::Editor
 
 FStateTreeEditorModule& FStateTreeEditorModule::GetModule()
 {
-	return FModuleManager::GetModuleChecked<FStateTreeEditorModule>("StateTreeEditorModule");
+	return FModuleManager::LoadModuleChecked<FStateTreeEditorModule>("StateTreeEditorModule");
 }
 
 FStateTreeEditorModule* FStateTreeEditorModule::GetModulePtr()
diff --git a/Engine/Plugins/Runtime/StateTree/Source/StateTreeEditorModule/Public/StateTreeEditorModule.h b/Engine/Plugins/Runtime/StateTree/Source/StateTreeEditorModule/Public/StateTreeEditorModule.h
--- a/Engine/Plugins/Runtime/StateTree/Source/StateTreeEditorModule/Public/StateTreeEditorModule.h
+++ b/Engine/Plugins/Runtime/StateTree/Source/StateTreeEditorModule/Public/StateTreeEditorModule.h
@@ -34,7 +34,10 @@ public:
 	virtual void ShutdownModule() override;
 	//~End IModuleInterface
 
+	/** Gets this module, will attempt to load and should always exist. */
 	static FStateTreeEditorModule& GetModule();
+
+	/** Gets this module, will not attempt to load and may not exist. */
 	static FStateTreeEditorModule* GetModulePtr();
 
 	/** Creates an instance of StateTree editor. Only virtual so that it can be called across the DLL boundary. */
