[State Tree] barf nullptr UClass due to BP/c++ class deletion in StateTree link

diff --git a/Engine/Plugins/Runtime/StateTree/Source/StateTreeModule/Private/StateTree.cpp b/Engine/Plugins/Runtime/StateTree/Source/StateTreeModule/Private/StateTree.cpp
--- a/Engine/Plugins/Runtime/StateTree/Source/StateTreeModule/Private/StateTree.cpp
+++ b/Engine/Plugins/Runtime/StateTree/Source/StateTreeModule/Private/StateTree.cpp
@@ -536,6 +536,14 @@ bool UStateTree::ValidateInstanceData()
 				}
 			}
 
+			if (CurrentInstanceDataType == nullptr)
+			{
+				UE_LOG(LogStateTree, Error, TEXT("%s: node '%s' failed. Missing instance value, possibly due to Blueprint class or C++ class/struct template deletion."), *GetPathName(), *WriteToString<64>(Node->StaticStruct()->GetFName()));
+
+				bResult = false;
+				continue;
+			}
+
 			auto HasNewerVersionExists = [](const UStruct* InstanceDataType)
 			{
 				// Is the class/scriptstruct a blueprint that got replaced by another class.
