StateTree: Use the correct tree instance data when migrating a linked asset tree.

diff --git a/Engine/Plugins/Runtime/StateTree/Source/StateTreeModule/Private/StateTreeExecutionContext.cpp b/Engine/Plugins/Runtime/StateTree/Source/StateTreeModule/Private/StateTreeExecutionContext.cpp
--- a/Engine/Plugins/Runtime/StateTree/Source/StateTreeModule/Private/StateTreeExecutionContext.cpp
+++ b/Engine/Plugins/Runtime/StateTree/Source/StateTreeModule/Private/StateTreeExecutionContext.cpp
@@ -2073,9 +2073,12 @@ void FStateTreeExecutionContext::UpdateInstanceData(TConstArrayView<FStateTreeEx
 						FConstStructView ParamsInstanceData;
 						if (State.Type == EStateTreeStateType::LinkedAsset)
 						{
-							if (const FStateTreeReference* Override = GetLinkedStateTreeOverrideForTag(State.Tag))
+							// This state is a container for the linked state tree.Its instance data matches the linked state tree parameters.The linked state tree asset is the next frame.
+							const bool bIsLastFrame = FrameIndex == NextActiveFrames.Num() - 1;
+							if (!bIsLastFrame)
 							{
-								ParamsInstanceData = FConstStructView::Make(TempParams.Emplace_GetRef(Override->GetParameters()));
+								FStateTreeExecutionFrame& FollowingNextFrame = NextActiveFrames[FrameIndex + 1];
+								ParamsInstanceData = FConstStructView::Make(TempParams.Emplace_GetRef(FollowingNextFrame.StateTree->GetDefaultParameters()));
 							}
 						}
 						if (!ParamsInstanceData.IsValid())
