StateTree: Caches the ParameterDataType on compile. It prevents crash when there is no schema in the compiled asset. It prevents a mismatch when the asset was compiled and the current schema value.

diff --git a/Engine/Plugins/Experimental/AnimNextStateTree/Source/AnimNextStateTree/Private/AnimNextStateTreeSchema.cpp b/Engine/Plugins/Experimental/AnimNextStateTree/Source/AnimNextStateTree/Private/AnimNextStateTreeSchema.cpp
--- a/Engine/Plugins/Experimental/AnimNextStateTree/Source/AnimNextStateTree/Private/AnimNextStateTreeSchema.cpp
+++ b/Engine/Plugins/Experimental/AnimNextStateTree/Source/AnimNextStateTree/Private/AnimNextStateTreeSchema.cpp
@@ -44,4 +44,9 @@ bool UStateTreeAnimNextSchema::IsClassAllowed(const UClass* InClass) const
 bool UStateTreeAnimNextSchema::IsExternalItemAllowed(const UStruct& InStruct) const
 {
 	return InStruct.IsChildOf<FAnimNextStateTreeTraitContext>();
+}
+
+EStateTreeParameterDataType UStateTreeAnimNextSchema::GetGlobalParameterDataType() const
+{
+	return EStateTreeParameterDataType::ExternalGlobalParameterData;
 }
\ No newline at end of file
diff --git a/Engine/Plugins/Experimental/AnimNextStateTree/Source/AnimNextStateTree/Public/AnimNextStateTreeSchema.h b/Engine/Plugins/Experimental/AnimNextStateTree/Source/AnimNextStateTree/Public/AnimNextStateTreeSchema.h
--- a/Engine/Plugins/Experimental/AnimNextStateTree/Source/AnimNextStateTree/Public/AnimNextStateTreeSchema.h
+++ b/Engine/Plugins/Experimental/AnimNextStateTree/Source/AnimNextStateTree/Public/AnimNextStateTreeSchema.h
@@ -21,7 +21,7 @@ public:
 	virtual bool IsExternalItemAllowed(const UStruct& InStruct) const override;
 	
 	virtual TConstArrayView<FStateTreeExternalDataDesc> GetContextDataDescs() const override { return ContextDataDescs; }
-	virtual EGlobalParameterDataType GetGlobalParameterDataType() const override { return EGlobalParameterDataType::ExternalGlobalParameterData; }
+	virtual EStateTreeParameterDataType GetGlobalParameterDataType() const override;
 	
 #if WITH_EDITOR
 	virtual bool AllowUtilityConsiderations() const override { return false; }
diff --git a/Engine/Plugins/Runtime/StateTree/Source/StateTreeEditorModule/Private/StateTreeCompiler.cpp b/Engine/Plugins/Runtime/StateTree/Source/StateTreeEditorModule/Private/StateTreeCompiler.cpp
--- a/Engine/Plugins/Runtime/StateTree/Source/StateTreeEditorModule/Private/StateTreeCompiler.cpp
+++ b/Engine/Plugins/Runtime/StateTree/Source/StateTreeEditorModule/Private/StateTreeCompiler.cpp
@@ -405,9 +405,10 @@ bool FStateTreeCompiler::CreateParameters()
 {
 	// Copy parameters from EditorData	
 	StateTree->Parameters = EditorData->GetRootParametersPropertyBag();
+	StateTree->ParameterDataType = EditorData->Schema->GetGlobalParameterDataType();
 
 	// Mark parameters as binding source
-	const EStateTreeDataSourceType GlobalParameterDataType = UStateTreeSchema::CastToDataSourceType(EditorData->Schema->GetGlobalParameterDataType());
+	const EStateTreeDataSourceType GlobalParameterDataType = UE::StateTree::CastToDataSourceType(StateTree->ParameterDataType);
 	const FStateTreeBindableStructDesc ParametersDesc = {
 			UE::StateTree::Editor::GlobalStateName,
 			TEXT("Parameters"),
diff --git a/Engine/Plugins/Runtime/StateTree/Source/StateTreeModule/Private/StateTree.cpp b/Engine/Plugins/Runtime/StateTree/Source/StateTreeModule/Private/StateTree.cpp
--- a/Engine/Plugins/Runtime/StateTree/Source/StateTreeModule/Private/StateTree.cpp
+++ b/Engine/Plugins/Runtime/StateTree/Source/StateTreeModule/Private/StateTree.cpp
@@ -182,6 +182,7 @@ void UStateTree::ResetCompiled()
 	ContextDataDescs.Reset();
 	PropertyBindings.Reset();
 	Parameters.Reset();
+	ParameterDataType = EStateTreeParameterDataType::GlobalParameterData;
 	IDToStateMappings.Reset();
 	IDToNodeMappings.Reset();
 	IDToTransitionMappings.Reset();
@@ -824,7 +825,7 @@ bool UStateTree::PatchBindings()
 	}
 
 	// Update property bag structs before resolving binding.
-	const EStateTreeDataSourceType GlobalParameterDataType = UStateTreeSchema::CastToDataSourceType(Schema->GetGlobalParameterDataType());
+	const EStateTreeDataSourceType GlobalParameterDataType = UE::StateTree::CastToDataSourceType(ParameterDataType);
 	if (FStateTreeBindableStructDesc* RootParamsDesc = GetSourceStructByHandle(FStateTreeDataHandle(GlobalParameterDataType)))
 	{
 		RootParamsDesc->Struct = Parameters.GetPropertyBagStruct();
diff --git a/Engine/Plugins/Runtime/StateTree/Source/StateTreeModule/Private/StateTreeSchema.cpp b/Engine/Plugins/Runtime/StateTree/Source/StateTreeModule/Private/StateTreeSchema.cpp
--- a/Engine/Plugins/Runtime/StateTree/Source/StateTreeModule/Private/StateTreeSchema.cpp
+++ b/Engine/Plugins/Runtime/StateTree/Source/StateTreeModule/Private/StateTreeSchema.cpp
@@ -11,7 +11,7 @@ bool UStateTreeSchema::IsChildOfBlueprintBase(const UClass* InClass) const
 	return InClass->IsChildOf<UStateTreeNodeBlueprintBase>();
 }
 
-EStateTreeDataSourceType UStateTreeSchema::CastToDataSourceType(EGlobalParameterDataType Value)
+EStateTreeParameterDataType UStateTreeSchema::GetGlobalParameterDataType() const
 {
-	return Value == EGlobalParameterDataType::ExternalGlobalParameterData ? EStateTreeDataSourceType::ExternalGlobalParameterData : EStateTreeDataSourceType::GlobalParameterData;
-}
+	return EStateTreeParameterDataType::GlobalParameterData;
+}
\ No newline at end of file
diff --git a/Engine/Plugins/Runtime/StateTree/Source/StateTreeModule/Public/StateTree.h b/Engine/Plugins/Runtime/StateTree/Source/StateTreeModule/Public/StateTree.h
--- a/Engine/Plugins/Runtime/StateTree/Source/StateTreeModule/Public/StateTree.h
+++ b/Engine/Plugins/Runtime/StateTree/Source/StateTreeModule/Public/StateTree.h
@@ -6,6 +6,7 @@
 #include "StateTreeSchema.h"
 #include "StateTreePropertyBindings.h"
 #include "StateTreeInstanceData.h"
+#include "StateTreeTypes.h"
 #include "UObject/ObjectKey.h"
 #include "StateTree.generated.h"
 
@@ -349,6 +350,10 @@ private:
 	UPROPERTY()
 	uint16 GlobalTasksNum = 0;
 
+	/** The parameter data type used by the schema. */
+	UPROPERTY()
+	EStateTreeParameterDataType ParameterDataType = EStateTreeParameterDataType::GlobalParameterData;
+
 	/** True if any global task is a transition task. */
 	UPROPERTY()
 	uint8 bHasGlobalTransitionTasks : 1 = false;
diff --git a/Engine/Plugins/Runtime/StateTree/Source/StateTreeModule/Public/StateTreeSchema.h b/Engine/Plugins/Runtime/StateTree/Source/StateTreeModule/Public/StateTreeSchema.h
--- a/Engine/Plugins/Runtime/StateTree/Source/StateTreeModule/Public/StateTreeSchema.h
+++ b/Engine/Plugins/Runtime/StateTree/Source/StateTreeModule/Public/StateTreeSchema.h
@@ -8,6 +8,7 @@
 struct FStateTreeExternalDataDesc;
 struct FStateTreeLinker;
 enum class EStateTreeDataSourceType : uint8;
+enum class EStateTreeParameterDataType : uint8;
 
 
 #if WITH_EDITOR
@@ -70,18 +71,8 @@ public:
 		return {};
 	}
 
-	enum class EGlobalParameterDataType : uint8
-	{
-		GlobalParameterData,
-		ExternalGlobalParameterData,
-	};
-	static EStateTreeDataSourceType CastToDataSourceType(EGlobalParameterDataType Value);
-
 	/** @return the global parameter type used by the schema. */
-	virtual EGlobalParameterDataType GetGlobalParameterDataType() const
-	{
-		return EGlobalParameterDataType::GlobalParameterData;
-	}
+	virtual EStateTreeParameterDataType GetGlobalParameterDataType() const;
 
 	/** Resolves schema references to other StateTree data. */
 	virtual bool Link(FStateTreeLinker& Linker)
diff --git a/Engine/Plugins/Runtime/StateTree/Source/StateTreeModule/Public/StateTreeTypes.h b/Engine/Plugins/Runtime/StateTree/Source/StateTreeModule/Public/StateTreeTypes.h
--- a/Engine/Plugins/Runtime/StateTree/Source/StateTreeModule/Public/StateTreeTypes.h
+++ b/Engine/Plugins/Runtime/StateTree/Source/StateTreeModule/Public/StateTreeTypes.h
@@ -325,6 +325,22 @@ enum class EStateTreeDataSourceType : uint8
 	ExternalGlobalParameterData,
 };
 
+/** Data type the global parameter type uses. */
+UENUM()
+enum class EStateTreeParameterDataType : uint8
+{
+	GlobalParameterData,
+	ExternalGlobalParameterData,
+};
+
+namespace UE::StateTree
+{
+	static EStateTreeDataSourceType CastToDataSourceType(EStateTreeParameterDataType Value)
+	{
+		return Value == EStateTreeParameterDataType::ExternalGlobalParameterData ? EStateTreeDataSourceType::ExternalGlobalParameterData : EStateTreeDataSourceType::GlobalParameterData;
+	}
+}
+
 /** Handle to a StateTree data */
 USTRUCT(BlueprintType)
 struct STATETREEMODULE_API FStateTreeDataHandle
