StateTree: Force the UStateTreeEditorData outer to be a UStateTree. When the asset is used or compiled the nodes expect the outer to be a state tree.

diff --git a/Engine/Plugins/Runtime/StateTree/Source/StateTreeEditorModule/Public/StateTreeEditorData.h b/Engine/Plugins/Runtime/StateTree/Source/StateTreeEditorModule/Public/StateTreeEditorData.h
--- a/Engine/Plugins/Runtime/StateTree/Source/StateTreeEditorModule/Public/StateTreeEditorData.h
+++ b/Engine/Plugins/Runtime/StateTree/Source/StateTreeEditorModule/Public/StateTreeEditorData.h
@@ -53,7 +53,7 @@ enum class EStateTreeVisitor : uint8
 /**
  * Edit time data for StateTree asset. This data gets baked into runtime format before being used by the StateTreeInstance.
  */
-UCLASS(BlueprintType, EditInlineNew, CollapseCategories)
+UCLASS(BlueprintType, EditInlineNew, CollapseCategories, Within = "StateTree")
 class STATETREEEDITORMODULE_API UStateTreeEditorData : public UObject, public IStateTreeEditorPropertyBindingsOwner
 {
 	GENERATED_BODY()
@@ -393,3 +393,10 @@ public:
 
 	friend class FStateTreeEditorDataDetails;
 };
+
+
+UCLASS()
+class UQAStateTreeEditorData : public UStateTreeEditorData
+{
+	GENERATED_BODY()
+};
\ No newline at end of file
