Readd comments removed by a bad merge.

#rnx

diff --git a/Engine/Plugins/Runtime/StateTree/Source/StateTreeEditorModule/Private/StateTreeEditingSubsystem.cpp b/Engine/Plugins/Runtime/StateTree/Source/StateTreeEditorModule/Private/StateTreeEditingSubsystem.cpp
--- a/Engine/Plugins/Runtime/StateTree/Source/StateTreeEditorModule/Private/StateTreeEditingSubsystem.cpp
+++ b/Engine/Plugins/Runtime/StateTree/Source/StateTreeEditorModule/Private/StateTreeEditingSubsystem.cpp
@@ -158,6 +158,7 @@ void UStateTreeEditingSubsystem::ValidateStateTree(TNotNull<UStateTree*> InState
 	auto FixEditorData = [](TNotNull<UStateTree*> StateTree)
 	{
 		UStateTreeEditorData* EditorData = Cast<UStateTreeEditorData>(StateTree->EditorData);
+		// The schema is defined in the EditorData. If we can't find the editor data (probably because the class doesn't exist anymore), then try the compiled schema in the state tree asset.
 		TSubclassOf<const UStateTreeSchema> SchemaClass;
 		if (EditorData && EditorData->Schema)
 		{
@@ -187,7 +188,7 @@ void UStateTreeEditingSubsystem::ValidateStateTree(TNotNull<UStateTree*> InState
 		}
 		else if (!EditorData->IsA(EditorDataClass.Get()))
 		{
-			// The current EditorData is not of the correct type. The data need to be patched.
+			// The current EditorData is not of the correct type. The data needs to be patched by the schema desired editor data subclass.
 			UStateTreeEditorData* PreviousEditorData = EditorData;
 			EditorData = CastChecked<UStateTreeEditorData>(StaticDuplicateObject(EditorData, StateTree, FName(), RF_Transactional, EditorDataClass.Get()));
 			if (EditorData->SubTrees.Num() == 0)
@@ -352,14 +353,13 @@ void UStateTreeEditingSubsystem::ValidateStateTree(TNotNull<UStateTree*> InState
 	};
 
 	FixEditorData(InStateTree);
-	{
-		UStateTreeEditorData* TreeData = Cast<UStateTreeEditorData>(InStateTree->EditorData);
-		if (TreeData)
+
+	if (InStateTree->EditorData)
 		{
 			constexpr bool bMarkDirty = false;
-			TreeData->Modify(bMarkDirty);
-		}
+		InStateTree->EditorData->Modify(bMarkDirty);
 	}
+
 	UpdateParents(*InStateTree);
 	ApplySchema(*InStateTree);
 	RemoveUnusedBindings(*InStateTree);
