StateTree: Disable runtime validation on engine shutdown.

diff --git a/Engine/Plugins/Runtime/StateTree/Source/StateTreeModule/Private/Debugger/StateTreeRuntimeValidation.cpp b/Engine/Plugins/Runtime/StateTree/Source/StateTreeModule/Private/Debugger/StateTreeRuntimeValidation.cpp
--- a/Engine/Plugins/Runtime/StateTree/Source/StateTreeModule/Private/Debugger/StateTreeRuntimeValidation.cpp
+++ b/Engine/Plugins/Runtime/StateTree/Source/StateTreeModule/Private/Debugger/StateTreeRuntimeValidation.cpp
@@ -51,7 +51,7 @@ FString NodeToString(const UObject* Obj, FGuid Id)
 
 FRuntimeValidationInstanceData::~FRuntimeValidationInstanceData()
 {
-	if (Private::bRuntimeValidationEnterExitState)
+	if (Private::bRuntimeValidationEnterExitState && UObjectInitialized() && !IsEngineExitRequested())
 	{
 		for (const FNodeStatePair& Pair : NodeStates)
 		{
