[State Tree Editor] Editor code refactorings and first pass of buttons for in-place node modifications: Copy, Paste, Duplicate and Deletes
- `FStateTreeScopedDataFixer` to centralize code for fixing up data after editor commands.  Shared between DetailsView and ViewModel.
- `SStateTreeContextMenuButton` to highlight the property in the DetailsView and bring up a menu for node manipulation
- `FStateTreeStateNodeEditorHandle` a helper struct to facilitate node manipulations in ViewModel
- fix for paste command not correctly transacted
- renaming for typos
- `StateTreeViewModel` Node commands
- Dropped comments for to-dos and cautions
- some misc refactorings and helpers

diff --git a/Engine/Plugins/Runtime/StateTree/Config/BaseStateTree.ini b/Engine/Plugins/Runtime/StateTree/Config/BaseStateTree.ini
--- a/Engine/Plugins/Runtime/StateTree/Config/BaseStateTree.ini
+++ b/Engine/Plugins/Runtime/StateTree/Config/BaseStateTree.ini
@@ -69,6 +69,7 @@
 
 +EnumRedirects=(OldName="/Script/StateTreeModule.EStateTreeDataSourceType",ValueChanges=(("LinkedStateParameterData", "StateParameterData")))
 +EnumRedirects=(OldName="/Script/StateTreeModule.EStateTreeBindableStructSource",ValueChanges=(("State", "StateParameter")))
++EnumRedirects=(OldName="/Script/StateTreeEditorModule.FStateTreeViewModelInsert",NewName="EStateTreeViewModelInsert")
 
 +EnumRedirects=(OldName="EStateTreeConditionOperand",NewName="/Script/StateTreeModule.EStateTreeExpressionOperand")
 +PropertyRedirects=(OldName="StateTreeEditorNode.ConditionIndent",NewName="StateTreeEditorNode.ExpressionIndent")
diff --git a/Engine/Plugins/Runtime/StateTree/Source/StateTreeEditorModule/Private/Customizations/StateTreeEditorNodeDetails.cpp b/Engine/Plugins/Runtime/StateTree/Source/StateTreeEditorModule/Private/Customizations/StateTreeEditorNodeDetails.cpp
--- a/Engine/Plugins/Runtime/StateTree/Source/StateTreeEditorModule/Private/Customizations/StateTreeEditorNodeDetails.cpp
+++ b/Engine/Plugins/Runtime/StateTree/Source/StateTreeEditorModule/Private/Customizations/StateTreeEditorNodeDetails.cpp
@@ -13,7 +13,6 @@
 #include "StateTreePropertyRef.h"
 #include "Widgets/Input/SButton.h"
 #include "Widgets/Input/SComboButton.h"
-#include "Widgets/Notifications/SNotificationList.h"
 #include "Widgets/Text/SInlineEditableTextBlock.h"
 #include "Widgets/Text/SRichTextBlock.h"
 #include "Widgets/Layout/SWidgetSwitcher.h"
@@ -34,7 +33,6 @@
 #include "HAL/PlatformApplicationMisc.h"
 #include "ScopedTransaction.h"
 #include "StateTreeEditorNodeUtils.h"
-#include "Framework/Notifications/NotificationManager.h"
 #include "Debugger/StateTreeDebuggerUIExtensions.h"
 #include "StateTreeEditingSubsystem.h"
 #include "StateTreePropertyBindings.h"
@@ -43,28 +41,10 @@
 #include "TextStyleDecorator.h"
 #include "Framework/Application/SlateApplication.h"
 #include "StateTreeDelegate.h"
+#include "StateTreeScopedEditorDataFixer.h"
 
 #define LOCTEXT_NAMESPACE "StateTreeEditor"
 
-/** Helper class to detect if there were issues when calling ImportText() */
-class FStateTreeDefaultValueImportErrorContext : public FOutputDevice
-{
-public:
-
-	int32 NumErrors;
-
-	FStateTreeDefaultValueImportErrorContext()
-		: FOutputDevice()
-		, NumErrors(0)
-	{
-	}
-
-	virtual void Serialize(const TCHAR* V, ELogVerbosity::Type Verbosity, const class FName& Category) override
-	{
-		NumErrors++;
-	}
-};
-
 namespace UE::StateTreeEditor::Internal
 {
 	/* Returns true if provided property is direct or indirect child of PropertyFunction */
@@ -821,126 +801,56 @@ void FStateTreeEditorNodeDetails::OnCopyNode()
 
 void FStateTreeEditorNodeDetails::OnPasteNode()
 {
-	FString PastedText;
-	FPlatformApplicationMisc::ClipboardPaste(PastedText);
-
-	if (PastedText.IsEmpty())
-	{
-		return;
-	}
-	
-	// Create node from the clipboard data to figure out the node type. 
-	FStateTreeEditorNode TempNode;
-	UScriptStruct* NodeScriptStruct = TBaseStructure<FStateTreeEditorNode>::Get();
-
-	TArray<UObject*> OuterObjects; 
+	TArray<UObject*> OuterObjects;
 	StructProperty->GetOuterObjects(OuterObjects);
 	if (OuterObjects.IsEmpty())
 	{
 		return;
 	}
 
-	FStateTreeDefaultValueImportErrorContext ErrorPipe;
-	NodeScriptStruct->ImportText(*PastedText, &TempNode, nullptr, PPF_None, &ErrorPipe, NodeScriptStruct->GetName());
-	
-	const UStruct* NodeTypeStruct = TempNode.Node.GetScriptStruct(); 
-	// Only allow valid node types for this property (e.g. do not mix task with conditions).
-	if (ErrorPipe.NumErrors > 0 || !NodeTypeStruct || !NodeTypeStruct->IsChildOf(BaseScriptStruct.Get()))
-	{
-		FNotificationInfo NotificationInfo(FText::GetEmpty());
-		NotificationInfo.Text = FText::Format(LOCTEXT("NotSupportedByType", "This property only accepts nodes of type {0}."), BaseScriptStruct->GetDisplayNameText());
-		NotificationInfo.ExpireDuration = 5.0f;
-		FSlateNotificationManager::Get().AddNotification(NotificationInfo);
-		return;
-	}
-	
-	// Reject nodes that are not allowed by the schema.
-	UStateTreeEditorData* EditorDataPtr = EditorData.Get();
-	const UStateTreeSchema* Schema = EditorDataPtr ? EditorDataPtr->Schema : nullptr;
-	if (Schema)
-	{
-		bool bNodeIsAllowed = false;
-		
-		// BP nodes are identified by the instance type.
-		if (NodeTypeStruct->IsChildOf(FStateTreeBlueprintEvaluatorWrapper::StaticStruct())
-			|| NodeTypeStruct->IsChildOf(FStateTreeBlueprintTaskWrapper::StaticStruct())
-			|| NodeTypeStruct->IsChildOf(FStateTreeBlueprintConditionWrapper::StaticStruct())
-			|| NodeTypeStruct->IsChildOf(FStateTreeBlueprintConsiderationWrapper::StaticStruct()))
-		{
-			if (const FStateTreeNodeBase* Node = TempNode.Node.GetPtr<FStateTreeNodeBase>())
-			{
-				NodeTypeStruct = Node->GetInstanceDataType(); // Report error with the BP node type, as that is what the user expects to see.
-				if (const UClass* InstanceClass = Cast<UClass>(NodeTypeStruct))
-				{
-					bNodeIsAllowed = Schema->IsClassAllowed(InstanceClass);
-				}
-			}
-		}
-		else
-		{
-			bNodeIsAllowed = Schema->IsStructAllowed(TempNode.Node.GetScriptStruct());
-		}
-		
-		if (!bNodeIsAllowed)
-		{
+	FStateTreeEditorNode TempNode;
+	bool bSuccess = UE::StateTreeEditor::EditorNodeUtils::ImportTextAsNode(BaseScriptStruct.Get(), EditorData.Get(), TempNode);
 
-			FNotificationInfo NotificationInfo(FText::GetEmpty());
-			NotificationInfo.Text = FText::Format(LOCTEXT("NotSupportedBySchema", "Node {0} is not supported by {1} schema."),
-										NodeTypeStruct->GetDisplayNameText(), Schema->GetClass()->GetDisplayNameText());
-			NotificationInfo.ExpireDuration = 5.0f;
-			FSlateNotificationManager::Get().AddNotification(NotificationInfo);
-			return;
-		}
+	if (!bSuccess)
+	{
+		return;
 	}
-	
-	FScopedTransaction Transaction(LOCTEXT("PasteNode", "Paste Node"));
 
-	StructProperty->NotifyPreChange();
-
-	// Reset GUIDs on paste
 	TArray<void*> RawNodeData;
 	StructProperty->AccessRawData(RawNodeData);
 	if (OuterObjects.Num() == RawNodeData.Num())
 	{
-		for (int32 Index = 0; Index < RawNodeData.Num(); Index++)
+		if (UStateTreeEditorData* EditorDataPtr = EditorData.Get())
 		{
-			UObject* OuterObject = OuterObjects[Index];
-			FStateTreeEditorNode* EditorNode = static_cast<FStateTreeEditorNode*>(RawNodeData[Index]);
-			if (EditorNode && OuterObject)
-			{
-				// Copy
-				*EditorNode = TempNode;
+			FScopedTransaction Transaction(LOCTEXT("PasteNode", "Paste Node"));
 
-				// Ensure unique instance value
-				UE::StateTreeEditor::EditorNodeUtils::InstantiateStructSubobjects(*OuterObject, EditorNode->Node);
-				if (EditorNode->InstanceObject)
-				{
-					EditorNode->InstanceObject = DuplicateObject(EditorNode->InstanceObject, OuterObject);
-				}
-				else
-				{
-					UE::StateTreeEditor::EditorNodeUtils::InstantiateStructSubobjects(*OuterObject, EditorNode->Instance);
-				}
-				
-				const FGuid OldStructID = EditorNode->ID; 
-				EditorNode->ID = FGuid::NewGuid();
+			StructProperty->NotifyPreChange();
+
+			// we will likely change the bindings on the Editor Data or the editor node could be global, so record
+			EditorDataPtr->Modify();
 
-				// Copy bindings from the copied node.
-				if (OldStructID.IsValid() && EditorDataPtr)
+			{
+				UE::StateTreeEditor::FScopedEditorDataFixer DataFixer(EditorDataPtr);
+
+				for (int32 Index = 0; Index < RawNodeData.Num(); Index++)
 				{
-					if (FStateTreeEditorPropertyBindings* Bindings = EditorDataPtr->GetPropertyEditorBindings())
+					if (FStateTreeEditorNode* EditorNode = static_cast<FStateTreeEditorNode*>(RawNodeData[Index]))
 					{
-						Bindings->CopyBindings(OldStructID, EditorNode->ID);
+						*EditorNode = TempNode;
+
+						constexpr bool bShouldCopyBindings = true;
+						constexpr bool bShouldReinstantiateObjects = true;
+						constexpr bool bShouldRegenerateGUID = true;
+						DataFixer.EditorNodesToFix.Emplace(OuterObjects[Index], *EditorNode, bShouldCopyBindings, bShouldReinstantiateObjects, bShouldRegenerateGUID);
 					}
 				}
-
 			}
+
+			StructProperty->NotifyPostChange(EPropertyChangeType::ValueSet);
+			StructProperty->NotifyFinishedChangingProperties();
 		}
 	}
 
-	StructProperty->NotifyPostChange(EPropertyChangeType::ValueSet);
-	StructProperty->NotifyFinishedChangingProperties();
-
 	if (PropUtils)
 	{
 		PropUtils->ForceRefresh();
@@ -1910,7 +1820,7 @@ TSharedRef<SWidget> FStateTreeEditorNodeDetails::GenerateOptionsMenu()
 			FCanExecuteAction()
 		));
 
-	// Delete
+	// Delete All
 	MenuBuilder.AddMenuEntry(
 		LOCTEXT("DeleteAllItems", "Delete all"),
 		LOCTEXT("DeleteAllItemsTooltip", "Delete all items"),
@@ -1945,7 +1855,17 @@ void FStateTreeEditorNodeDetails::OnDeleteNode() const
 	{
 		if (const TSharedPtr<IPropertyHandleArray> ArrayHandle = ParentHandle->AsArray())
 		{
-			ArrayHandle->DeleteItem(Index);
+			if (UStateTreeEditorData* EditorDataPtr = EditorData.Get())
+			{
+				FScopedTransaction Transaction(LOCTEXT("DeleteNode", "Delete Node"));
+
+				EditorDataPtr->Modify();
+
+				ArrayHandle->DeleteItem(Index);
+
+				UE::StateTreeEditor::FScopedEditorDataFixer DataFixer(EditorDataPtr);
+				DataFixer.bRemoveInvalidBindings = true;
+			}
 		}
 	}
 }
@@ -1956,19 +1876,61 @@ void FStateTreeEditorNodeDetails::OnDeleteAllNodes() const
 	{
 		if (const TSharedPtr<IPropertyHandleArray> ArrayHandle = ParentHandle->AsArray())
 		{
-			ArrayHandle->EmptyArray();
+			if (UStateTreeEditorData* EditorDataPtr = EditorData.Get())
+			{
+				FScopedTransaction Transaction(LOCTEXT("DeleteAllNodes", "Delete All Nodes"));
+
+				EditorDataPtr->Modify();
+
+				ArrayHandle->EmptyArray();
+
+				UE::StateTreeEditor::FScopedEditorDataFixer DataFixer(EditorDataPtr);
+				DataFixer.bRemoveInvalidBindings = true;
+			}
 		}
 	}
 }
 
 void FStateTreeEditorNodeDetails::OnDuplicateNode() const
 {
-	const int32 Index = StructProperty->GetArrayIndex();
+	TArray<UObject*> OuterObjects;
+	StructProperty->GetOuterObjects(OuterObjects);
+	if (OuterObjects.IsEmpty())
+	{
+		return;
+	}
+
+	const int32 ArrayIndex = StructProperty->GetArrayIndex();
 	if (const TSharedPtr<IPropertyHandle> ParentHandle = StructProperty->GetParentHandle())
 	{
 		if (const TSharedPtr<IPropertyHandleArray> ArrayHandle = ParentHandle->AsArray())
 		{
-			ArrayHandle->DuplicateItem(Index);
+			if (UStateTreeEditorData* EditorDataPtr = EditorData.Get())
+			{
+				FScopedTransaction Transaction(LOCTEXT("DuplicateNode", "Duplicate Node"));
+
+				// Might modify the bindings data
+				EditorDataPtr->Modify();
+
+				ArrayHandle->DuplicateItem(ArrayIndex);
+				
+				TSharedRef<IPropertyHandle> DuplicatedStructHandle = ArrayHandle->GetElement(ArrayIndex);
+
+				UE::StateTreeEditor::FScopedEditorDataFixer DataFixer(EditorDataPtr);
+
+				TArray<void*> RawNodeData;
+				StructProperty->AccessRawData(RawNodeData);
+				for (int32 Index = 0; Index < RawNodeData.Num(); ++Index)
+				{
+					if (FStateTreeEditorNode* EditorNode = static_cast<FStateTreeEditorNode*>(RawNodeData[Index]))
+					{
+						constexpr bool bShouldCopyBindings = true;
+						constexpr bool bShouldRegenerateGUID = true;
+						constexpr bool bShouldRenstantiateSubObjects = false; // ArrayHandle Duplication has done deep copies for us
+						DataFixer.EditorNodesToFix.Emplace(OuterObjects[Index], *EditorNode, bShouldCopyBindings, bShouldRenstantiateSubObjects, bShouldRegenerateGUID);
+					}
+				}
+			}
 		}
 	}
 }
@@ -1989,6 +1951,7 @@ void FStateTreeEditorNodeDetails::OnRenameNode() const
 	}
 }
 
+// @todo: refactor it to use FStateTreeEditorDataFixer
 void FStateTreeEditorNodeDetails::OnNodePicked(const UStruct* InStruct) const
 {
 	GEditor->BeginTransaction(LOCTEXT("SelectNode", "Select Node"));
diff --git a/Engine/Plugins/Runtime/StateTree/Source/StateTreeEditorModule/Private/Customizations/StateTreeEditorNodeUtils.cpp b/Engine/Plugins/Runtime/StateTree/Source/StateTreeEditorModule/Private/Customizations/StateTreeEditorNodeUtils.cpp
--- a/Engine/Plugins/Runtime/StateTree/Source/StateTreeEditorModule/Private/Customizations/StateTreeEditorNodeUtils.cpp
+++ b/Engine/Plugins/Runtime/StateTree/Source/StateTreeEditorModule/Private/Customizations/StateTreeEditorNodeUtils.cpp
@@ -13,6 +13,7 @@
 #include "StateTreeEditorData.h"
 #include "StateTreeEditorNode.h"
 #include "StateTreeEditorSettings.h"
+#include "HAL/PlatformApplicationMisc.h"
 #include "StateTreeEditorStyle.h"
 #include "StateTreeTaskBase.h"
 #include "Blueprint/StateTreeConditionBlueprintBase.h"
@@ -23,12 +24,29 @@
 #include "Widgets/Input/SComboButton.h"
 #include "Widgets/Input/SButton.h"
 #include "Editor.h"
+#include "Framework/Notifications/NotificationManager.h"
+#include "Widgets/Notifications/SNotificationList.h"
 
 #define LOCTEXT_NAMESPACE "StateTreeEditor"
 
 namespace UE::StateTreeEditor::EditorNodeUtils
 {
 
+/** Helper class to detect if there were issues when calling ImportText() */
+class FDefaultValueImportErrorContext : public FOutputDevice
+{
+public:
+
+	int32 NumErrors = 0;
+
+	FDefaultValueImportErrorContext() = default;
+
+	virtual void Serialize(const TCHAR* V, ELogVerbosity::Type Verbosity, const class FName& Category) override
+	{
+		++NumErrors;
+	}
+};
+
 EStateTreeConditionEvaluationMode GetConditionEvaluationMode(const TSharedPtr<IPropertyHandle>& StructProperty)
 {
 	if (const FStateTreeEditorNode* Node = GetCommonNode(StructProperty))
@@ -852,6 +870,71 @@ IDetailCategoryBuilder& MakeArrayCategory(
 	return Category;
 }
 
+bool ImportTextAsNode(const UScriptStruct* BaseScriptStruct, const UStateTreeEditorData* EditorData, FStateTreeEditorNode& OutEditorNode)
+{
+	FString PastedText;
+	FPlatformApplicationMisc::ClipboardPaste(PastedText);
+
+	if (PastedText.IsEmpty())
+	{
+		return false;
+	}
+
+	UScriptStruct* NodeScriptStruct = TBaseStructure<FStateTreeEditorNode>::Get();
+	FDefaultValueImportErrorContext ErrorPipe;
+	NodeScriptStruct->ImportText(*PastedText, &OutEditorNode, nullptr, PPF_None, &ErrorPipe, NodeScriptStruct->GetName());
+
+	const UStruct* NodeTypeStruct = OutEditorNode.Node.GetScriptStruct();
+	// Only allow valid node types for this property (e.g. do not mix task with conditions).
+	if (ErrorPipe.NumErrors > 0 || !NodeTypeStruct || !NodeTypeStruct->IsChildOf(BaseScriptStruct))
+	{
+		FNotificationInfo NotificationInfo(FText::GetEmpty());
+		NotificationInfo.Text = FText::Format(LOCTEXT("NotSupportedByType", "This property only accepts nodes of type {0}."), BaseScriptStruct->GetDisplayNameText());
+		NotificationInfo.ExpireDuration = 5.0f;
+		FSlateNotificationManager::Get().AddNotification(NotificationInfo);
+		return false;
+	}
+
+	const UStateTreeSchema* Schema = EditorData ? EditorData->Schema : nullptr;
+	if (Schema)
+	{
+		bool bNodeIsAllowed = false;
+
+		// BP nodes are identified by the instance type.
+		if (NodeTypeStruct->IsChildOf(FStateTreeBlueprintEvaluatorWrapper::StaticStruct())
+			|| NodeTypeStruct->IsChildOf(FStateTreeBlueprintTaskWrapper::StaticStruct())
+			|| NodeTypeStruct->IsChildOf(FStateTreeBlueprintConditionWrapper::StaticStruct())
+			|| NodeTypeStruct->IsChildOf(FStateTreeBlueprintConsiderationWrapper::StaticStruct()))
+		{
+			if (const FStateTreeNodeBase* Node = OutEditorNode.Node.GetPtr<FStateTreeNodeBase>())
+			{
+				NodeTypeStruct = Node->GetInstanceDataType(); // Report error with the BP node type, as that is what the user expects to see.
+				if (const UClass* InstanceClass = Cast<UClass>(NodeTypeStruct))
+				{
+					bNodeIsAllowed = Schema->IsClassAllowed(InstanceClass);
+				}
+			}
+		}
+		else
+		{
+			bNodeIsAllowed = Schema->IsStructAllowed(OutEditorNode.Node.GetScriptStruct());
+		}
+
+		if (!bNodeIsAllowed)
+		{
+
+			FNotificationInfo NotificationInfo(FText::GetEmpty());
+			NotificationInfo.Text = FText::Format(LOCTEXT("NotSupportedBySchema", "Node {0} is not supported by {1} schema."),
+												  NodeTypeStruct->GetDisplayNameText(), Schema->GetClass()->GetDisplayNameText());
+			NotificationInfo.ExpireDuration = 5.0f;
+			FSlateNotificationManager::Get().AddNotification(NotificationInfo);
+			return false;
+		}
+	}
+
+	// No schema, any types would be accepted
+	return true;
+}
 } // namespace UE::StateTreeEditor::EditorNodeUtils
 
 #undef LOCTEXT_NAMESPACE
diff --git a/Engine/Plugins/Runtime/StateTree/Source/StateTreeEditorModule/Private/Customizations/StateTreeEditorNodeUtils.h b/Engine/Plugins/Runtime/StateTree/Source/StateTreeEditorModule/Private/Customizations/StateTreeEditorNodeUtils.h
--- a/Engine/Plugins/Runtime/StateTree/Source/StateTreeEditorModule/Private/Customizations/StateTreeEditorNodeUtils.h
+++ b/Engine/Plugins/Runtime/StateTree/Source/StateTreeEditorModule/Private/Customizations/StateTreeEditorNodeUtils.h
@@ -6,6 +6,7 @@
 #include "Textures/SlateIcon.h"
 #include "Styling/SlateColor.h"
 
+class UStateTreeEditorData;
 enum class EStateTreeConditionEvaluationMode : uint8;
 struct FStateTreeEditorNode;
 struct FStructView;
@@ -159,4 +160,12 @@ IDetailCategoryBuilder& MakeArrayCategory(
 	const FText& AddButtonTooltipText,
 	int32 SortOrder);
 
+/**
+ * Import the text from clipboard as an Editor Node
+ * @param BaseScriptStruct the Base Struct of the Editor Node's Node
+ * @param EditorData Editor Data of the Node
+ * @param OutEditorNode Editor Node reference to which the text will be imported
+ * @return Import succeeded or not
+ */
+bool ImportTextAsNode(const UScriptStruct* BaseScriptStruct, const UStateTreeEditorData* EditorData, FStateTreeEditorNode& OutEditorNode);
 }; // UE::StateTreeEditor::EditorNodeUtils
diff --git a/Engine/Plugins/Runtime/StateTree/Source/StateTreeEditorModule/Private/Customizations/StateTreeTransitionDetails.cpp b/Engine/Plugins/Runtime/StateTree/Source/StateTreeEditorModule/Private/Customizations/StateTreeTransitionDetails.cpp
--- a/Engine/Plugins/Runtime/StateTree/Source/StateTreeEditorModule/Private/Customizations/StateTreeTransitionDetails.cpp
+++ b/Engine/Plugins/Runtime/StateTree/Source/StateTreeEditorModule/Private/Customizations/StateTreeTransitionDetails.cpp
@@ -25,6 +25,7 @@
 #include "TextStyleDecorator.h"
 #include "Framework/Application/SlateApplication.h"
 #include "StateTreeBindingExtension.h"
+#include "StateTreeScopedEditorDataFixer.h"
 
 #define LOCTEXT_NAMESPACE "StateTreeEditor"
 
@@ -354,7 +355,16 @@ void FStateTreeTransitionDetails::OnDeleteTransition() const
 	{
 		if (const TSharedPtr<IPropertyHandleArray> ArrayHandle = ParentHandle->AsArray())
 		{
-			ArrayHandle->DeleteItem(Index);
+			if (UStateTreeEditorData* EditorData = GetEditorData())
+			{
+				FScopedTransaction Transaction(LOCTEXT("DeleteTransition", "Delete Transition"));
+				EditorData->Modify();
+
+				ArrayHandle->DeleteItem(Index);
+
+				UE::StateTreeEditor::FScopedEditorDataFixer DataFixer(EditorData);
+				DataFixer.bRemoveInvalidBindings = true;
+			}
 		}
 	}
 }
@@ -365,19 +375,61 @@ void FStateTreeTransitionDetails::OnDeleteAllTransitions() const
 	{
 		if (const TSharedPtr<IPropertyHandleArray> ArrayHandle = ParentHandle->AsArray())
 		{
-			ArrayHandle->EmptyArray();
+			if (UStateTreeEditorData* EditorData = GetEditorData())
+			{
+				FScopedTransaction Transaction(LOCTEXT("DeleteAllTransitions", "Delete All Transitions"));
+				EditorData->Modify();
+
+				ArrayHandle->EmptyArray();
+
+				UE::StateTreeEditor::FScopedEditorDataFixer DataFixer(EditorData);
+				DataFixer.bRemoveInvalidBindings = true;
+			}
 		}
 	}
 }
 
 void FStateTreeTransitionDetails::OnDuplicateTransition() const
 {
-	const int32 Index = StructProperty->GetArrayIndex();
+	TArray<UObject*> OuterObjects;
+	StructProperty->GetOuterObjects(OuterObjects);
+	if (OuterObjects.IsEmpty())
+	{
+		return;
+	}
+
+	const int32 ArrayIndex = StructProperty->GetArrayIndex();
 	if (const TSharedPtr<IPropertyHandle> ParentHandle = StructProperty->GetParentHandle())
 	{
 		if (const TSharedPtr<IPropertyHandleArray> ArrayHandle = ParentHandle->AsArray())
 		{
-			ArrayHandle->DuplicateItem(Index);
+			if (UStateTreeEditorData* EditorData = GetEditorData())
+			{
+				FScopedTransaction Transaction(LOCTEXT("DuplicateTransition", "Duplicate Transition"));
+
+				EditorData->Modify();
+
+				ArrayHandle->DuplicateItem(ArrayIndex);
+
+				TSharedRef<IPropertyHandle> DuplicatedStructHandle = ArrayHandle->GetElement(ArrayIndex);
+
+				{
+					UE::StateTreeEditor::FScopedEditorDataFixer DataFixer(EditorData);
+
+					TArray<void*> RawNodeData;
+					StructProperty->AccessRawData(RawNodeData);
+					for (int32 Index = 0; Index < RawNodeData.Num(); ++Index)
+					{
+						if (FStateTreeTransition* Transition = static_cast<FStateTreeTransition*>(RawNodeData[Index]))
+						{
+							constexpr bool bShouldCopyBindings = true;
+							constexpr bool bShouldReinstantiateInstanceObjects = false; // ArrayHandle has already deep copied UObjects
+							constexpr bool bShouldRegenerateGUID = true;
+							DataFixer.TransitionsToFix.Emplace(OuterObjects[Index], *Transition, bShouldCopyBindings, bShouldReinstantiateInstanceObjects, bShouldRegenerateGUID);
+						}
+					}
+				}
+			}
 		}
 	}
 }
@@ -448,13 +500,9 @@ UStateTreeEditorData* FStateTreeTransitionDetails::GetEditorData() const
 
 void FStateTreeTransitionDetails::OnPasteTransition() const
 {
-	UStateTreeEditorData* EditorData = GetEditorData();
-	if (!EditorData)
-	{
-		return;
-	}
-	FStateTreeEditorPropertyBindings* Bindings = EditorData->GetPropertyEditorBindings();
-	if (!Bindings)
+	TArray<UObject*> OuterObjects;
+	StructProperty->GetOuterObjects(OuterObjects);
+	if (OuterObjects.IsEmpty())
 	{
 		return;
 	}
@@ -467,36 +515,38 @@ void FStateTreeTransitionDetails::OnPasteTransition() const
 		return;
 	}
 
-	FScopedTransaction Transaction(LOCTEXT("PasteTransition", "Paste Transition"));
+	if (UStateTreeEditorData* EditorData = GetEditorData())
+	{
+		FScopedTransaction Transaction(LOCTEXT("PasteTransition", "Paste Transition"));
+
+		TArray<void*> RawData;
+		StructProperty->AccessRawData(RawData);
 
-	StructProperty->NotifyPreChange();
+		// Might modify the bindings
+		EditorData->Modify();
 
-	// Make sure we instantiate new objects when setting the value.
-	StructProperty->SetValueFromFormattedString(PastedText, EPropertyValueSetFlags::InstanceObjects);
+		// Make sure we instantiate new objects when setting the value.
+		// Calls NotifyPreChange for us, will also call NotifyFinishedChangingProperties
+		StructProperty->SetValueFromFormattedString(PastedText, EPropertyValueSetFlags::InstanceObjects);
 
-	// Reset GUIDs on paste
-	TArray<void*> RawData;
-	StructProperty->AccessRawData(RawData);
-	for (int32 Index = 0; Index < RawData.Num(); Index++)
-	{
-		if (FStateTreeTransition* Transition = static_cast<FStateTreeTransition*>(RawData[Index]))
 		{
-			Transition->ID = FGuid::NewGuid();
+			UE::StateTreeEditor::FScopedEditorDataFixer DataFixer(EditorData);
 
-			for (FStateTreeEditorNode& Condition : Transition->Conditions)
+			for (int32 Index = 0; Index < RawData.Num(); Index++)
 			{
-				const FGuid OldStructID = Condition.ID;
-				Condition.ID = FGuid::NewGuid();
-				if (OldStructID.IsValid())
+				if (FStateTreeTransition* Transition = static_cast<FStateTreeTransition*>(RawData[Index]))
 				{
-					Bindings->CopyBindings(OldStructID, Condition.ID);
+					constexpr bool bShouldCopyBindings = true;
+					constexpr bool bShouldReinstantiateInstanceData = false; // StructProperty has already done UObject deep copies
+					constexpr bool bShouldRegenerateGUID = true;
+					DataFixer.TransitionsToFix.Emplace(OuterObjects[Index], *Transition, bShouldCopyBindings, bShouldReinstantiateInstanceData, bShouldRegenerateGUID);
 				}
 			}
 		}
-	}
 
-	StructProperty->NotifyPostChange(EPropertyChangeType::ValueSet);
-	StructProperty->NotifyFinishedChangingProperties();
+		StructProperty->NotifyPostChange(EPropertyChangeType::ValueSet);
+		StructProperty->NotifyFinishedChangingProperties();
+	}
 
 	if (PropUtils)
 	{
diff --git a/Engine/Plugins/Runtime/StateTree/Source/StateTreeEditorModule/Private/Customizations/Widgets/SStateTreeContextMenuButton.cpp b/Engine/Plugins/Runtime/StateTree/Source/StateTreeEditorModule/Private/Customizations/Widgets/SStateTreeContextMenuButton.cpp
new file mode 100644
--- /dev/null
+++ b/Engine/Plugins/Runtime/StateTree/Source/StateTreeEditorModule/Private/Customizations/Widgets/SStateTreeContextMenuButton.cpp
@@ -0,0 +1,176 @@
+// Copyright Epic Games, Inc. All Rights Reserved.
+
+#include "SStateTreeContextMenuButton.h"
+
+#include "Framework/Application/SlateApplication.h"
+#include "Framework/MultiBox/MultiBoxBuilder.h"
+#include "Widgets/Input/SMenuAnchor.h"
+#include "StateTreeState.h"
+#include "StateTreeViewModel.h"
+
+#define LOCTEXT_NAMESPACE "StateTreeEditor"
+
+void SStateTreeContextMenuButton::Construct(const FArguments& InArgs, const TSharedRef<FStateTreeViewModel>& InStateTreeViewModel, TWeakObjectPtr<UStateTreeState> InOwnerState, const FGuid& InNodeID, bool InbIsTransition)
+{
+	StateTreeViewModel = InStateTreeViewModel.ToSharedPtr();
+	OwnerStateWeak = InOwnerState;
+	NodeID = InNodeID;
+
+	bIsStateTransition = false;
+	bIsTransition = InbIsTransition;
+	if (bIsTransition)
+	{
+		if (UStateTreeState* OwnerState = OwnerStateWeak.Get())
+		{
+			for (const FStateTreeTransition& Transition : OwnerState->Transitions)
+			{
+				if (NodeID == Transition.ID)
+				{
+					bIsStateTransition = true;
+					break;
+				}
+			}
+		}
+	}
+
+	SButton::FArguments ButtonArgs;
+
+	ButtonArgs
+	.OnClicked_Lambda([this]()
+	{
+		StateTreeViewModel->BringNodeToFocus(OwnerStateWeak.Get(), NodeID);
+		return FReply::Handled();
+	})
+	.ButtonStyle(InArgs._ButtonStyle)
+	.ContentPadding(InArgs._ContentPadding)
+	[
+		SAssignNew(MenuAnchor, SMenuAnchor)
+		.Placement(MenuPlacement_BelowAnchor)
+		.OnGetMenuContent(FOnGetContent::CreateSP(this, &SStateTreeContextMenuButton::MakeContextMenu))
+		[
+			InArgs._Content.Widget
+		]
+	];
+
+	SButton::Construct(ButtonArgs);
+}
+
+FReply SStateTreeContextMenuButton::OnMouseButtonUp(const FGeometry& MyGeometry, const FPointerEvent& MouseEvent)
+{
+	StateTreeViewModel->BringNodeToFocus(OwnerStateWeak.Get(), NodeID);
+
+	if (MouseEvent.GetEffectingButton() == EKeys::RightMouseButton)
+	{
+		if (MenuAnchor.IsValid())
+		{
+			if (MenuAnchor->ShouldOpenDueToClick())
+			{
+				MenuAnchor->SetIsOpen(true);
+			}
+			else
+			{
+				MenuAnchor->SetIsOpen(false);
+			}
+
+			return FReply::Handled();
+		}
+	}
+
+	return SButton::OnMouseButtonUp(MyGeometry, MouseEvent);
+}
+
+TSharedRef<SWidget> SStateTreeContextMenuButton::MakeContextMenu() const
+{
+	FMenuBuilder MenuBuilder(/*ShouldCloseWindowAfterMenuSelection*/true, /*CommandList*/nullptr);
+
+	if (StateTreeViewModel.IsValid() && OwnerStateWeak.IsValid())
+	{
+		MenuBuilder.BeginSection(FName("Edit"), LOCTEXT("Edit", "Edit"));
+
+		// Copy
+		MenuBuilder.AddMenuEntry(
+			LOCTEXT("CopyItem", "Copy"),
+			LOCTEXT("CopyItemTooltip", "Copy this item"),
+			FSlateIcon(FAppStyle::GetAppStyleSetName(), "GenericCommands.Copy"),
+			FUIAction(
+			FExecuteAction::CreateSPLambda(this, [this]()
+			{
+				StateTreeViewModel->CopyNode(OwnerStateWeak, NodeID);
+			}),
+			FCanExecuteAction::CreateSPLambda(this, [this]()
+			{
+				return !bIsTransition || bIsStateTransition;
+			})
+		));
+
+		// Paste
+		MenuBuilder.AddMenuEntry(
+			LOCTEXT("PasteItem", "Paste"),
+			LOCTEXT("PasteItemTooltip", "Paste into this item"),
+			FSlateIcon(FAppStyle::GetAppStyleSetName(), "GenericCommands.Paste"),
+			FUIAction(
+			FExecuteAction::CreateSPLambda(this, [this]()
+			{
+				StateTreeViewModel->PasteNode(OwnerStateWeak, NodeID);
+			}),
+			FCanExecuteAction::CreateSPLambda(this, [this]()
+			{
+				return !bIsTransition || bIsStateTransition;
+			})
+		));
+
+		// Duplicate
+		MenuBuilder.AddMenuEntry(
+			LOCTEXT("DuplicateItem", "Duplicate"),
+			LOCTEXT("DuplicateItemTooltip", "Duplicate this item"),
+			FSlateIcon(FAppStyle::GetAppStyleSetName(), "GenericCommands.Duplicate"),
+			FUIAction(
+			FExecuteAction::CreateSPLambda(this, [&]()
+			{
+				StateTreeViewModel->DuplicateNode(OwnerStateWeak, NodeID);
+			}),
+			FCanExecuteAction::CreateSPLambda(this, [this]()
+			{
+				return !bIsTransition || bIsStateTransition;
+			})
+		));
+
+		// Delete
+		MenuBuilder.AddMenuEntry(
+			LOCTEXT("DeleteItem", "Delete"),
+			LOCTEXT("DeleteItemTooltip", "Delete this item"),
+			FSlateIcon(FAppStyle::GetAppStyleSetName(), "GenericCommands.Delete"),
+			FUIAction(
+			FExecuteAction::CreateSPLambda(this, [&]()
+			{
+				StateTreeViewModel->DeleteNode(OwnerStateWeak, NodeID);
+			}),
+			FCanExecuteAction::CreateSPLambda(this, [this]()
+			{
+				return !bIsTransition || bIsStateTransition;
+			})
+		));
+
+		// Delete All
+		MenuBuilder.AddMenuEntry(
+			LOCTEXT("DeleteAllItems", "Delete All"),
+			LOCTEXT("DeleteAllItemsTooltip", "Delete All Items"),
+			FSlateIcon(FAppStyle::GetAppStyleSetName(), "GenericCommands.Delete"),
+			FUIAction(
+			FExecuteAction::CreateSPLambda(this, [&]()
+			{
+				StateTreeViewModel->DeleteAllNodes(OwnerStateWeak, NodeID);
+			}),
+			FCanExecuteAction::CreateSPLambda(this, [this]()
+			{
+				return !bIsTransition || bIsStateTransition;
+			})
+		));
+
+		MenuBuilder.EndSection();
+	}
+
+	return MenuBuilder.MakeWidget();
+}
+
+#undef LOCTEXT_NAMESPACE
diff --git a/Engine/Plugins/Runtime/StateTree/Source/StateTreeEditorModule/Private/Customizations/Widgets/SStateTreeContextMenuButton.h b/Engine/Plugins/Runtime/StateTree/Source/StateTreeEditorModule/Private/Customizations/Widgets/SStateTreeContextMenuButton.h
new file mode 100644
--- /dev/null
+++ b/Engine/Plugins/Runtime/StateTree/Source/StateTreeEditorModule/Private/Customizations/Widgets/SStateTreeContextMenuButton.h
@@ -0,0 +1,53 @@
+// Copyright Epic Games, Inc.All Rights Reserved.
+
+#pragma once
+
+#include "Widgets/Input/SButton.h"
+
+class SMenuAnchor;
+class UStateTreeEditorData;
+struct FGeometry;
+struct FPointerEvent;
+class FStateTreeViewModel;
+class UStateTreeState;
+
+class SStateTreeContextMenuButton : public SButton
+{
+public:
+	SLATE_BEGIN_ARGS(SStateTreeContextMenuButton)
+		: _Content()
+		, _ButtonStyle(&FCoreStyle::Get().GetWidgetStyle< FButtonStyle >("Button"))
+		, _ContentPadding(FMargin(4.0, 2.0))
+		{
+		}
+
+		/** Slot for this button's content (optional) */
+		SLATE_DEFAULT_SLOT(FArguments, Content)
+
+		/** The visual style of the button */
+		SLATE_STYLE_ARGUMENT(FButtonStyle, ButtonStyle)
+
+		/** Spacing between button's border and the content. */
+		SLATE_ATTRIBUTE(FMargin, ContentPadding)
+
+	SLATE_END_ARGS()
+
+	void Construct(const FArguments& InArgs, const TSharedRef<FStateTreeViewModel>& InStateTreeViewModel, const TWeakObjectPtr<UStateTreeState> InOwnerState, const FGuid& InSourceID, bool InbIsTransition = false);
+
+	//~ SWidget interface
+	virtual FReply OnMouseButtonUp(const FGeometry& MyGeometry, const FPointerEvent& MouseEvent) override;
+	//~ SWidget overrides
+
+	TSharedRef<SWidget> MakeContextMenu() const;
+
+	TSharedPtr<FStateTreeViewModel> StateTreeViewModel;
+	TWeakObjectPtr<UStateTreeState> OwnerStateWeak;
+	FGuid NodeID;
+
+	uint8 bIsTransition : 1;
+	// We have State Transition, Task Transition and Default Transition back to root
+	uint8 bIsStateTransition : 1;
+
+private:
+	TSharedPtr<SMenuAnchor> MenuAnchor;
+};
diff --git a/Engine/Plugins/Runtime/StateTree/Source/StateTreeEditorModule/Private/SStateTreeView.cpp b/Engine/Plugins/Runtime/StateTree/Source/StateTreeEditorModule/Private/SStateTreeView.cpp
--- a/Engine/Plugins/Runtime/StateTree/Source/StateTreeEditorModule/Private/SStateTreeView.cpp
+++ b/Engine/Plugins/Runtime/StateTree/Source/StateTreeEditorModule/Private/SStateTreeView.cpp
@@ -45,6 +45,7 @@ SStateTreeView::~SStateTreeView()
 			StateTreeViewModel->GetOnStateAdded().RemoveAll(this);
 			StateTreeViewModel->GetOnStatesChanged().RemoveAll(this);
 			StateTreeViewModel->GetOnSelectionChanged().RemoveAll(this);
+			StateTreeViewModel->GetOnStateNodesChanged().RemoveAll(this);
 		}
 	}
 }
@@ -59,6 +60,7 @@ void SStateTreeView::Construct(const FArguments& InArgs, TSharedRef<FStateTreeVi
 	StateTreeViewModel->GetOnStateAdded().AddSP(this, &SStateTreeView::HandleModelStateAdded);
 	StateTreeViewModel->GetOnStatesChanged().AddSP(this, &SStateTreeView::HandleModelStatesChanged);
 	StateTreeViewModel->GetOnSelectionChanged().AddSP(this, &SStateTreeView::HandleModelSelectionChanged);
+	StateTreeViewModel->GetOnStateNodesChanged().AddSP(this, &SStateTreeView::HandleModelStateNodesChanged);
 
 	SettingsChangedHandle = GetMutableDefault<UStateTreeEditorUserSettings>()->OnSettingsChanged.AddSP(this, &SStateTreeView::HandleUserSettingsChanged);
 
@@ -417,7 +419,11 @@ void SStateTreeView::HandleUserSettingsChanged()
 
 void SStateTreeView::HandleModelAssetChanged()
 {
-	bItemsDirty = true;
+	// this only refresh the list. i.e. each row widget will not be refreshed
+	// bItemsDirty = true;
+
+	// we need to rebuild the list to update each row widget
+	TreeView->RebuildList();
 }
 
 void SStateTreeView::HandleModelStatesRemoved(const TSet<UStateTreeState*>& AffectedParents)
@@ -458,6 +464,11 @@ void SStateTreeView::HandleModelStatesChanged(const TSet<UStateTreeState*>& Affe
 	}
 }
 
+void SStateTreeView::HandleModelStateNodesChanged(const UStateTreeState* AffectedState)
+{
+	TreeView->RebuildList();
+}
+
 void SStateTreeView::HandleModelSelectionChanged(const TArray<TWeakObjectPtr<UStateTreeState>>& SelectedStates)
 {
 	if (bUpdatingSelection)
diff --git a/Engine/Plugins/Runtime/StateTree/Source/StateTreeEditorModule/Private/SStateTreeView.h b/Engine/Plugins/Runtime/StateTree/Source/StateTreeEditorModule/Private/SStateTreeView.h
--- a/Engine/Plugins/Runtime/StateTree/Source/StateTreeEditorModule/Private/SStateTreeView.h
+++ b/Engine/Plugins/Runtime/StateTree/Source/StateTreeEditorModule/Private/SStateTreeView.h
@@ -47,6 +47,7 @@ private:
 	void HandleModelStatesMoved(const TSet<UStateTreeState*>& AffectedParents, const TSet<UStateTreeState*>& MovedStates);
 	void HandleModelStateAdded(UStateTreeState* ParentState, UStateTreeState* NewState);
 	void HandleModelStatesChanged(const TSet<UStateTreeState*>& AffectedStates, const FPropertyChangedEvent& PropertyChangedEvent);
+	void HandleModelStateNodesChanged(const UStateTreeState* AffectedState);
 	void HandleModelSelectionChanged(const TArray<TWeakObjectPtr<UStateTreeState>>& SelectedStates);
 
 	//~ Treeview handlers
diff --git a/Engine/Plugins/Runtime/StateTree/Source/StateTreeEditorModule/Private/SStateTreeViewRow.cpp b/Engine/Plugins/Runtime/StateTree/Source/StateTreeEditorModule/Private/SStateTreeViewRow.cpp
--- a/Engine/Plugins/Runtime/StateTree/Source/StateTreeEditorModule/Private/SStateTreeViewRow.cpp
+++ b/Engine/Plugins/Runtime/StateTree/Source/StateTreeEditorModule/Private/SStateTreeViewRow.cpp
@@ -22,6 +22,7 @@
 #include "Widgets/Views/SListView.h"
 #include "TextStyleDecorator.h"
 #include "Customizations/StateTreeEditorNodeUtils.h"
+#include "Customizations/Widgets/SStateTreeContextMenuButton.h"
 #include "Widgets/Layout/SWrapBox.h"
 #include "Widgets/Text/SRichTextBlock.h"
 #include "Widgets/Input/SButton.h"
@@ -465,14 +466,9 @@ TSharedRef<SWidget> SStateTreeViewRow::MakeTasksWidget(const TSharedPtr<SScrollB
 			TasksBox->AddSlot()
 				.Padding(FMargin(0, 0, 6, 0))
 				[
-					SNew(SButton)
+					SNew(SStateTreeContextMenuButton, StateTreeViewModel.ToSharedRef(), WeakState, TaskId)
 					.ButtonStyle(FAppStyle::Get(), "SimpleButton")
 					.ContentPadding(FMargin(0, 0))
-					.OnClicked_Lambda([StateTreeViewModel = StateTreeViewModel, WeakState = WeakState, TaskId]()
-					{
-						StateTreeViewModel->BringNodeToFocus(WeakState.Get(), TaskId);
-						return FReply::Handled();
-					})
 					[
 						SNew(SBorder)
 						.VAlign(VAlign_Center)
@@ -746,14 +742,9 @@ TSharedRef<SWidget> SStateTreeViewRow::MakeConditionsWidget(const TSharedPtr<SSc
 									+ SOverlay::Slot()
 									[
 
-										SNew(SButton)
+										SNew(SStateTreeContextMenuButton, StateTreeViewModel.ToSharedRef(), WeakState, ConditionId)
 										.ButtonStyle(FAppStyle::Get(), "SimpleButton")
 										.ContentPadding(FMargin(2.f, 0.f))
-										.OnClicked_Lambda([StateTreeViewModel = StateTreeViewModel, WeakState = WeakState, ConditionId]()
-										{
-											StateTreeViewModel->BringNodeToFocus(WeakState.Get(), ConditionId);
-											return FReply::Handled();
-											})
 										[
 											SNew(SHorizontalBox)
 										
@@ -1107,22 +1098,15 @@ TSharedRef<SWidget> SStateTreeViewRow::MakeTransitionWidgetInternal(const EState
 				.TextStyle(&FStateTreeEditorStyle::Get().GetWidgetStyle<FTextBlockStyle>("Transition.Subdued"))
 			];
 		}
-		
+
+		constexpr bool bIsTransition = true;
 		TransitionContainer->AddSlot()
 		.AutoWidth()
 		.VAlign(VAlign_Center)
 		[
-			SNew(SButton)
-			.ButtonStyle(FAppStyle::Get(), "SimpleButton")
-			.ContentPadding(FMargin(0, 0))
-			.OnClicked_Lambda([StateTreeViewModel = StateTreeViewModel, WeakState = WeakState, Item]()
-			{
-				if (StateTreeViewModel)
-				{
-					StateTreeViewModel->BringNodeToFocus(WeakState.Get(), Item.NodeID);
-				}
-				return FReply::Handled();
-			})
+			SNew(SStateTreeContextMenuButton, StateTreeViewModel.ToSharedRef(), WeakState, Item.NodeID, bIsTransition)
+				.ButtonStyle(FAppStyle::Get(), "SimpleButton")
+				.ContentPadding(FMargin(0, 0))
 			[
 				SNew(SBorder)
 				.VAlign(VAlign_Center)
diff --git a/Engine/Plugins/Runtime/StateTree/Source/StateTreeEditorModule/Private/StateTreeEditorMode.cpp b/Engine/Plugins/Runtime/StateTree/Source/StateTreeEditorModule/Private/StateTreeEditorMode.cpp
--- a/Engine/Plugins/Runtime/StateTree/Source/StateTreeEditorModule/Private/StateTreeEditorMode.cpp
+++ b/Engine/Plugins/Runtime/StateTree/Source/StateTreeEditorModule/Private/StateTreeEditorMode.cpp
@@ -217,6 +217,7 @@ void UStateTreeEditorMode::Exit()
 				ViewModel->GetOnStateAdded().RemoveAll(this);
 				ViewModel->GetOnStatesRemoved().RemoveAll(this);
 				ViewModel->GetOnStatesMoved().RemoveAll(this);
+				ViewModel->GetOnStateNodesChanged().RemoveAll(this);
 				ViewModel->GetOnSelectionChanged().RemoveAll(this);
 				ViewModel->GetOnBringNodeToFocus().RemoveAll(this);
 			}
@@ -278,6 +279,7 @@ void UStateTreeEditorMode::OnStateTreeChanged()
 					NewViewModel->GetOnStateAdded().AddUObject(this, &UStateTreeEditorMode::HandleStateAdded);
 					NewViewModel->GetOnStatesRemoved().AddUObject(this, &UStateTreeEditorMode::HandleStatesRemoved);
 					NewViewModel->GetOnStatesMoved().AddUObject(this, &UStateTreeEditorMode::HandleOnStatesMoved);
+					NewViewModel->GetOnStateNodesChanged().AddUObject(this, &UStateTreeEditorMode::HandleOnStateNodesChanged);
 					NewViewModel->GetOnSelectionChanged().AddUObject(this, &UStateTreeEditorMode::HandleModelSelectionChanged);
 					NewViewModel->GetOnBringNodeToFocus().AddUObject(this, &UStateTreeEditorMode::HandleModelBringNodeToFocus);
 				}
@@ -481,7 +483,7 @@ void UStateTreeEditorMode::HandleModelSelectionChanged(const TArray<TWeakObjectP
 
 void UStateTreeEditorMode::HandleModelBringNodeToFocus(const UStateTreeState* State, const FGuid NodeID) const
 {
-	auto BingToFocus = [Self = this](const FPropertyPath& HighlightPath, TSharedPtr<IDetailsView>& DetailsView)
+	auto BringToFocus = [Self = this](const FPropertyPath& HighlightPath, TSharedPtr<IDetailsView>& DetailsView)
 		{
 			if (HighlightPath.IsValid())
 			{
@@ -570,7 +572,7 @@ void UStateTreeEditorMode::HandleModelBringNodeToFocus(const UStateTreeState* St
 			}
 		}
 
-		BingToFocus(HighlightPath, DetailsView);
+		BringToFocus(HighlightPath, DetailsView);
 	}
 	else
 	{
@@ -619,7 +621,7 @@ void UStateTreeEditorMode::HandleModelBringNodeToFocus(const UStateTreeState* St
 			}
 		}
 
-		BingToFocus(HighlightPath, DetailsView);
+		BringToFocus(HighlightPath, DetailsView);
 	}
 }
 
diff --git a/Engine/Plugins/Runtime/StateTree/Source/StateTreeEditorModule/Private/StateTreePropertyHelpers.cpp b/Engine/Plugins/Runtime/StateTree/Source/StateTreeEditorModule/Private/StateTreePropertyHelpers.cpp
--- a/Engine/Plugins/Runtime/StateTree/Source/StateTreeEditorModule/Private/StateTreePropertyHelpers.cpp
+++ b/Engine/Plugins/Runtime/StateTree/Source/StateTreeEditorModule/Private/StateTreePropertyHelpers.cpp
@@ -2,8 +2,11 @@
 
 #include "StateTreePropertyHelpers.h"
 #include "StateTreeEditorNode.h"
+#include "ScopedTransaction.h"
 #include "Hash/Blake3.h"
+#include "Misc/EnumerateRange.h"
 #include "Misc/StringBuilder.h"
+#include "String/ParseTokens.h"
 #include "UObject/Field.h"
 #include "StateTreePropertyBindings.h"
 #include "StateTreeEditorData.h"
@@ -217,6 +220,41 @@ void DispatchPostEditToNodes(UObject& Owner, FPropertyChangedChainEvent& InPrope
 	}
 }
 
+void ModifyStateInPreAndPostEdit(
+	const FText& TransactionDescription, 
+	TNotNull<UStateTreeState*> State,
+	TNotNull<UStateTreeEditorData*> EditorData,
+	FStringView RelativeNodePath, 
+	TFunctionRef<void(TNotNull<UStateTreeState*> OwnerState, TNotNull<UStateTreeEditorData*> EditorData, const FStateTreeEditPropertyPath& PropertyPath)> Func,
+	int32 ArrayIndex, 
+	EPropertyChangeType::Type ChangeType)
+{
+	FScopedTransaction ScopedTransaction(TransactionDescription);
+
+	FEditPropertyChain PropertyChain;
+
+	FStateTreeEditPropertyPath PropertyPath = FStateTreeEditPropertyPath(State->GetClass(), RelativeNodePath);
+	PropertyPath.MakeEditPropertyChain(PropertyChain);
+
+	State->PreEditChange(PropertyChain);
+
+	Func(State, EditorData, PropertyPath);
+
+	FProperty* ActiveNode = PropertyChain.GetActiveNode()->GetValue();
+	TArray<TMap<FString, int32>> ArrayIndicesPerObject;
+
+	if (ArrayIndex != INDEX_NONE)
+	{
+		ArrayIndicesPerObject.Add(TMap<FString, int32>());
+		ArrayIndicesPerObject[0].Add(ActiveNode->GetName(), ArrayIndex);
+	}
+
+	FPropertyChangedEvent ChangedEvent(PropertyChain.GetActiveNode()->GetValue(), ChangeType);
+	ChangedEvent.SetArrayIndexPerObject(ArrayIndicesPerObject);
+
+	FPropertyChangedChainEvent ChainEvent(PropertyChain, ChangedEvent);
+	State->PostEditChangeChainProperty(ChainEvent);
+}
 
 FGuid MakeDeterministicID(const UObject& Owner, const FString& PropertyPath, const uint64 Seed)
 {
@@ -250,16 +288,16 @@ bool HasOptionalMetadata(const FProperty& Property)
 // ------------------------------------------------------------------------------
 // FStateTreeEditPropertyPath
 // ------------------------------------------------------------------------------
-FStateTreeEditPropertyPath::FStateTreeEditPropertyPath(const UStruct* BaseStruct, const FString& InPath)
+FStateTreeEditPropertyPath::FStateTreeEditPropertyPath(const UStruct* BaseStruct, FStringView InPath)
 {
-	TArray<FString> PathSegments;
-	InPath.ParseIntoArray(PathSegments, TEXT("."));
+	TArray<FStringView> PathSegments;
+	UE::String::ParseTokens(InPath, TEXT("."), PathSegments, UE::String::EParseTokensOptions::SkipEmpty);
 
 	const UStruct* CurrBase = BaseStruct;
-	for (const FString& Segment : PathSegments)
+	for (FStringView Segment : PathSegments)
 	{
 		const FName PropertyName(Segment);
-		if (const FProperty* Property = CurrBase->FindPropertyByName(PropertyName))
+		if (FProperty* Property = CurrBase->FindPropertyByName(PropertyName))
 		{
 			Path.Emplace(Property, PropertyName);
 
@@ -279,7 +317,7 @@ FStateTreeEditPropertyPath::FStateTreeEditPropertyPath(const UStruct* BaseStruct
 		}
 		else
 		{
-			checkf(false, TEXT("Path %s id not part of type %s."), *InPath, *GetNameSafe(BaseStruct));
+			checkf(false, TEXT("Path %s id not part of type %s."), InPath.GetData(), *GetNameSafe(BaseStruct));
 			Path.Reset();
 			break;
 		}
@@ -315,6 +353,18 @@ FStateTreeEditPropertyPath::FStateTreeEditPropertyPath(const FEditPropertyChain&
 	}
 }
 
+void FStateTreeEditPropertyPath::MakeEditPropertyChain(FEditPropertyChain& OutPropertyChain) const
+{
+	OutPropertyChain.Empty();
+
+	for (const FStateTreeEditPropertySegment& Segment : Path)
+	{
+		OutPropertyChain.AddTail(Segment.Property);
+	}
+
+	OutPropertyChain.SetActiveMemberPropertyNode(Path[0].Property);
+}
+
 bool FStateTreeEditPropertyPath::ContainsPath(const FStateTreeEditPropertyPath& InPath) const
 {
 	if (InPath.Path.Num() > Path.Num())
diff --git a/Engine/Plugins/Runtime/StateTree/Source/StateTreeEditorModule/Private/StateTreePropertyHelpers.h b/Engine/Plugins/Runtime/StateTree/Source/StateTreeEditorModule/Private/StateTreePropertyHelpers.h
--- a/Engine/Plugins/Runtime/StateTree/Source/StateTreeEditorModule/Private/StateTreePropertyHelpers.h
+++ b/Engine/Plugins/Runtime/StateTree/Source/StateTreeEditorModule/Private/StateTreePropertyHelpers.h
@@ -2,13 +2,17 @@
 
 #pragma once
 
+#include "Misc/NotNull.h"
 #include "Editor.h"
 #include "PropertyHandle.h"
-#include "Misc/EnumerateRange.h"
 
+struct FStateTreeEditorNode;
 struct FGuid;
 struct FSlateBrush;
+struct FStateTreeEditPropertyPath;
+
 class UStateTreeEditorData;
+class UStateTreeState;
 
 #define LOCTEXT_NAMESPACE "StateTreeEditor"
 
@@ -20,6 +24,18 @@ namespace UE::StateTree::PropertyHelpers {
  */
 void DispatchPostEditToNodes(UObject& Owner, FPropertyChangedChainEvent& PropertyChangedEvent, UStateTreeEditorData& EditorData);
 
+/**
+ * Modify a StateTreeState in PreEdit and PostEdit callbacks in a transaction.
+ */
+void ModifyStateInPreAndPostEdit(
+	const FText& TransactionDescription,
+	TNotNull<UStateTreeState*> State,
+	TNotNull<UStateTreeEditorData*> EditorData,
+	FStringView RelativeNodePath,
+	TFunctionRef<void(TNotNull<UStateTreeState*>, TNotNull<UStateTreeEditorData*>, const FStateTreeEditPropertyPath&)> Func,
+	int32 ArrayIndex = INDEX_NONE,
+	EPropertyChangeType::Type ChangeType = EPropertyChangeType::ValueSet);
+
 /** Makes deterministic ID from the owners property path, a property path (or any string), and a seed value (e.g. array index). */
 FGuid MakeDeterministicID(const UObject& Owner, const FString& PropertyPath, const uint64 Seed);
 
@@ -182,14 +198,14 @@ private:
 	struct FStateTreeEditPropertySegment
 	{
 		FStateTreeEditPropertySegment() = default;
-		FStateTreeEditPropertySegment(const FProperty* InProperty, const FName InPropertyName, const int32 InArrayIndex = INDEX_NONE)
+		FStateTreeEditPropertySegment(FProperty* InProperty, const FName InPropertyName, const int32 InArrayIndex = INDEX_NONE)
 			: Property(InProperty)
 			, PropertyName(InPropertyName)
 			, ArrayIndex(InArrayIndex)
 		{
 		}
 	
-		const FProperty* Property = nullptr;
+		FProperty* Property = nullptr;
 		FName PropertyName = FName();
 		int32 ArrayIndex = INDEX_NONE;
 	};
@@ -198,7 +214,7 @@ public:
 	FStateTreeEditPropertyPath() = default;
 
 	/** Makes property path relative to BaseStruct. Checks if the path is not part of the type. */
-	explicit FStateTreeEditPropertyPath(const UStruct* BaseStruct, const FString& InPath);
+	explicit FStateTreeEditPropertyPath(const UStruct* BaseStruct, FStringView InPath);
 
 	/** Makes property path from property change event. */
 	explicit FStateTreeEditPropertyPath(const FPropertyChangedChainEvent& PropertyChangedEvent);
@@ -206,6 +222,12 @@ public:
 	/** Makes property path from property chain. */
 	explicit FStateTreeEditPropertyPath(const FEditPropertyChain& PropertyChain);
 
+	/**
+	 * Makes property chain from property path.
+	 * Copy Ctor of FEditPropertyChain is implicitly deleted so a conversion operator is not available here
+	 */
+	void MakeEditPropertyChain(FEditPropertyChain& OutPropertyChain) const;
+
 	/** @return true if the property path contains specified path. */
 	bool ContainsPath(const FStateTreeEditPropertyPath& InPath) const;
 
diff --git a/Engine/Plugins/Runtime/StateTree/Source/StateTreeEditorModule/Private/StateTreeScopedEditorDataFixer.cpp b/Engine/Plugins/Runtime/StateTree/Source/StateTreeEditorModule/Private/StateTreeScopedEditorDataFixer.cpp
new file mode 100644
--- /dev/null
+++ b/Engine/Plugins/Runtime/StateTree/Source/StateTreeEditorModule/Private/StateTreeScopedEditorDataFixer.cpp
@@ -0,0 +1,103 @@
+// Copyright Epic Games, Inc. All Rights Reserved.
+
+#include "StateTreeScopedEditorDataFixer.h"
+
+#include "Customizations/StateTreeEditorNodeUtils.h"
+#include "StateTreeEditorData.h"
+
+namespace UE::StateTreeEditor
+{
+FScopedEditorDataFixer::~FScopedEditorDataFixer()
+{
+	auto CopyBindings = [](FStateTreeEditorPropertyBindings& Bindings, const FStateTreeEditorNode& EditorNode, const FGuid& CurrentID)
+		{
+			if (CurrentID.IsValid())
+			{
+				Bindings.CopyBindings(CurrentID, EditorNode.ID);
+			}
+		};
+
+	auto ReinstantiateEditorNodeInstanceData = [](TNotNull<UObject*> NodeOwner, FStateTreeEditorNode& EditorNode)
+		{
+			EditorNodeUtils::InstantiateStructSubobjects(*NodeOwner, EditorNode.Node);
+			if (EditorNode.InstanceObject)
+			{
+				EditorNode.InstanceObject = DuplicateObject(EditorNode.InstanceObject, NodeOwner);
+			}
+			else
+			{
+				EditorNodeUtils::InstantiateStructSubobjects(*NodeOwner, EditorNode.Instance);
+			}
+		};
+
+	FStateTreeEditorPropertyBindings* Bindings = EditorData->GetPropertyEditorBindings();
+	for (FEditorNodeDataFixItem& EditorNodeFixInfo : EditorNodesToFix)
+	{
+		UObject& NodeOwner = *EditorNodeFixInfo.NodeOwner;
+		FStateTreeEditorNode& EditorNode = EditorNodeFixInfo.Node;
+
+		FGuid CurrentID = EditorNode.ID;
+		if (EditorNodeFixInfo.bShouldRegenerateGUID)
+		{
+			EditorNode.ID = FGuid::NewGuid();
+		}
+
+		if (EditorNodeFixInfo.bShouldCopyBindings && Bindings != nullptr)
+		{
+			CopyBindings(*Bindings, EditorNode, CurrentID);
+		}
+
+		if (EditorNodeFixInfo.bShouldReinstantiateInstanceData)
+		{
+			ReinstantiateEditorNodeInstanceData(&NodeOwner, EditorNode);
+		}
+	}
+
+	for (FTransitionDataFixItem& TransitionFixInfo : TransitionsToFix)
+	{
+		UObject& TransitionOwner = *TransitionFixInfo.TransitionOwner;
+		FStateTreeTransition& Transition = TransitionFixInfo.Transition;
+
+		FGuid CurrentTransitionID = Transition.ID;
+		TArray<FGuid> CurrentConditionIDs;
+		CurrentConditionIDs.Reserve(Transition.Conditions.Num());
+		for (FStateTreeEditorNode& ConditionNode : Transition.Conditions)
+		{
+			CurrentConditionIDs.Add(ConditionNode.ID);
+		}
+
+		if (TransitionFixInfo.bShouldRegenerateGUID)
+		{
+			Transition.ID = FGuid::NewGuid();
+
+			for (FStateTreeEditorNode& ConditionNode : Transition.Conditions)
+			{
+				ConditionNode.ID = FGuid::NewGuid();
+			}
+		}
+
+		if (TransitionFixInfo.bShouldCopyBindings && Bindings != nullptr)
+		{
+			for (int32 Index = 0; Index < Transition.Conditions.Num(); ++Index)
+			{
+				CopyBindings(*Bindings, Transition.Conditions[Index], CurrentConditionIDs[Index]);
+			}
+		}
+
+		if (TransitionFixInfo.bShouldReinstantiateInstanceData)
+		{
+			for (int32 Index = 0; Index < Transition.Conditions.Num(); ++Index)
+			{
+				ReinstantiateEditorNodeInstanceData(&TransitionOwner, Transition.Conditions[Index]);
+			}
+		}
+	}
+
+	if (bRemoveInvalidBindings && Bindings != nullptr)
+	{
+		TMap<FGuid, const FPropertyBindingDataView> AllStructValues;
+		EditorData->GetAllStructValues(AllStructValues);
+		Bindings->RemoveInvalidBindings(AllStructValues);
+	}
+}
+}
diff --git a/Engine/Plugins/Runtime/StateTree/Source/StateTreeEditorModule/Private/StateTreeScopedEditorDataFixer.h b/Engine/Plugins/Runtime/StateTree/Source/StateTreeEditorModule/Private/StateTreeScopedEditorDataFixer.h
new file mode 100644
--- /dev/null
+++ b/Engine/Plugins/Runtime/StateTree/Source/StateTreeEditorModule/Private/StateTreeScopedEditorDataFixer.h
@@ -0,0 +1,71 @@
+// Copyright Epic Games, Inc. All Rights Reserved.
+
+#pragma once
+
+#include "Containers/Array.h"
+#include "UObject/WeakObjectPtr.h"
+#include "Misc/NotNull.h"
+
+class UStateTreeEditorData;
+struct FStateTreeEditorNode;
+struct FStateTreeTransition;
+
+namespace UE::StateTreeEditor
+{
+/**
+ * Helper struct for fixing up Editor data following editor manipulations
+ * @todo: this should take into account fixing up data for state manipulations in the future
+ */
+struct FScopedEditorDataFixer
+{
+	TNotNull<UStateTreeEditorData*> EditorData;
+	uint8 bRemoveInvalidBindings : 1 = false;
+
+	struct FEditorNodeDataFixItem
+	{
+		// either State or EditorData
+		TNotNull<UObject*> NodeOwner;
+		FStateTreeEditorNode& Node;
+		uint8 bShouldCopyBindings : 1;
+		uint8 bShouldReinstantiateInstanceData : 1;
+		uint8 bShouldRegenerateGUID : 1;
+
+		explicit FEditorNodeDataFixItem(TNotNull<UObject*> InNodeOwner, FStateTreeEditorNode& InNode, bool InbShouldCopyBindings, bool InbShouldReinstantiateInstanceData, bool InbShouldRegenerateGUID)
+			: NodeOwner(InNodeOwner)
+			, Node(InNode)
+			, bShouldCopyBindings(InbShouldCopyBindings)
+			, bShouldReinstantiateInstanceData(InbShouldReinstantiateInstanceData)
+			, bShouldRegenerateGUID(InbShouldRegenerateGUID)
+		{
+		}
+	};
+
+	struct FTransitionDataFixItem
+	{
+		TNotNull<UObject*> TransitionOwner;
+		FStateTreeTransition& Transition;
+		uint8 bShouldCopyBindings : 1;
+		uint8 bShouldReinstantiateInstanceData : 1;
+		uint8 bShouldRegenerateGUID : 1;
+
+		explicit FTransitionDataFixItem(TNotNull<UObject*> InTransitionOwner, FStateTreeTransition& InTransition, bool InbShouldCopyBindings, bool InbShouldReinstantiateInstanceData, bool InbShouldRegenerateGUID)
+			: TransitionOwner(InTransitionOwner)
+			, Transition(InTransition)
+			, bShouldCopyBindings(InbShouldCopyBindings)
+			, bShouldReinstantiateInstanceData(InbShouldReinstantiateInstanceData)
+			, bShouldRegenerateGUID(InbShouldRegenerateGUID)
+		{
+		}
+	};
+
+	TArray<FEditorNodeDataFixItem> EditorNodesToFix;
+	TArray<FTransitionDataFixItem> TransitionsToFix;
+
+	explicit FScopedEditorDataFixer(TNotNull<UStateTreeEditorData*> InEditorData)
+		: EditorData(InEditorData)
+	{
+	}
+
+	~FScopedEditorDataFixer();
+};
+}
diff --git a/Engine/Plugins/Runtime/StateTree/Source/StateTreeEditorModule/Private/StateTreeState.cpp b/Engine/Plugins/Runtime/StateTree/Source/StateTreeEditorModule/Private/StateTreeState.cpp
--- a/Engine/Plugins/Runtime/StateTree/Source/StateTreeEditorModule/Private/StateTreeState.cpp
+++ b/Engine/Plugins/Runtime/StateTree/Source/StateTreeEditorModule/Private/StateTreeState.cpp
@@ -117,22 +117,7 @@ void UStateTreeState::PostEditChangeChainProperty(FPropertyChangedChainEvent& Pr
 	Super::PostEditChangeChainProperty(PropertyChangedEvent);
 
 	const FStateTreeEditPropertyPath ChangePropertyPath(PropertyChangedEvent);
-	
-	auto CopyBindings = [this](const FGuid FromStructID, const FGuid ToStructID)
-	{
-		if (UStateTreeEditorData* EditorData = GetTypedOuter<UStateTreeEditorData>())
-		{
-		    if (FromStructID.IsValid())
-            {
-                if (FStateTreeEditorPropertyBindings* Bindings = EditorData->GetPropertyEditorBindings())
-                {
-                	Bindings->CopyBindings(FromStructID, ToStructID);
-                }
-            }
-		}
-	};
 
-	
 	static const FStateTreeEditPropertyPath StateNamePath(UStateTreeState::StaticClass(), TEXT("Name"));
 	static const FStateTreeEditPropertyPath StateTypePath(UStateTreeState::StaticClass(), TEXT("Type"));
 	static const FStateTreeEditPropertyPath SelectionBehaviorPath(UStateTreeState::StaticClass(), TEXT("SelectionBehavior"));
@@ -248,84 +233,6 @@ void UStateTreeState::PostEditChangeChainProperty(FPropertyChangedChainEvent& Pr
 		}
 	}
 
-	// Ensure unique ID on duplicated items.
-	if (PropertyChangedEvent.ChangeType == EPropertyChangeType::Duplicate)
-	{
-		
-		// Tasks
-		if (ChangePropertyPath.IsPathExact(StateTasksPath))
-		{
-			const int32 ArrayIndex = ChangePropertyPath.GetPropertyArrayIndex(StateTasksPath);
-			if (Tasks.IsValidIndex(ArrayIndex))
-			{
-				const FGuid OldStructID = Tasks[ArrayIndex].ID; 
-				Tasks[ArrayIndex].ID = FGuid::NewGuid();
-				CopyBindings(OldStructID, Tasks[ArrayIndex].ID);
-			}
-		}
-
-		// Enter conditions
-		if (ChangePropertyPath.IsPathExact(StateEnterConditionsPath))
-		{
-			const int32 ArrayIndex = ChangePropertyPath.GetPropertyArrayIndex(StateEnterConditionsPath);
-			if (EnterConditions.IsValidIndex(ArrayIndex))
-			{
-				const FGuid OldStructID = EnterConditions[ArrayIndex].ID; 
-				EnterConditions[ArrayIndex].ID = FGuid::NewGuid();
-				CopyBindings(OldStructID, EnterConditions[ArrayIndex].ID);
-			}
-		}
-
-		// Utility Considerations
-		if (ChangePropertyPath.IsPathExact(StateConsiderationsPath))
-		{
-			const int32 ArrayIndex = ChangePropertyPath.GetPropertyArrayIndex(StateConsiderationsPath);
-			if (Considerations.IsValidIndex(ArrayIndex))
-			{
-				const FGuid OldStructID = Considerations[ArrayIndex].ID;
-				Considerations[ArrayIndex].ID = FGuid::NewGuid();
-				CopyBindings(OldStructID, Considerations[ArrayIndex].ID);
-			}
-		}
-
-		// Transitions
-		if (ChangePropertyPath.IsPathExact(StateTransitionsPath))
-		{
-			const int32 TransitionsIndex = ChangePropertyPath.GetPropertyArrayIndex(StateTransitionsPath);
-			if (Transitions.IsValidIndex(TransitionsIndex))
-			{
-				FStateTreeTransition& Transition = Transitions[TransitionsIndex];
-				Transition.ID = FGuid::NewGuid();
-
-				// Update conditions
-				for (FStateTreeEditorNode& Condition : Transition.Conditions)
-				{
-					const FGuid OldStructID = Condition.ID;
-					Condition.ID = FGuid::NewGuid();
-					CopyBindings(OldStructID, Condition.ID);
-				}
-			}
-		}
-
-		// Transition conditions
-		if (ChangePropertyPath.IsPathExact(StateTransitionsConditionsPath))
-		{
-			const int32 TransitionsIndex = ChangePropertyPath.GetPropertyArrayIndex(StateTransitionsPath);
-			const int32 ConditionsIndex = ChangePropertyPath.GetPropertyArrayIndex(StateTransitionsConditionsPath);
-
-			if (Transitions.IsValidIndex(TransitionsIndex))
-			{
-				FStateTreeTransition& Transition = Transitions[TransitionsIndex];
-				if (Transition.Conditions.IsValidIndex(ConditionsIndex))
-				{
-					const FGuid OldStructID = Transition.Conditions[ConditionsIndex].ID;
-					Transition.Conditions[ConditionsIndex].ID = FGuid::NewGuid();
-					CopyBindings(OldStructID, Transition.Conditions[ConditionsIndex].ID);
-				}
-			}
-		}
-	}
-	
 	// Set default state to root and Id on new transitions.
 	if (PropertyChangedEvent.ChangeType == EPropertyChangeType::ArrayAdd)
 	{
@@ -343,26 +250,6 @@ void UStateTreeState::PostEditChangeChainProperty(FPropertyChangedChainEvent& Pr
 		}
 	}
 
-	// Remove bindings when bindable nodes are removed.
-	if (PropertyChangedEvent.ChangeType == EPropertyChangeType::ArrayRemove)
-	{
-		if (ChangePropertyPath.IsPathExact(StateTasksPath)
-			|| ChangePropertyPath.IsPathExact(StateEnterConditionsPath)
-			|| ChangePropertyPath.IsPathExact(StateConsiderationsPath)
-			|| ChangePropertyPath.IsPathExact(StateTransitionsConditionsPath))
-		{
-			if (UStateTreeEditorData* TreeData = GetTypedOuter<UStateTreeEditorData>())
-			{
-				TreeData->Modify();
-				FStateTreeEditorPropertyBindings* Bindings = TreeData->GetPropertyEditorBindings();
-				check(Bindings);
-				TMap<FGuid, const FPropertyBindingDataView> AllStructValues;
-				TreeData->GetAllStructValues(AllStructValues);
-				Bindings->RemoveInvalidBindings(AllStructValues);
-			}
-		}
-	}
-
 	if (UStateTreeEditorData* TreeData = GetTypedOuter<UStateTreeEditorData>())
 	{
 		UE::StateTree::PropertyHelpers::DispatchPostEditToNodes(*this, PropertyChangedEvent, *TreeData);
diff --git a/Engine/Plugins/Runtime/StateTree/Source/StateTreeEditorModule/Private/StateTreeViewModel.cpp b/Engine/Plugins/Runtime/StateTree/Source/StateTreeEditorModule/Private/StateTreeViewModel.cpp
--- a/Engine/Plugins/Runtime/StateTree/Source/StateTreeEditorModule/Private/StateTreeViewModel.cpp
+++ b/Engine/Plugins/Runtime/StateTree/Source/StateTreeEditorModule/Private/StateTreeViewModel.cpp
@@ -1,6 +1,8 @@
 // Copyright Epic Games, Inc. All Rights Reserved.
 
 #include "StateTreeViewModel.h"
+
+#include "Customizations/StateTreeEditorNodeUtils.h"
 #include "StateTree.h"
 #include "StateTreeEditorData.h"
 #include "StateTreeDelegates.h"
@@ -11,7 +13,12 @@
 #include "UnrealExporter.h"
 #include "Exporters/Exporter.h"
 #include "Factories.h"
-
+#include "StateTreeTaskBase.h"
+#include "StateTreeConditionBase.h"
+#include "StateTreeConsiderationBase.h"
+#include "StateTreePropertyHelpers.h"
+#include "StateTreeScopedEditorDataFixer.h"
+#include "Misc/NotNull.h"
 
 #define LOCTEXT_NAMESPACE "StateTreeEditor"
 
@@ -110,6 +117,7 @@ namespace UE::StateTree::Editor
 		}
 	}
 
+	// todo: Should refactor it into FStateTreeScopedEditorDataFixer
 	void FixNodesAfterDuplication(TArrayView<FStateTreeEditorNode> Nodes, TMap<FGuid, FGuid>& IDsMap, TArray<FStateTreeStateLink*>& Links)
 	{
 		for (FStateTreeEditorNode& Node : Nodes)
@@ -133,6 +141,7 @@ namespace UE::StateTree::Editor
 		}
 	}
 
+	// todo: Should refactor it into FStateTreeScopedEditorDataFixer
 	void FixStateAfterDuplication(UStateTreeState* State, UStateTreeState* NewParentState, TMap<FGuid, FGuid>& IDsMap, TArray<FStateTreeStateLink*>& Links, TArray<UStateTreeState*>& NewStates)
 	{
 		State->Modify();
@@ -228,8 +237,129 @@ namespace UE::StateTree::Editor
 		return false;
 	}
 
-};
+namespace Private
+{
+	/* Short-lived helper struct for node manipulation in the editor */
+	struct FStateTreeStateNodeEditorHandle
+	{
+		const TNotNull<UStateTreeEditorData*> EditorData;
+		const TNotNull<UStateTreeState*> OwnerState;
+
+	private:
+		FStringView NodePath;
+		int32 ArrayIndex = INDEX_NONE;
+		void* TargetArray = nullptr;
+		void* TargetNode = nullptr;
+		bool bIsTransition = false;
 
+	public:
+		FStateTreeStateNodeEditorHandle(const TNotNull<UStateTreeEditorData*> InEditorData, const TNotNull<UStateTreeState*> InOwnerState, const FGuid& NodeID)
+			: EditorData(InEditorData),
+			OwnerState(InOwnerState)
+		{
+			auto FindNode = [Self = this, &NodeID]<typename T>(T& Nodes, FStringView Path)
+			{
+				for (int32 Index = 0; Index < Nodes.Num(); ++Index)
+				{
+					if (NodeID == Nodes[Index].ID)
+					{
+						Self->ArrayIndex = Index;
+						Self->TargetArray = &Nodes;
+						Self->TargetNode = &Nodes[Self->ArrayIndex];
+						Self->NodePath = Path;
+
+						return true;
+					}
+				}
+
+				return false;
+			};
+
+			if (!NodeID.IsValid())
+			{
+				return;
+			}
+
+			if (FindNode(OwnerState->EnterConditions, TEXT("EnterConditions")))
+			{
+				return;
+			}
+
+			if (FindNode(OwnerState->Tasks, TEXT("Tasks")))
+			{
+				return;
+			}
+
+			if (NodeID == OwnerState->SingleTask.ID)
+			{
+				NodePath = TEXT("SingleTask");
+				TargetNode = &OwnerState->SingleTask;
+				return;
+			}
+
+			if (FindNode(OwnerState->Considerations, TEXT("Considerations")))
+			{
+				return;
+			}
+
+			if (FindNode(OwnerState->Transitions, TEXT("Transitions")))
+			{
+				bIsTransition = true;
+				return;
+			}
+		}
+
+		bool IsValid() const
+		{
+			return TargetNode != nullptr;
+		}
+
+		bool IsTransition() const
+		{
+			return bIsTransition;
+		}
+
+		FStateTreeEditorNode& GetEditorNode() const
+		{
+			check(IsValid() && !IsTransition());
+
+			return *static_cast<FStateTreeEditorNode*>(TargetNode);
+		}
+
+		TArray<FStateTreeEditorNode>& GetEditorNodeArray() const
+		{
+			check(IsValid() && !IsTransition() && GetNodeIndex() != INDEX_NONE);
+
+			return *static_cast<TArray<FStateTreeEditorNode>*>(TargetArray);
+		}
+
+		FStateTreeTransition& GetTransition() const
+		{
+			check(IsValid() && IsTransition());
+
+			return *static_cast<FStateTreeTransition*>(TargetNode);
+		}
+
+		TArray<FStateTreeTransition>& GetTransitionArray() const
+		{
+			check(IsValid() && IsTransition() && GetNodeIndex() != INDEX_NONE);
+
+			return *static_cast<TArray<FStateTreeTransition>*>(TargetArray);
+		}
+
+		FStringView GetNodePath() const
+		{
+			return NodePath;
+		}
+
+		int32 GetNodeIndex() const
+		{
+			return ArrayIndex;
+		}
+	};
+}
+
+};
 
 FStateTreeViewModel::FStateTreeViewModel()
 	: TreeDataWeak(nullptr)
@@ -665,6 +795,266 @@ void FStateTreeViewModel::SetSelection(UStateTreeState* Selected)
 	OnSelectionChanged.Broadcast(SelectedStatesArr);
 }
 
+void FStateTreeViewModel::DeleteNode(TWeakObjectPtr<UStateTreeState> State, const FGuid& ID)
+{
+	if (UStateTreeEditorData* EditorData = TreeDataWeak.Get())
+	{
+		if (UStateTreeState* OwnerState = State.Get())
+		{
+			const UE::StateTree::Editor::Private::FStateTreeStateNodeEditorHandle StateNodeHandle(EditorData, OwnerState, ID);
+			// If the op can not be executed, don't even start the transaction
+			if (!StateNodeHandle.IsValid())
+			{
+				return;
+			}
+
+			auto DeleteFunc = [&StateNodeHandle](TNotNull<UStateTreeState*> InOwnerState, TNotNull<UStateTreeEditorData*> InEditorData, const FStateTreeEditPropertyPath& InPropertyPath)
+				{
+					StateNodeHandle.EditorData->Modify();
+
+					UE::StateTreeEditor::FScopedEditorDataFixer DataFixer(InEditorData);
+					DataFixer.bRemoveInvalidBindings = true;
+
+					if (StateNodeHandle.IsTransition())
+					{
+						StateNodeHandle.GetTransitionArray().RemoveAt(StateNodeHandle.GetNodeIndex());
+					}
+					else
+					{
+						StateNodeHandle.GetEditorNodeArray().RemoveAt(StateNodeHandle.GetNodeIndex());
+					}
+				};
+
+			UE::StateTree::PropertyHelpers::ModifyStateInPreAndPostEdit(
+				LOCTEXT("DeleteNodeTransaction", "Delete Node"),
+				StateNodeHandle.OwnerState,
+				StateNodeHandle.EditorData,
+				StateNodeHandle.GetNodePath(),
+				DeleteFunc,
+				StateNodeHandle.GetNodeIndex(),
+				EPropertyChangeType::ArrayRemove);
+
+
+			OnStateNodesChanged.Broadcast(OwnerState);
+		}
+	}
+}
+
+void FStateTreeViewModel::DeleteAllNodes(TWeakObjectPtr<UStateTreeState> State, const FGuid& ID)
+{
+	if (UStateTreeEditorData* EditorData = TreeDataWeak.Get())
+	{
+		if (UStateTreeState* OwnerState = State.Get())
+		{
+			UE::StateTree::Editor::Private::FStateTreeStateNodeEditorHandle StateNodeHandle(EditorData, OwnerState, ID);
+			// If the op can not be executed, don't even start the transaction
+			if (!StateNodeHandle.IsValid())
+			{
+				return;
+			}
+
+			auto DeleteAllFunc = [&StateNodeHandle](const TNotNull<UStateTreeState*> InOwnerState, const TNotNull<UStateTreeEditorData*> InEditorData, const FStateTreeEditPropertyPath& InPropertyPath)
+				{
+					StateNodeHandle.EditorData->Modify();
+
+					UE::StateTreeEditor::FScopedEditorDataFixer DataFixer(InEditorData);
+					DataFixer.bRemoveInvalidBindings = true;
+
+					if (StateNodeHandle.IsTransition())
+					{
+						StateNodeHandle.GetTransitionArray().Empty();
+					}
+					else
+					{
+						StateNodeHandle.GetEditorNodeArray().Empty();
+					}
+				};
+
+			UE::StateTree::PropertyHelpers::ModifyStateInPreAndPostEdit(
+				LOCTEXT("DeleteAllNodesTransaction", "Delete All Nodes"),
+				StateNodeHandle.OwnerState,
+				StateNodeHandle.EditorData,
+				StateNodeHandle.GetNodePath(),
+				DeleteAllFunc,
+				INDEX_NONE, // Pass Invalid Index to Array Clear Op
+				EPropertyChangeType::ArrayClear);
+
+
+			OnStateNodesChanged.Broadcast(OwnerState);
+		}
+	}
+}
+
+void FStateTreeViewModel::CopyNode(TWeakObjectPtr<UStateTreeState> State, const FGuid& ID)
+{
+	if (UStateTreeEditorData* EditorData = TreeDataWeak.Get())
+	{
+		if (UStateTreeState* OwnerState = State.Get())
+		{
+			UE::StateTree::Editor::Private::FStateTreeStateNodeEditorHandle StateNodeHandle(EditorData, OwnerState, ID);
+			// If the op can not be executed, don't even start the transaction
+			if (!StateNodeHandle.IsValid())
+			{
+				return;
+			}
+
+			const UScriptStruct* BaseStruct = nullptr;
+			if (StateNodeHandle.IsTransition())
+			{
+				BaseStruct = FStateTreeTransition::StaticStruct();
+			}
+			else
+			{
+				BaseStruct = FStateTreeEditorNode::StaticStruct();
+			}
+
+			// Use PPF_Copy so that all properties get copied.
+			FString Value;
+			constexpr EPropertyPortFlags PortFlags = PPF_Copy;
+			BaseStruct->ExportText(Value, StateNodeHandle.IsTransition() ? (void*)&StateNodeHandle.GetTransition() : (void*)&StateNodeHandle.GetEditorNode(), nullptr, nullptr, PortFlags, nullptr);
+
+			FPlatformApplicationMisc::ClipboardCopy(*Value);
+		}
+	}
+}
+
+void FStateTreeViewModel::PasteNode(TWeakObjectPtr<UStateTreeState> State, const FGuid& ID)
+{
+	if (UStateTreeEditorData* EditorData = TreeDataWeak.Get())
+	{
+		if (UStateTreeState* OwnerState = State.Get())
+		{
+			UE::StateTree::Editor::Private::FStateTreeStateNodeEditorHandle StateNodeHandle(EditorData, OwnerState, ID);
+			// If the op can not be executed, don't even start the transaction
+			if (!StateNodeHandle.IsValid())
+			{
+				return;
+			}
+
+			// For Editor Node, we have a validation process which might early out if failed. Dry run with a Temp.
+			FStateTreeEditorNode TempEditorNode;
+			if (!StateNodeHandle.IsTransition())
+			{
+				static const UScriptStruct* TaskBaseStruct = FStateTreeTaskBase::StaticStruct();
+				static const UScriptStruct* ConditionBaseStruct = FStateTreeConditionBase::StaticStruct();
+				static const UScriptStruct* ConsiderationBaseStruct = FStateTreeConsiderationBase::StaticStruct();
+
+				const UScriptStruct* BaseNodeScriptStruct = nullptr;
+				const UScriptStruct* NodeScriptStruct = StateNodeHandle.GetEditorNode().Node.GetScriptStruct();
+				if (NodeScriptStruct->IsChildOf(TaskBaseStruct))
+				{
+					BaseNodeScriptStruct = TaskBaseStruct;
+				}
+				else if (NodeScriptStruct->IsChildOf(ConditionBaseStruct))
+				{
+					BaseNodeScriptStruct = ConditionBaseStruct;
+				}
+				else if (NodeScriptStruct->IsChildOf(ConsiderationBaseStruct))
+				{
+					BaseNodeScriptStruct = ConsiderationBaseStruct;
+				}
+
+				const bool bSuccess = UE::StateTreeEditor::EditorNodeUtils::ImportTextAsNode(BaseNodeScriptStruct, StateNodeHandle.EditorData, TempEditorNode);
+				// If the op can not be executed, don't even start the transaction
+				if (!bSuccess)
+				{
+					return;
+				}
+			}
+
+			auto PasteFunc = [&StateNodeHandle, &TempEditorNode](const TNotNull<UStateTreeState*> InOwnerState, const TNotNull<UStateTreeEditorData*> InEditorData, const FStateTreeEditPropertyPath& InPropertyPath)
+				{
+					InEditorData->Modify();
+
+					UE::StateTreeEditor::FScopedEditorDataFixer DataFixer(StateNodeHandle.EditorData);
+					constexpr bool bShouldCopyBindings = true;
+					constexpr bool bShouldReinstantiateInstanceData = true;
+					constexpr bool bShouldRegenerateGUID = true;
+
+					if (StateNodeHandle.IsTransition())
+					{
+						FString PastedText;
+						FPlatformApplicationMisc::ClipboardPaste(PastedText);
+
+						const UScriptStruct* TransitionStruct = FStateTreeTransition::StaticStruct();
+						TransitionStruct->ImportText(*PastedText, &StateNodeHandle.GetTransition(), nullptr, PPF_None, nullptr, TransitionStruct->GetName());
+						DataFixer.TransitionsToFix.Emplace(StateNodeHandle.OwnerState, StateNodeHandle.GetTransition(), bShouldCopyBindings, bShouldReinstantiateInstanceData, bShouldRegenerateGUID);
+					}
+					else
+					{
+						StateNodeHandle.GetEditorNode() = TempEditorNode;
+						DataFixer.EditorNodesToFix.Emplace(StateNodeHandle.OwnerState, StateNodeHandle.GetEditorNode(), bShouldCopyBindings, bShouldReinstantiateInstanceData, bShouldRegenerateGUID);
+					}
+				};
+
+			UE::StateTree::PropertyHelpers::ModifyStateInPreAndPostEdit(
+				LOCTEXT("PasteNodeTransaction", "Paste Node"),
+				StateNodeHandle.OwnerState,
+				StateNodeHandle.EditorData,
+				StateNodeHandle.GetNodePath(),
+				PasteFunc,
+				INDEX_NONE, // Value Set Op, skip the index
+				EPropertyChangeType::ValueSet);
+
+			OnStateNodesChanged.Broadcast(OwnerState);
+		}
+	}
+}
+
+void FStateTreeViewModel::DuplicateNode(TWeakObjectPtr<UStateTreeState> State, const FGuid& ID)
+{
+	if (UStateTreeEditorData* EditorData = TreeDataWeak.Get())
+	{
+		if (UStateTreeState* OwnerState = State.Get())
+		{
+			UE::StateTree::Editor::Private::FStateTreeStateNodeEditorHandle StateNodeHandle(EditorData, OwnerState, ID);
+			// If the op can not be executed, don't even start the transaction
+			if (!StateNodeHandle.IsValid())
+			{
+				return;
+			}
+
+			auto DuplicateFunc = [&StateNodeHandle](const TNotNull<UStateTreeState*> InOwnerState, const TNotNull<UStateTreeEditorData*> InEditorData, const FStateTreeEditPropertyPath& InPropertyPath)
+				{
+					InEditorData->Modify();
+
+					UE::StateTreeEditor::FScopedEditorDataFixer DataFixer(StateNodeHandle.EditorData);
+
+					constexpr bool bShouldCopyBindings = true;
+					constexpr bool bShouldReinstantiateInstanceData = true;
+					constexpr bool bShouldRegenerateGUID = true;
+
+					if (StateNodeHandle.IsTransition())
+					{
+						// Make a temp copy
+						const FStateTreeTransition TempTransition = StateNodeHandle.GetTransition();
+						FStateTreeTransition& NewTransition = StateNodeHandle.GetTransitionArray().Insert_GetRef(TempTransition, StateNodeHandle.GetNodeIndex());
+
+						DataFixer.TransitionsToFix.Emplace(StateNodeHandle.OwnerState, NewTransition, bShouldCopyBindings, bShouldReinstantiateInstanceData, bShouldRegenerateGUID);
+					}
+					else
+					{
+						// Make a temp copy
+						const FStateTreeEditorNode TempEditorNode = StateNodeHandle.GetEditorNode();
+						FStateTreeEditorNode& NewEditorNode = StateNodeHandle.GetEditorNodeArray().Insert_GetRef(TempEditorNode, StateNodeHandle.GetNodeIndex());
+						DataFixer.EditorNodesToFix.Emplace(StateNodeHandle.OwnerState, NewEditorNode, bShouldCopyBindings, bShouldReinstantiateInstanceData, bShouldRegenerateGUID);
+					}
+				};
+
+			UE::StateTree::PropertyHelpers::ModifyStateInPreAndPostEdit(
+				LOCTEXT("DuplicateNodeTransaction", "Duplicate Node"),
+				StateNodeHandle.OwnerState,
+				StateNodeHandle.EditorData,
+				StateNodeHandle.GetNodePath(),
+				DuplicateFunc,
+				StateNodeHandle.GetNodeIndex(),
+				EPropertyChangeType::Duplicate);
+
+			OnStateNodesChanged.Broadcast(OwnerState);
+		}
+	}
+}
+
 void FStateTreeViewModel::SetSelection(const TArray<TWeakObjectPtr<UStateTreeState>>& InSelectedStates)
 {
 	if (SelectedStates.Num() == InSelectedStates.Num() && SelectedStates.Array() == InSelectedStates)
@@ -1122,17 +1512,17 @@ void FStateTreeViewModel::DuplicateSelectedStates()
 
 void FStateTreeViewModel::MoveSelectedStatesBefore(UStateTreeState* TargetState)
 {
-	MoveSelectedStates(TargetState, FStateTreeViewModelInsert::Before);
+	MoveSelectedStates(TargetState, EStateTreeViewModelInsert::Before);
 }
 
 void FStateTreeViewModel::MoveSelectedStatesAfter(UStateTreeState* TargetState)
 {
-	MoveSelectedStates(TargetState, FStateTreeViewModelInsert::After);
+	MoveSelectedStates(TargetState, EStateTreeViewModelInsert::After);
 }
 
 void FStateTreeViewModel::MoveSelectedStatesInto(UStateTreeState* TargetState)
 {
-	MoveSelectedStates(TargetState, FStateTreeViewModelInsert::Into);
+	MoveSelectedStates(TargetState, EStateTreeViewModelInsert::Into);
 }
 
 bool FStateTreeViewModel::CanEnableStates() const
@@ -1188,7 +1578,7 @@ void FStateTreeViewModel::SetSelectedStatesEnabled(const bool bEnable)
 	}
 }
 
-void FStateTreeViewModel::MoveSelectedStates(UStateTreeState* TargetState, const FStateTreeViewModelInsert RelativeLocation)
+void FStateTreeViewModel::MoveSelectedStates(UStateTreeState* TargetState, const EStateTreeViewModelInsert RelativeLocation)
 {
 	UStateTreeEditorData* TreeData = TreeDataWeak.Get();
 	if (TreeData == nullptr || TargetState == nullptr)
@@ -1216,7 +1606,7 @@ void FStateTreeViewModel::MoveSelectedStates(UStateTreeState* TargetState, const
 		TSet<UStateTreeState*> AffectedStates;
 
 		UStateTreeState* TargetParent = TargetState->Parent;
-		if (RelativeLocation == FStateTreeViewModelInsert::Into)
+		if (RelativeLocation == EStateTreeViewModelInsert::Into)
 		{
 			AffectedParents.Add(TargetState);
 		}
@@ -1234,7 +1624,7 @@ void FStateTreeViewModel::MoveSelectedStates(UStateTreeState* TargetState, const
 			}
 		}
 
-		if (RelativeLocation == FStateTreeViewModelInsert::Into)
+		if (RelativeLocation == EStateTreeViewModelInsert::Into)
 		{
 			// Move into
 			TargetState->Modify();
@@ -1271,7 +1661,7 @@ void FStateTreeViewModel::MoveSelectedStates(UStateTreeState* TargetState, const
 				}
 
 				// Insert to new parent
-				if (RelativeLocation == FStateTreeViewModelInsert::Into)
+				if (RelativeLocation == EStateTreeViewModelInsert::Into)
 				{
 					// Into
 					TargetState->Children.Insert(SelectedState, /*Index*/0);
@@ -1283,13 +1673,13 @@ void FStateTreeViewModel::MoveSelectedStates(UStateTreeState* TargetState, const
 					const int32 TargetIndex = ArrayToMoveTo.Find(TargetState);
 					if (TargetIndex != INDEX_NONE)
 					{
-						if (RelativeLocation == FStateTreeViewModelInsert::Before)
+						if (RelativeLocation == EStateTreeViewModelInsert::Before)
 						{
 							// Before
 							ArrayToMoveTo.Insert(SelectedState, TargetIndex);
 							SelectedState->Parent = TargetParent;
 						}
-						else if (RelativeLocation == FStateTreeViewModelInsert::After)
+						else if (RelativeLocation == EStateTreeViewModelInsert::After)
 						{
 							// After
 							ArrayToMoveTo.Insert(SelectedState, TargetIndex + 1);
diff --git a/Engine/Plugins/Runtime/StateTree/Source/StateTreeEditorModule/Public/StateTreeEditingSubsystem.h b/Engine/Plugins/Runtime/StateTree/Source/StateTreeEditorModule/Public/StateTreeEditingSubsystem.h
--- a/Engine/Plugins/Runtime/StateTree/Source/StateTreeEditorModule/Public/StateTreeEditingSubsystem.h
+++ b/Engine/Plugins/Runtime/StateTree/Source/StateTreeEditorModule/Public/StateTreeEditingSubsystem.h
@@ -31,6 +31,7 @@ public:
 	
 	/**
 	 * Validates and applies the schema restrictions on the StateTree.
+	 * Also serves as the "safety net" of fixing up editor data following an editor operation.
 	 * Updates state's link, removes the unused node while validating the StateTree asset.
 	 */
 	static void ValidateStateTree(TNotNull<UStateTree*> InStateTree);
diff --git a/Engine/Plugins/Runtime/StateTree/Source/StateTreeEditorModule/Public/StateTreeEditorMode.h b/Engine/Plugins/Runtime/StateTree/Source/StateTreeEditorModule/Public/StateTreeEditorMode.h
--- a/Engine/Plugins/Runtime/StateTree/Source/StateTreeEditorModule/Public/StateTreeEditorMode.h
+++ b/Engine/Plugins/Runtime/StateTree/Source/StateTreeEditorModule/Public/StateTreeEditorMode.h
@@ -52,9 +52,26 @@ protected:
 	void HandleModelSelectionChanged(const TArray<TWeakObjectPtr<UStateTreeState>>& SelectedStates) const;
 	void HandleModelBringNodeToFocus(const UStateTreeState* State, const FGuid NodeID) const;
 
-	void HandleStateAdded(UStateTreeState* , UStateTreeState*) { UpdateAsset(); }
-	void HandleStatesRemoved(const TSet<UStateTreeState*>&) { UpdateAsset(); }
-	void HandleOnStatesMoved(const TSet<UStateTreeState*>&, const TSet<UStateTreeState*>&) { UpdateAsset(); }
+	void HandleStateAdded(UStateTreeState* , UStateTreeState*)
+	{
+		UpdateAsset();
+	}
+
+	void HandleStatesRemoved(const TSet<UStateTreeState*>&)
+	{
+		UpdateAsset();
+	}
+
+	void HandleOnStatesMoved(const TSet<UStateTreeState*>&, const TSet<UStateTreeState*>&)
+	{
+		UpdateAsset();
+	}
+
+	void HandleOnStateNodesChanged(const UStateTreeState*)
+	{
+		UpdateAsset();
+	}
+
 	/** Resolve the internal editor data and fixup the StateTree nodes. */
 	void UpdateAsset();
 
diff --git a/Engine/Plugins/Runtime/StateTree/Source/StateTreeEditorModule/Public/StateTreeViewModel.h b/Engine/Plugins/Runtime/StateTree/Source/StateTreeEditorModule/Public/StateTreeViewModel.h
--- a/Engine/Plugins/Runtime/StateTree/Source/StateTreeEditorModule/Public/StateTreeViewModel.h
+++ b/Engine/Plugins/Runtime/StateTree/Source/StateTreeEditorModule/Public/StateTreeViewModel.h
@@ -20,7 +20,14 @@ struct FStateTreeDebuggerBreakpoint;
 struct FStateTreeEditorBreakpoint;
 struct FStateTreePropertyPathBinding;
 
-enum class FStateTreeViewModelInsert : uint8
+enum class EStateTreeViewModelInsert : uint8
+{
+	Before,
+	After,
+	Into,
+};
+
+enum class UE_DEPRECATED(5.6, "Use the enum with the E prefix") FStateTreeViewModelInsert : uint8
 {
 	Before,
 	After,
@@ -39,8 +46,9 @@ public:
 	DECLARE_MULTICAST_DELEGATE_TwoParams(FOnStateAdded, UStateTreeState* /*ParentState*/, UStateTreeState* /*NewState*/);
 	DECLARE_MULTICAST_DELEGATE_OneParam(FOnStatesRemoved, const TSet<UStateTreeState*>& /*AffectedParents*/);
 	DECLARE_MULTICAST_DELEGATE_TwoParams(FOnStatesMoved, const TSet<UStateTreeState*>& /*AffectedParents*/, const TSet<UStateTreeState*>& /*MovedStates*/);
+	DECLARE_MULTICAST_DELEGATE_OneParam(FOnStateNodesChanged, const UStateTreeState* /*AffectedState*/);
 	DECLARE_MULTICAST_DELEGATE_OneParam(FOnSelectionChanged, const TArray<TWeakObjectPtr<UStateTreeState>>& /*SelectedStates*/);
-	DECLARE_MULTICAST_DELEGATE_TwoParams(FOnBringNodeToFocus, const UStateTreeState* /*State*/, const FGuid /*TaskID*/);
+	DECLARE_MULTICAST_DELEGATE_TwoParams(FOnBringNodeToFocus, const UStateTreeState* /*State*/, const FGuid /*NodeID*/);
 
 	FStateTreeViewModel();
 	virtual ~FStateTreeViewModel() override;
@@ -83,7 +91,7 @@ public:
 	void SetPersistentExpandedStates(TSet<TWeakObjectPtr<UStateTreeState>>& InExpandedStates);
 	void GetPersistentExpandedStates(TSet<TWeakObjectPtr<UStateTreeState>>& OutExpandedStates);
 
-	// State manipulation commands
+	// State manipulation
 	void AddState(UStateTreeState* AfterState);
 	void AddChildState(UStateTreeState* ParentState);
 	void RenameState(UStateTreeState* State, FName NewName);
@@ -99,7 +107,15 @@ public:
 	bool CanEnableStates() const;
 	bool CanDisableStates() const;
 	void SetSelectedStatesEnabled(bool bEnable);
-	
+
+	// EditorNode and Transition manipulation
+	// @todo: support Add, ReplaceWith and Rename
+	void DeleteNode(TWeakObjectPtr<UStateTreeState> State, const FGuid& ID);
+	void DeleteAllNodes(TWeakObjectPtr<UStateTreeState> State, const FGuid& ID);
+	void CopyNode(TWeakObjectPtr<UStateTreeState> State, const FGuid& ID);
+	void PasteNode(TWeakObjectPtr<UStateTreeState> State, const FGuid& ID);
+	void DuplicateNode(TWeakObjectPtr<UStateTreeState> State, const FGuid& ID);
+
 	// Force to update the view externally.
 	void NotifyAssetChangedExternally() const;
 	void NotifyStatesChangedExternally(const TSet<UStateTreeState*>& ChangedStates, const FPropertyChangedEvent& PropertyChangedEvent) const;
@@ -120,28 +136,56 @@ public:
 	bool IsStateActiveInDebugger(const UStateTreeState& State) const;
 
 	// Called when the whole asset is updated (i.e. undo/redo).
-	FOnAssetChanged& GetOnAssetChanged() { return OnAssetChanged; }
+	FOnAssetChanged& GetOnAssetChanged()
+	{
+		return OnAssetChanged;
+	}
 	
 	// Called when States are changed (i.e. change name or properties).
-	FOnStatesChanged& GetOnStatesChanged() { return OnStatesChanged; }
+	FOnStatesChanged& GetOnStatesChanged()
+	{
+		return OnStatesChanged;
+	}
 	
 	// Called each time a state is added.
-	FOnStateAdded& GetOnStateAdded() { return OnStateAdded; }
+	FOnStateAdded& GetOnStateAdded()
+	{
+		return OnStateAdded;
+	}
 
 	// Called each time a states are removed.
-	FOnStatesRemoved& GetOnStatesRemoved() { return OnStatesRemoved; }
+	FOnStatesRemoved& GetOnStatesRemoved()
+	{
+		return OnStatesRemoved;
+	}
 
 	// Called each time a state is removed.
-	FOnStatesMoved& GetOnStatesMoved() { return OnStatesMoved; }
+	FOnStatesMoved& GetOnStatesMoved()
+	{
+		return OnStatesMoved;
+	}
 
-	// Called each time the selection changes.
-	FOnSelectionChanged& GetOnSelectionChanged() { return OnSelectionChanged; }
+	// Called each time a state's Editor nodes or transitions are changed except from the DetailsView.
+	FOnStateNodesChanged& GetOnStateNodesChanged()
+	{
+		return OnStateNodesChanged;
+	}
 
-	FOnBringNodeToFocus& GetOnBringNodeToFocus() { return OnBringNodeToFocus; }
-	
+	// Called each time the selection changes.
+	FOnSelectionChanged& GetOnSelectionChanged()
+	{
+		return OnSelectionChanged;
+	}
+
+	FOnBringNodeToFocus& GetOnBringNodeToFocus()
+	{
+		return OnBringNodeToFocus;
+	}
+	 
 protected:
 	void GetExpandedStatesRecursive(UStateTreeState* State, TSet<TWeakObjectPtr<UStateTreeState>>& ExpandedStates);
-	void MoveSelectedStates(UStateTreeState* TargetState, const FStateTreeViewModelInsert RelativeLocation);
+
+	void MoveSelectedStates(UStateTreeState* TargetState, const EStateTreeViewModelInsert RelativeLocation);
 
 	void PasteStatesAsChildrenFromText(const FString& TextToImport, UStateTreeState* ParentState, const int32 IndexToInsertAt);
 
@@ -165,6 +209,7 @@ protected:
 	FOnStateAdded OnStateAdded;
 	FOnStatesRemoved OnStatesRemoved;
 	FOnStatesMoved OnStatesMoved;
+	FOnStateNodesChanged OnStateNodesChanged;
 	FOnSelectionChanged OnSelectionChanged;
 	FOnBringNodeToFocus OnBringNodeToFocus;
 };
