Fix crash binding ST property on ST which had it's editor data updated during ST validation

diff --git a/Engine/Plugins/Runtime/StateTree/Source/StateTreeEditorModule/Private/StateTreeEditingSubsystem.cpp b/Engine/Plugins/Runtime/StateTree/Source/StateTreeEditorModule/Private/StateTreeEditingSubsystem.cpp
--- a/Engine/Plugins/Runtime/StateTree/Source/StateTreeEditorModule/Private/StateTreeEditingSubsystem.cpp
+++ b/Engine/Plugins/Runtime/StateTree/Source/StateTreeEditorModule/Private/StateTreeEditingSubsystem.cpp
@@ -181,6 +181,7 @@ void UStateTreeEditingSubsystem::ValidateStateTree(TNotNull<UStateTree*> InState
 			EditorData = NewObject<UStateTreeEditorData>(StateTree, EditorDataClass.Get(), FName(), RF_Transactional);
 			EditorData->AddRootState();
 			EditorData->Schema = NewObject<UStateTreeSchema>(EditorData, SchemaClass.Get());
+			EditorData->EditorBindings.SetBindingsOwner(EditorData);
 
 			constexpr bool bMarkDirty = false;
 			StateTree->Modify(bMarkDirty);
@@ -199,6 +200,7 @@ void UStateTreeEditingSubsystem::ValidateStateTree(TNotNull<UStateTree*> InState
 			{
 				EditorData->Schema = NewObject<UStateTreeSchema>(EditorData, SchemaClass.Get());
 			}
+			EditorData->EditorBindings.SetBindingsOwner(EditorData);
 
 			constexpr bool bMarkDirty = false;
 			StateTree->Modify(bMarkDirty);
