[State Tree Editor] Fixed localization warnings

diff --git a/Engine/Plugins/Runtime/StateTree/Source/StateTreeEditorModule/Private/Customizations/Widgets/SStateTreeContextMenuButton.cpp b/Engine/Plugins/Runtime/StateTree/Source/StateTreeEditorModule/Private/Customizations/Widgets/SStateTreeContextMenuButton.cpp
--- a/Engine/Plugins/Runtime/StateTree/Source/StateTreeEditorModule/Private/Customizations/Widgets/SStateTreeContextMenuButton.cpp
+++ b/Engine/Plugins/Runtime/StateTree/Source/StateTreeEditorModule/Private/Customizations/Widgets/SStateTreeContextMenuButton.cpp
@@ -153,8 +153,8 @@ TSharedRef<SWidget> SStateTreeContextMenuButton::MakeContextMenu() const
 
 		// Delete All
 		MenuBuilder.AddMenuEntry(
-			LOCTEXT("DeleteAllItems", "Delete All"),
-			LOCTEXT("DeleteAllItemsTooltip", "Delete All Items"),
+			LOCTEXT("DeleteAllItems", "Delete all"),
+			LOCTEXT("DeleteAllItemsTooltip", "Delete all items"),
 			FSlateIcon(FAppStyle::GetAppStyleSetName(), "GenericCommands.Delete"),
 			FUIAction(
 			FExecuteAction::CreateSPLambda(this, [&]()
