StateTree: Prevent reusing the same compiler more than once.
#rnx

diff --git a/Engine/Plugins/Runtime/StateTree/Source/StateTreeEditorModule/Private/StateTreeCompiler.cpp b/Engine/Plugins/Runtime/StateTree/Source/StateTreeEditorModule/Private/StateTreeCompiler.cpp
--- a/Engine/Plugins/Runtime/StateTree/Source/StateTreeEditorModule/Private/StateTreeCompiler.cpp
+++ b/Engine/Plugins/Runtime/StateTree/Source/StateTreeEditorModule/Private/StateTreeCompiler.cpp
@@ -220,6 +220,13 @@ namespace UE::StateTree::Compiler
 
 bool FStateTreeCompiler::Compile(UStateTree& InStateTree)
 {
+	if (bCompiled)
+	{
+		Log.Reportf(EMessageSeverity::Error, TEXT("Internal error. The compiler has already been executed. Create a new compiler instance."));
+		return false;
+	}
+	bCompiled = true;
+
 	StateTree = &InStateTree;
 	EditorData = Cast<UStateTreeEditorData>(StateTree->EditorData);
 	if (!EditorData)
diff --git a/Engine/Plugins/Runtime/StateTree/Source/StateTreeEditorModule/Public/StateTreeCompiler.h b/Engine/Plugins/Runtime/StateTree/Source/StateTreeEditorModule/Public/StateTreeCompiler.h
--- a/Engine/Plugins/Runtime/StateTree/Source/StateTreeEditorModule/Public/StateTreeCompiler.h
+++ b/Engine/Plugins/Runtime/StateTree/Source/StateTreeEditorModule/Public/StateTreeCompiler.h
@@ -98,6 +98,9 @@ private:
 	TArray<FInstancedStruct> SharedInstanceStructs;
 	
 	FStateTreePropertyBindingCompiler BindingsCompiler;
+
+	/** The Compile function executed. */
+	bool bCompiled = false;
 };
 
 
