StateTree: Move Describe functions from .h to .cpp

#rnx

diff --git a/Engine/Plugins/Runtime/StateTree/Source/StateTreeModule/Private/StateTreeTypes.cpp b/Engine/Plugins/Runtime/StateTree/Source/StateTreeModule/Private/StateTreeTypes.cpp
--- a/Engine/Plugins/Runtime/StateTree/Source/StateTreeModule/Private/StateTreeTypes.cpp
+++ b/Engine/Plugins/Runtime/StateTree/Source/StateTreeModule/Private/StateTreeTypes.cpp
@@ -98,6 +98,23 @@ FStateTreeStateHandle FStateTreeStateHandle::FromCompletionStatus(const EStateTr
 	return {};
 }
 
+FString FStateTreeStateHandle::Describe() const
+{
+	switch (Index)
+	{
+	case InvalidIndex:
+		return TEXT("Invalid");
+	case SucceededIndex:
+		return TEXT("Succeeded");
+	case FailedIndex:
+		return TEXT("Failed");
+	case StoppedIndex:
+		return TEXT("Stopped");
+	default:
+		return FString::Printf(TEXT("%d"), Index);
+	}
+}
+
 bool FStateTreeStateLink::Serialize(FStructuredArchive::FSlot Slot)
 {
 	Slot.GetUnderlyingArchive().UsingCustomVersion(FStateTreeCustomVersion::GUID);
@@ -118,7 +135,40 @@ PRAGMA_DISABLE_DEPRECATION_WARNINGS
 		}
 PRAGMA_ENABLE_DEPRECATION_WARNINGS
 	}
-#endif // WITH_EDITORONLY_DAT
+#endif // WITH_EDITORONLY_DATA
+}
+
+FString FStateTreeDataHandle::Describe() const
+{
+	switch (Source)
+	{
+	case EStateTreeDataSourceType::None:
+		return TEXT("None");
+	case EStateTreeDataSourceType::GlobalInstanceData:
+		return FString::Printf(TEXT("Global[%d]"), Index);
+	case EStateTreeDataSourceType::GlobalInstanceDataObject:
+		return FString::Printf(TEXT("GlobalO[%d]"), Index);
+	case EStateTreeDataSourceType::ActiveInstanceData:
+		return FString::Printf(TEXT("Active[%d]"), Index);
+	case EStateTreeDataSourceType::ActiveInstanceDataObject:
+		return FString::Printf(TEXT("ActiveO[%d]"), Index);
+	case EStateTreeDataSourceType::SharedInstanceData:
+		return FString::Printf(TEXT("Shared[%d]"), Index);
+	case EStateTreeDataSourceType::SharedInstanceDataObject:
+		return FString::Printf(TEXT("SharedO[%d]"), Index);
+	case EStateTreeDataSourceType::ContextData:
+		return FString::Printf(TEXT("Context[%d]"), Index);
+	case EStateTreeDataSourceType::GlobalParameterData:
+		return TEXT("GlobalParam");
+	case EStateTreeDataSourceType::ExternalGlobalParameterData:
+		return TEXT("ExternalGlobalParam");
+	case EStateTreeDataSourceType::SubtreeParameterData:
+		return FString::Printf(TEXT("SubtreeParam[%d]"), Index);
+	case EStateTreeDataSourceType::StateParameterData:
+		return FString::Printf(TEXT("LinkedParam[%d]"), Index);
+	default:
+		return TEXT("---");
+	}
 }
 
 bool FCompactEventDesc::DoesEventMatchDesc(const FStateTreeEvent& Event) const
diff --git a/Engine/Plugins/Runtime/StateTree/Source/StateTreeModule/Public/StateTreeTypes.h b/Engine/Plugins/Runtime/StateTree/Source/StateTreeModule/Public/StateTreeTypes.h
--- a/Engine/Plugins/Runtime/StateTree/Source/StateTreeModule/Public/StateTreeTypes.h
+++ b/Engine/Plugins/Runtime/StateTree/Source/StateTreeModule/Public/StateTreeTypes.h
@@ -259,17 +259,7 @@ struct STATETREEMODULE_API FStateTreeStateHandle
 	bool operator==(const FStateTreeStateHandle& RHS) const { return Index == RHS.Index; }
 	bool operator!=(const FStateTreeStateHandle& RHS) const { return Index != RHS.Index; }
 
-	FString Describe() const
-	{
-		switch (Index)
-		{
-		case InvalidIndex:		return TEXT("Invalid");
-		case SucceededIndex:	return TEXT("Succeeded");
-		case FailedIndex: 		return TEXT("Failed");
-		case StoppedIndex: 		return TEXT("Stopped");
-		default: 				return FString::Printf(TEXT("%d"), Index);
-		}
-	}
+	FString Describe() const;
 
 	UPROPERTY()
 	uint16 Index = InvalidIndex;
@@ -363,9 +353,9 @@ struct STATETREEMODULE_API FStateTreeDataHandle
 		Hash = HashCombineFast(Hash, GetTypeHash(Handle.StateHandle));
 		return Hash;
 	}
-	
+
 	FStateTreeDataHandle() = default;
-	
+
 	explicit FStateTreeDataHandle(const EStateTreeDataSourceType InSource, const uint16 InIndex, const FStateTreeStateHandle InStateHandle = FStateTreeStateHandle::Invalid)
 		: Source(InSource)
 		, Index(InIndex)
@@ -408,7 +398,7 @@ struct STATETREEMODULE_API FStateTreeDataHandle
 	{
 		return Source == RHS.Source && Index == RHS.Index && StateHandle == RHS.StateHandle;
 	}
-	
+
 	bool operator!=(const FStateTreeDataHandle& RHS) const
 	{
 		return !(*this == RHS);
@@ -435,7 +425,7 @@ struct STATETREEMODULE_API FStateTreeDataHandle
 			|| Source == EStateTreeDataSourceType::ActiveInstanceDataObject
 			|| Source == EStateTreeDataSourceType::SharedInstanceDataObject;
 	}
-	
+
 	FStateTreeDataHandle ToObjectSource() const
 	{
 		switch (Source)
@@ -450,41 +440,9 @@ struct STATETREEMODULE_API FStateTreeDataHandle
 			return *this;
 		}
 	}
-	
-	FString Describe() const
-	{
 
-		switch (Source)
-		{
-		case EStateTreeDataSourceType::None:
-			return TEXT("None");
-		case EStateTreeDataSourceType::GlobalInstanceData:
-			return FString::Printf(TEXT("Global[%d]"), Index);
-		case EStateTreeDataSourceType::GlobalInstanceDataObject:
-			return FString::Printf(TEXT("GlobalO[%d]"), Index);
-		case EStateTreeDataSourceType::ActiveInstanceData:
-			return FString::Printf(TEXT("Active[%d]"), Index);
-		case EStateTreeDataSourceType::ActiveInstanceDataObject:
-			return FString::Printf(TEXT("ActiveO[%d]"), Index);
-		case EStateTreeDataSourceType::SharedInstanceData:
-			return FString::Printf(TEXT("Shared[%d]"), Index);
-		case EStateTreeDataSourceType::SharedInstanceDataObject:
-			return FString::Printf(TEXT("SharedO[%d]"), Index);
-		case EStateTreeDataSourceType::ContextData:
-			return FString::Printf(TEXT("Context[%d]"), Index);
-		case EStateTreeDataSourceType::GlobalParameterData:
-			return FString::Printf(TEXT("GlobalParam"));
-		case EStateTreeDataSourceType::ExternalGlobalParameterData:
-			return FString::Printf(TEXT("ExternalGlobalParam"));
-		case EStateTreeDataSourceType::SubtreeParameterData:
-			return FString::Printf(TEXT("SubtreeParam[%d]"), Index);
-		case EStateTreeDataSourceType::StateParameterData:
-			return FString::Printf(TEXT("LinkedParam[%d]"), Index);
-		default:
-			return TEXT("---");
-		}
-	}
-	
+	FString Describe() const;
+
 private:
 	UPROPERTY()
 	EStateTreeDataSourceType Source = EStateTreeDataSourceType::None;
