StateTree: Add task completion control. You can now define the task that will complete the states. Any tasks that are not on the list are considered "background" tasks. They will not complete the state when they are completed. A state can specify if all the tasks need to be completed or if any tasks can be completed to complete the state.
A task failure completes the state as a failure, only if the task is in the list, and regardless of whether the control is any or all.

diff --git a/Engine/Plugins/Runtime/StateTree/Source/StateTreeEditorModule/Private/StateTreeCompiler.cpp b/Engine/Plugins/Runtime/StateTree/Source/StateTreeEditorModule/Private/StateTreeCompiler.cpp
--- a/Engine/Plugins/Runtime/StateTree/Source/StateTreeEditorModule/Private/StateTreeCompiler.cpp
+++ b/Engine/Plugins/Runtime/StateTree/Source/StateTreeEditorModule/Private/StateTreeCompiler.cpp
@@ -216,6 +216,67 @@ namespace UE::StateTree::Compiler
 		return ::DuplicateObject(&(*InstanceObject), &(*Owner), NewObjectName);
 	}
 
+	struct FCompletionTasksMaskResult
+	{
+		FStateTreeTasksCompletionStatus::FMaskType Mask;
+		int32 MaskBufferIndex;	// Index of FStateTreeTasksCompletionStatus::Buffer.
+		int32 MaskFirstTaskBitOffset; // Inside FStateTreeTasksCompletionStatus::Buffer[MaskBufferIndex], the bit offset of the first task.
+		int32 FullMaskEndTaskBitOffset; // the next bit that the next child can take in the full FStateTreeTasksCompletionStatus::Buffer.
+	};
+
+	/** Makes the completion mask for the state or frame. */
+	FCompletionTasksMaskResult MakeCompletionTasksMask(int32 FullStartBitIndex, TConstArrayView<FStateTreeEditorNode> AllTasks, TConstArrayView<int32> ValidTasks)
+	{
+		FStateTreeTasksCompletionStatus::FMaskType Mask = 0;
+		int32 NumberOfBitsNeeded = 0;
+		const int32 NumberOfTasks = ValidTasks.Num();
+
+		// No task, state/frame needs at least one flag to set the state itself completes (ie. for linked state).
+		//Each state will take at least 1 bit.
+		if (NumberOfTasks == 0)
+		{
+			Mask = 1;
+			NumberOfBitsNeeded = 1;
+		}
+		else
+		{
+			for (int32 Index = 0; Index < NumberOfTasks; ++Index)
+			{
+				const int32 TaskIndex = ValidTasks[Index];
+				const FStateTreeTaskBase& Task = AllTasks[TaskIndex].Node.Get<FStateTreeTaskBase>();
+				Mask <<= 1;
+				if (Task.bConsideredForCompletion)
+				{
+					Mask |= 1;
+				}
+			}
+			NumberOfBitsNeeded = NumberOfTasks;
+		}
+
+		constexpr int32 NumberOfBitsPerMask = sizeof(FStateTreeTasksCompletionStatus::FMaskType) * 8;
+
+		// Is the new amount of bits bring up over the next buffer?
+		const int32 CurrentEndBitIndex = FullStartBitIndex + NumberOfBitsNeeded;
+		if (CurrentEndBitIndex / NumberOfBitsPerMask != FullStartBitIndex / NumberOfBitsPerMask)
+		{
+			// Do not shift the mask. Use the next int32
+			const int32 NewMaskBufferIndex = CurrentEndBitIndex / NumberOfBitsPerMask;
+			const int32 NewMaskFirstTaskBitOffset = 0;
+			const int32 NewMaskEndTaskBitOffset = (NewMaskBufferIndex * NumberOfBitsPerMask) + NumberOfBitsNeeded;
+			return { .Mask = Mask, .MaskBufferIndex = NewMaskBufferIndex, .MaskFirstTaskBitOffset = NewMaskFirstTaskBitOffset, .FullMaskEndTaskBitOffset = NewMaskEndTaskBitOffset };
+		}
+		else
+		{
+			const int32 NewMaskBufferIndex = FullStartBitIndex / NumberOfBitsPerMask;
+			const int32 NewMaskFirstTaskBitOffset = FullStartBitIndex % NumberOfBitsPerMask;
+			const int32 NewMaskEndTaskBitOffset = CurrentEndBitIndex;
+
+			Mask <<= NewMaskFirstTaskBitOffset;
+
+			return { .Mask = Mask, .MaskBufferIndex = NewMaskBufferIndex, .MaskFirstTaskBitOffset = NewMaskFirstTaskBitOffset, .FullMaskEndTaskBitOffset = NewMaskEndTaskBitOffset };
+		}
+	}
+
 }; // UE::StateTree::Compiler
 
 bool FStateTreeCompiler::Compile(UStateTree& InStateTree)
@@ -564,6 +625,14 @@ bool FStateTreeCompiler::CreateStateRecursive(UStateTreeState& State, const FSta
 	}
 	StateTree->States[StateIdx].ChildrenEnd = uint16(ChildrenEnd); // Not using CompactState here because the array may have changed.
 	
+	// create sub frame info
+	if (!Parent.IsValid())
+	{
+		FCompactStateTreeFrame& CompactFrame = StateTree->Frames.AddDefaulted_GetRef();
+		CompactFrame.RootState = FStateTreeStateHandle((uint16)StateIdx);
+		CompactFrame.NumberOfTasksStatusMasks = 0;
+	}
+
 	return true;
 }
 
@@ -575,7 +644,7 @@ bool FStateTreeCompiler::CreateConditions(UStateTreeState& State, const FString&
 	{
 		const bool bIsFirst = Index == 0;
 		const FStateTreeEditorNode& CondNode = Conditions[Index];
-		// First operand should be copied as we dont have a previous item to operate on.
+		// First operand should be copied as we don't have a previous item to operate on.
 		const EStateTreeExpressionOperand Operand = bIsFirst ? EStateTreeExpressionOperand::Copy : CondNode.ExpressionOperand;
 		// First indent must be 0 to make the parentheses calculation match.
 		const int32 CurrIndent = bIsFirst ? 0 : FMath::Clamp((int32)CondNode.ExpressionIndent, 0, UE::StateTree::MaxExpressionIndent);
@@ -645,9 +714,12 @@ bool FStateTreeCompiler::CreateGlobalTasks()
 		return false;
 	}
 	StateTree->GlobalTasksBegin = uint16(GlobalTasksBegin);
+	StateTree->CompletionGlobalTasksMask = 0;
 
-	for (FStateTreeEditorNode& TaskNode : EditorData->GlobalTasks)
+	TArray<int32, TInlineAllocator<32>> ValidTaskNodeIndex;
+	for (int32 TaskIndex = 0; TaskIndex < EditorData->GlobalTasks.Num(); ++TaskIndex)
 	{
+		FStateTreeEditorNode& TaskNode = EditorData->GlobalTasks[TaskIndex];
 		// Silently ignore empty nodes.
 		if (!TaskNode.Node.IsValid())
 		{
@@ -661,6 +733,28 @@ bool FStateTreeCompiler::CreateGlobalTasks()
 			bSucceeded = false;
 			continue;
 		}
+
+		ValidTaskNodeIndex.Add(TaskIndex);
+	}
+
+	if (ValidTaskNodeIndex.Num() > FStateTreeTasksCompletionStatus::MaxNumberOfTasksPerGroup)
+	{
+		Log.Reportf(EMessageSeverity::Error, FStateTreeBindableStructDesc(),
+			TEXT("Exceeds the maximum number of global tasks (%d)"), FStateTreeTasksCompletionStatus::MaxNumberOfTasksPerGroup);
+		return false;
+	}
+	
+	constexpr int32 CompletionGlobalTaskStartBitIndex = 0;
+	const UE::StateTree::Compiler::FCompletionTasksMaskResult MaskResult = UE::StateTree::Compiler::MakeCompletionTasksMask(CompletionGlobalTaskStartBitIndex, EditorData->GlobalTasks, ValidTaskNodeIndex);
+	StateTree->CompletionGlobalTasksMask = MaskResult.Mask;
+	GlobalTaskEndBit = MaskResult.FullMaskEndTaskBitOffset;
+	StateTree->CompletionGlobalTasksControl = EditorData->GlobalTasksCompletion;
+
+	if (MaskResult.MaskFirstTaskBitOffset != 0)
+	{
+		ensureMsgf(false, TEXT("Invalid bit offset %d. The Global task should start at 0."), MaskResult.MaskFirstTaskBitOffset);
+		Log.Reportf(EMessageSeverity::Error, FStateTreeBindableStructDesc(), TEXT("Internal Error. Global task bit offset starts at 0."));
+		return false;
 	}
 	
 	const int32 GlobalTasksNum = Nodes.Num() - GlobalTasksBegin;
@@ -681,27 +775,37 @@ bool FStateTreeCompiler::CreateStateTasksAndParameters()
 	bool bSucceeded = true;
 
 	// Index of the first instance data per state. Accumulated depth first.
-	TArray<int32> FirstInstanceDataIndex;
-	FirstInstanceDataIndex.Init(0, StateTree->States.Num());
+	struct FTaskAndParametersCompactState
+	{
+		int32 FirstInstanceDataIndex = 0;
+		int32 NextBitIndexForCompletionMask = 0;
+		bool bProcessed = false;
+	};
+	TArray<FTaskAndParametersCompactState> StateInfos;
+	StateInfos.SetNum(StateTree->States.Num());
 	
-	for (int32 i = 0; i < StateTree->States.Num(); i++)
+	for (int32 StateIndex = 0; StateIndex < StateTree->States.Num(); ++StateIndex)
 	{
-		FCompactStateTreeState& CompactState = StateTree->States[i];
-		const FStateTreeStateHandle CompactStateHandle(i);
-		UStateTreeState* State = SourceStates[i];
+		FCompactStateTreeState& CompactState = StateTree->States[StateIndex];
+		const FStateTreeStateHandle CompactStateHandle(StateIndex);
+		UStateTreeState* State = SourceStates[StateIndex];
 		check(State != nullptr);
 
 		// Carry over instance data count from parent.
 		if (CompactState.Parent.IsValid())
 		{
 			const FCompactStateTreeState& ParentCompactState = StateTree->States[CompactState.Parent.Index];
-			const int32 InstanceDataBegin = FirstInstanceDataIndex[CompactState.Parent.Index] + (int32)ParentCompactState.InstanceDataNum;
-			FirstInstanceDataIndex[i] = InstanceDataBegin;
+
+			check(StateInfos[StateIndex].bProcessed == false);
+			check(!bSucceeded || StateInfos[CompactState.Parent.Index].bProcessed == true);
+
+			const int32 InstanceDataBegin = StateInfos[CompactState.Parent.Index].FirstInstanceDataIndex + (int32)ParentCompactState.InstanceDataNum;
+			StateInfos[StateIndex].FirstInstanceDataIndex = InstanceDataBegin;
 
 			CompactState.Depth = ParentCompactState.Depth + 1;
 		}
 
-		int32 InstanceDataIndex = FirstInstanceDataIndex[i];
+		int32 InstanceDataIndex = StateInfos[StateIndex].FirstInstanceDataIndex;
 
 		FStateTreeCompilerLogStateScope LogStateScope(State, Log);
 
@@ -834,8 +938,10 @@ bool FStateTreeCompiler::CreateStateTasksAndParameters()
 		
 		bool bCreateTaskSucceeded = true;
 		int32 EnabledTasksNum = 0;
-		for (FStateTreeEditorNode& TaskNode : Tasks)
+		TArray<int32, TInlineAllocator<32>> ValidTaskNodeIndex;
+		for (int32 TaskIndex = 0; TaskIndex < Tasks.Num(); ++TaskIndex)
 		{
+			FStateTreeEditorNode& TaskNode = Tasks[TaskIndex];
 			// Silently ignore empty nodes.
 			if (!TaskNode.Node.IsValid())
 			{
@@ -855,6 +961,8 @@ bool FStateTreeCompiler::CreateStateTasksAndParameters()
 				bCreateTaskSucceeded = false;
 				continue;
 			}
+
+			ValidTaskNodeIndex.Add(TaskIndex);
 		}
 
 		if (!bCreateTaskSucceeded)
@@ -863,13 +971,23 @@ bool FStateTreeCompiler::CreateStateTasksAndParameters()
 		}
 		
 		const int32 TasksNum = Nodes.Num() - TasksBegin;
+		check(ValidTaskNodeIndex.Num() == TasksNum);
 		if (const auto Validation = UE::StateTree::Compiler::IsValidCount8(TasksNum); Validation.DidFail())
 		{
 			Validation.Log(Log, TEXT("TasksNum"));
 			return false;
 		}
 
-		const int32 InstanceDataNum = InstanceDataIndex - FirstInstanceDataIndex[i];
+		// Create tasks
+		if (TasksNum > FStateTreeTasksCompletionStatus::MaxNumberOfTasksPerGroup)
+		{
+			Log.Reportf(EMessageSeverity::Error, LinkedParamsDesc,
+				TEXT("Exceeds the maximum number of tasks (%d)"), FStateTreeTasksCompletionStatus::MaxNumberOfTasksPerGroup);
+			bSucceeded = false;
+			continue;
+		}
+
+		const int32 InstanceDataNum = InstanceDataIndex - StateInfos[StateIndex].FirstInstanceDataIndex;
 		if (const auto Validation = UE::StateTree::Compiler::IsValidCount8(InstanceDataNum); Validation.DidFail())
 		{
 			Validation.Log(Log, TEXT("InstanceDataNum"));
@@ -880,11 +998,79 @@ bool FStateTreeCompiler::CreateStateTasksAndParameters()
 		CompactState.EnabledTasksNum = uint8(EnabledTasksNum);
 		CompactState.InstanceDataNum = uint8(InstanceDataNum);
 
+		// Create completion mask
+		{
+			int32 StartBitIndex = 0;
+			if (CompactState.Parent.IsValid())
+			{
+				StartBitIndex = StateInfos[CompactState.Parent.Index].NextBitIndexForCompletionMask;
+			}
+			else
+			{
+				// Frame need an extra buffer for global tasks.
+				//Subtree do not contains global task.
+				const bool bFrameWithoutGlobalTasks = CompactState.Type == EStateTreeStateType::Subtree && StateIndex != 0;
+				StartBitIndex = bFrameWithoutGlobalTasks ? 0 : GlobalTaskEndBit;
+			}
+
+			const UE::StateTree::Compiler::FCompletionTasksMaskResult MaskResult = UE::StateTree::Compiler::MakeCompletionTasksMask(StartBitIndex, Tasks, ValidTaskNodeIndex);
+
+			const int32 CompletionTasksMaskBufferIndex = MaskResult.MaskBufferIndex;
+			if (const auto Validation = UE::StateTree::Compiler::IsValidCount8(CompletionTasksMaskBufferIndex); Validation.DidFail())
+			{
+				Validation.Log(Log, TEXT("CompletionTasksMaskBufferIndex"));
+				bSucceeded = false;
+				continue;
+			}
+			const int32 CompletionTasksMaskBitsOffset = MaskResult.MaskFirstTaskBitOffset;
+			if (const auto Validation = UE::StateTree::Compiler::IsValidCount8(CompletionTasksMaskBitsOffset); Validation.DidFail())
+			{
+				Validation.Log(Log, TEXT("CompletionTasksMaskBitsOffset"));
+				bSucceeded = false;
+				continue;
+			}
+
+			CompactState.CompletionTasksMask = MaskResult.Mask;
+			CompactState.CompletionTasksControl = State->TasksCompletion;
+			CompactState.CompletionTasksMaskBufferIndex = static_cast<uint8>(CompletionTasksMaskBufferIndex);
+			CompactState.CompletionTasksMaskBitsOffset = static_cast<uint8>(CompletionTasksMaskBitsOffset);
+			StateInfos[StateIndex].NextBitIndexForCompletionMask = MaskResult.FullMaskEndTaskBitOffset;
+
+			// Find Frame and update the number of masks.
+			{
+				FStateTreeStateHandle FrameHandle = CompactStateHandle;
+				while (true)
+				{
+					const FCompactStateTreeState* ParentState = StateTree->GetStateFromHandle(FrameHandle);
+					check(ParentState);
+					if (!ParentState->Parent.IsValid())
+					{
+						break;
+					}
+					FrameHandle = ParentState->Parent;
+				}
+				FCompactStateTreeFrame* FoundFrame = StateTree->Frames.FindByPredicate([FrameHandle](const FCompactStateTreeFrame& Frame)
+					{
+						return Frame.RootState == FrameHandle;
+					});
+				if (FoundFrame == nullptr)
+				{
+					Log.Reportf(EMessageSeverity::Error, LinkedParamsDesc, TEXT("The parent frame can't be found"));
+					bSucceeded = false;
+					continue;
+				}
+
+				FoundFrame->NumberOfTasksStatusMasks = FMath::Max(FoundFrame->NumberOfTasksStatusMasks, static_cast<uint8>(CompactState.CompletionTasksMaskBufferIndex+1));
+			}
+		}
+
 		if (!CreateBindingsForNodes(Tasks, FStateTreeIndex16(TasksBegin), InstanceStructs))
 		{
 			bSucceeded = false;
 			continue;
 		}
+
+		StateInfos[StateIndex].bProcessed = true;
 	}
 	
 	return bSucceeded;
diff --git a/Engine/Plugins/Runtime/StateTree/Source/StateTreeEditorModule/Public/StateTreeCompiler.h b/Engine/Plugins/Runtime/StateTree/Source/StateTreeEditorModule/Public/StateTreeCompiler.h
--- a/Engine/Plugins/Runtime/StateTree/Source/StateTreeEditorModule/Public/StateTreeCompiler.h
+++ b/Engine/Plugins/Runtime/StateTree/Source/StateTreeEditorModule/Public/StateTreeCompiler.h
@@ -97,6 +97,9 @@ private:
 	TArray<FInstancedStruct> InstanceStructs;
 	TArray<FInstancedStruct> SharedInstanceStructs;
 	
+	/** Cached result of MakeCompletionTasksMask for global tasks. Indicates where state tasks should start. */
+	int32 GlobalTaskEndBit = 0;
+
 	FStateTreePropertyBindingCompiler BindingsCompiler;
 
 	/** The Compile function executed. */
diff --git a/Engine/Plugins/Runtime/StateTree/Source/StateTreeEditorModule/Public/StateTreeEditorData.h b/Engine/Plugins/Runtime/StateTree/Source/StateTreeEditorModule/Public/StateTreeEditorData.h
--- a/Engine/Plugins/Runtime/StateTree/Source/StateTreeEditorModule/Public/StateTreeEditorData.h
+++ b/Engine/Plugins/Runtime/StateTree/Source/StateTreeEditorModule/Public/StateTreeEditorData.h
@@ -364,6 +364,9 @@ public:
 	UPROPERTY(EditDefaultsOnly, Category = "Global Tasks", meta = (BaseStruct = "/Script/StateTreeModule.StateTreeTaskBase", BaseClass = "/Script/StateTreeModule.StateTreeTaskBlueprintBase"))
 	TArray<FStateTreeEditorNode> GlobalTasks;
 
+	UPROPERTY(EditDefaultsOnly, Category = "Global Tasks")
+	EStateTreeTaskCompletionType GlobalTasksCompletion = EStateTreeTaskCompletionType::Any;
+
 	UPROPERTY(meta = (ExcludeFromHash))
 	FStateTreeEditorPropertyBindings EditorBindings;
 
diff --git a/Engine/Plugins/Runtime/StateTree/Source/StateTreeEditorModule/Public/StateTreeState.h b/Engine/Plugins/Runtime/StateTree/Source/StateTreeEditorModule/Public/StateTreeState.h
--- a/Engine/Plugins/Runtime/StateTree/Source/StateTreeEditorModule/Public/StateTreeState.h
+++ b/Engine/Plugins/Runtime/StateTree/Source/StateTreeEditorModule/Public/StateTreeState.h
@@ -360,6 +360,10 @@ public:
 	/** How to treat child states when this State is selected.  */
 	UPROPERTY(EditDefaultsOnly, Category = "State")
 	EStateTreeStateSelectionBehavior SelectionBehavior = EStateTreeStateSelectionBehavior::TrySelectChildrenInOrder;
+	
+	/** How tasks will complete the state. Only tasks that are considered for completion can complete the state. */
+	UPROPERTY(EditDefaultsOnly, Category = "State")
+	EStateTreeTaskCompletionType TasksCompletion = EStateTreeTaskCompletionType::Any;
 
 	/** Subtree to run as extension of this State. */
 	UPROPERTY(EditDefaultsOnly, Category = "State", Meta=(DirectStatesOnly, SubtreesOnly))
diff --git a/Engine/Plugins/Runtime/StateTree/Source/StateTreeModule/Private/StateTree.cpp b/Engine/Plugins/Runtime/StateTree/Source/StateTreeModule/Private/StateTree.cpp
--- a/Engine/Plugins/Runtime/StateTree/Source/StateTreeModule/Private/StateTree.cpp
+++ b/Engine/Plugins/Runtime/StateTree/Source/StateTreeModule/Private/StateTree.cpp
@@ -57,6 +57,14 @@ FGuid UStateTree::GetNodeIdFromIndex(const FStateTreeIndex16 NodeIndex) const
 	return Entry != nullptr ? Entry->Id : FGuid();
 }
 
+const FCompactStateTreeFrame* UStateTree::GetFrameFromHandle(const FStateTreeStateHandle StateHandle) const
+{
+	return Frames.FindByPredicate([StateHandle](const FCompactStateTreeFrame& Frame)
+		{
+			return Frame.RootState == StateHandle;
+		});
+}
+
 const FCompactStateTreeState* UStateTree::GetStateFromHandle(const FStateTreeStateHandle StateHandle) const
 {
 	return States.IsValidIndex(StateHandle.Index) ? &States[StateHandle.Index] : nullptr;
@@ -174,6 +182,7 @@ namespace UE::StateTree::Compiler
 void UStateTree::ResetCompiled()
 {
 	Schema = nullptr;
+	Frames.Reset();
 	States.Reset();
 	Transitions.Reset();
 	Nodes.Reset();
diff --git a/Engine/Plugins/Runtime/StateTree/Source/StateTreeModule/Private/StateTreeExecutionContext.cpp b/Engine/Plugins/Runtime/StateTree/Source/StateTreeModule/Private/StateTreeExecutionContext.cpp
--- a/Engine/Plugins/Runtime/StateTree/Source/StateTreeModule/Private/StateTreeExecutionContext.cpp
+++ b/Engine/Plugins/Runtime/StateTree/Source/StateTreeModule/Private/StateTreeExecutionContext.cpp
@@ -80,43 +80,101 @@ namespace UE::StateTree::ExecutionContext::Private
 		return false;
 	}
 
-	EStateTreeRunStatus GetRunStatusFromFinishTaskType(EStateTreeFinishTaskType FinishTask)
+	constexpr uint32 NumEStateTreeRunStatus()
 	{
-		static_assert((int32)EStateTreeRunStatus::Succeeded == (int32)EStateTreeFinishTaskType::Succeeded, "Succeeded do not match in both types.");
-		static_assert((int32)EStateTreeRunStatus::Failed == (int32)EStateTreeFinishTaskType::Failed, "Failed do not match in both types.");
-		return (EStateTreeRunStatus)FinishTask;
-	}
-
-	/** @return in order {Running, Failed, Succeeded, Stopped, Unset} */
-	EStateTreeRunStatus GetPriorityRunStatus(EStateTreeRunStatus A, EStateTreeRunStatus B)
-	{
-		constexpr auto NumEStateTreeRunStatus = []()
-		{
 #ifdef UE_STATETREE_ESTATETREERUNSTATUS_PRIVATE
 #error UE_STATETREE_ESTATETREERUNSTATUS_PRIVATE_ALREADY_DEFINED
 #endif
 
 #define UE_STATETREE_ESTATETREERUNSTATUS_PRIVATE(EnumValue) ++Number;
 
-			int32 Number = 0;
-			FOREACH_ENUM_ESTATETREERUNSTATUS(UE_STATETREE_ESTATETREERUNSTATUS_PRIVATE)
+		int32 Number = 0;
+		FOREACH_ENUM_ESTATETREERUNSTATUS(UE_STATETREE_ESTATETREERUNSTATUS_PRIVATE)
 
 #undef UE_STATETREE_ESTATETREERUNSTATUS_PRIVATE
 			return Number;
-		};
-		constexpr int32 PriorityMatrix[] = { 1, 4, 3, 2, 0 };
-		static_assert(NumEStateTreeRunStatus() == UE_ARRAY_COUNT(PriorityMatrix), "The number of entries in EStateTreeRunStatus changed. Update GetPriorityRunStatus.");
+	}
+	constexpr uint32 NumEStateTreeFinishTaskType()
+	{
+#ifdef UE_STATETREE_ESTATETREEFINISHTASKTYPE_PRIVATE
+#error UE_STATETREE_ESTATETREEFINISHTASKTYPE_PRIVATE_ALREADY_DEFINED
+#endif
+
+#define UE_STATETREE_ESTATETREEFINISHTASKTYPE_PRIVATE(EnumValue) ++Number;
+
+		int32 Number = 0;
+		FOREACH_ENUM_ESTATETREEFINISHTASKTYPE(UE_STATETREE_ESTATETREEFINISHTASKTYPE_PRIVATE)
+
+#undef UE_STATETREE_ESTATETREEFINISHTASKTYPE_PRIVATE
+			return Number;
+	}
+
+	/** @return in order {Failed, Succeeded, Stopped, Running, Unset} */
+	EStateTreeRunStatus GetPriorityRunStatus(EStateTreeRunStatus A, EStateTreeRunStatus B)
+	{
+		static_assert((int32)EStateTreeRunStatus::Running == 0);
+		static_assert((int32)EStateTreeRunStatus::Failed == 1);
+		static_assert((int32)EStateTreeRunStatus::Succeeded == 2);
+		static_assert((int32)EStateTreeRunStatus::Stopped == 3);
+		static_assert((int32)EStateTreeRunStatus::Unset == 4);
+		static_assert(NumEStateTreeRunStatus() == 5, "The number of entries in EStateTreeRunStatus changed. Update GetPriorityRunStatus.");
+
+		static constexpr int32 PriorityMatrix[] = { 1, 4, 3, 2, 0 };
 		return PriorityMatrix[(uint8)A] > PriorityMatrix[(uint8)B] ? A : B;
 	}
 
-	void RemoveTickProcessedFinishedTasks(FStateTreeExecutionState& Exec)
+	UE::StateTree::ETaskCompletionStatus CastToTaskStatus(EStateTreeFinishTaskType FinishTask)
 	{
-		Exec.FinishedTasks.RemoveAllSwap([](const UE::StateTree::FFinishedTask& Task)
+		static_assert(NumEStateTreeFinishTaskType() == 2, "The number of entries in EStateTreeFinishTaskType changed. Update CastToTaskStatus.");
+		return FinishTask == EStateTreeFinishTaskType::Succeeded ? UE::StateTree::ETaskCompletionStatus::Succeeded : UE::StateTree::ETaskCompletionStatus::Failed;
+	}
+
+	EStateTreeRunStatus CastToRunStatus(EStateTreeFinishTaskType FinishTask)
+	{
+		static_assert((int32)EStateTreeRunStatus::Succeeded == (int32)EStateTreeFinishTaskType::Succeeded, "Succeeded does not match in both types.");
+		static_assert((int32)EStateTreeRunStatus::Failed == (int32)EStateTreeFinishTaskType::Failed, "Failed does not match in both types.");
+		static_assert(NumEStateTreeFinishTaskType() == 2, "The number of entries in EStateTreeFinishTaskType changed. Update CastToRunStatus.");
+		return (EStateTreeRunStatus)FinishTask;
+	}
+
+	UE::StateTree::ETaskCompletionStatus CastToTaskStatus(EStateTreeRunStatus InStatus)
+	{
+		static_assert((int32)EStateTreeRunStatus::Running == 0);
+		static_assert((int32)EStateTreeRunStatus::Failed == 1);
+		static_assert((int32)EStateTreeRunStatus::Succeeded == 2);
+		static_assert((int32)EStateTreeRunStatus::Stopped == 3);
+		static_assert((int32)EStateTreeRunStatus::Unset == 4);
+		static_assert(NumEStateTreeRunStatus() == 5, "The number of entries in EStateTreeRunStatus changed. Update CastToTaskStatus.");
+
+		static constexpr UE::StateTree::ETaskCompletionStatus ConversionMatrix[] =
 			{
-				return Task.bTickProcessed;
-			});
+				UE::StateTree::ETaskCompletionStatus::Running,
+				UE::StateTree::ETaskCompletionStatus::Failed,
+				UE::StateTree::ETaskCompletionStatus::Succeeded,
+				UE::StateTree::ETaskCompletionStatus::Stopped,
+				UE::StateTree::ETaskCompletionStatus::Running,
+			};
+		return ConversionMatrix[(int32)InStatus];
 	}
-} // namespace
+	
+	EStateTreeRunStatus CastToRunStatus(UE::StateTree::ETaskCompletionStatus InStatus)
+	{
+		//@todo reorder EStateTreeRunStatus enum to be 1 to 1 with UE::StateTree::ETaskCompletionStatus
+		static_assert((int32)UE::StateTree::ETaskCompletionStatus::Running == 0);
+		static_assert((int32)UE::StateTree::ETaskCompletionStatus::Stopped == 1);
+		static_assert((int32)UE::StateTree::ETaskCompletionStatus::Succeeded == 2);
+		static_assert((int32)UE::StateTree::ETaskCompletionStatus::Failed == 3);
+		static_assert(UE::StateTree::NumberOfTaskStatus == 4, "The number of entries in EStateTreeRunStatus changed. Update CastToRunStatus.");
+		static constexpr EStateTreeRunStatus ConversionMatrix[] =
+		{
+			EStateTreeRunStatus::Running,
+			EStateTreeRunStatus::Stopped,
+			EStateTreeRunStatus::Succeeded,
+			EStateTreeRunStatus::Failed,
+		};
+		return ConversionMatrix[(int32)InStatus];
+	}
+}
 
 /**
  * FStateTreeReadOnlyExecutionContext implementation
@@ -254,7 +312,7 @@ FStateTreeScheduledTick FStateTreeReadOnlyExecutionContext::GetNextScheduledTick
 	}
 
 	// Completed task. For EnterState or for user that only called TickTasks and not TickTransitions.
-	if (Exec.FinishedTasks.Num() > 0)
+	if (Exec.bHasPendingCompletedState)
 	{
 		return FStateTreeScheduledTick::MakeNextFrame();
 	}
@@ -1060,6 +1118,10 @@ EStateTreeRunStatus FStateTreeExecutionContext::Start(FStartParameters Parameter
 	InitFrame.RootState = FStateTreeStateHandle::Root;
 	InitFrame.ActiveStates = {};
 	InitFrame.bIsGlobalFrame = true;
+
+	const FCompactStateTreeFrame* FrameInfo = RootStateTree.GetFrameFromHandle(FStateTreeStateHandle::Root);
+	ensureAlwaysMsgf(FrameInfo, TEXT("The compiled data is invalid. It should contains the information for the root frame."));
+	InitFrame.ActiveTasksStatus = FrameInfo ? FStateTreeTasksCompletionStatus(*FrameInfo) : FStateTreeTasksCompletionStatus();
 	
 	UpdateInstanceData({}, Exec.ActiveFrames);
 	Exec.RandomStream.Initialize(Parameters.RandomSeed.IsSet() ? Parameters.RandomSeed.GetValue() : FPlatformTime::Cycles());
@@ -1345,29 +1407,7 @@ EStateTreeRunStatus FStateTreeExecutionContext::TickTriggerTransitions()
 		return PreludeResult;
 	}
 
-	if (!bCompletedStatesUpdated)
-	{
-		// Create the completed state list from the process tasks list.
-		//Do not mark the task as processed because the task could be added between TickTasks() and TickTriggerTransitions()
-		constexpr bool bMarkTaskProcessed = false;
-		UpdateCompletedStateList(bMarkTaskProcessed);
-	}
-
-	// If we have pending global task that completed, then stop the tree (like it would with a regular Tick).
-	//This should not occurs and its probably a user mistake to not tick the system in the proper order.
-	const EStateTreeRunStatus GlobalTaskStatus = GetGlobalTasksCompletedStatesStatus();
-	if (GlobalTaskStatus == EStateTreeRunStatus::Running)
-	{
-		TickTriggerTransitionsInternal();
-	}
-	else
-	{
-		//@TODO check what we want to do with global task from linked asset state. Should they request a stop or close their own tree.
-		FStateTreeExecutionState& Exec = GetExecState();
-		UE_STATETREE_DEBUG_LOG_EVENT(this, Log, TEXT("Global tasks completed (%s), stopping the tree"), *UEnum::GetDisplayValueAsText(GlobalTaskStatus).ToString());
-		STATETREE_LOG(Log, TEXT("Global tasks completed (%s), stopping the tree"), *UEnum::GetDisplayValueAsText(GlobalTaskStatus).ToString());
-		Exec.RequestedStop = GlobalTaskStatus;
-	}
+	TickTriggerTransitionsInternal();
 
 	return TickPostlude();
 }
@@ -1391,10 +1431,6 @@ void FStateTreeExecutionContext::TickUpdateTasksInternal(float DeltaTime)
 		DelayedState.TimeLeft -= DeltaTime;
 	}
 
-	// Create the completed list and mark the task as processed.
-	constexpr bool bMarkTaskProcessed = true;
-	UpdateCompletedStateList(bMarkTaskProcessed);
-
 	auto LogRequestStop = [&Exec, this]()
 		{
 			if (Exec.RequestedStop != EStateTreeRunStatus::Unset) // -V547
@@ -1657,246 +1693,105 @@ void FStateTreeExecutionContext::FinishTask(const FStateTreeTaskBase& Task, ESta
 		return;
 	}
 
-	// Like GetInstanceData, only accept task if we are currently processing them.
+	// Like GetInstanceData, only accept task if we are currently processing.
 	if (!ensure(CurrentNode == &Task))
 	{
 		return;
 	}
 	check(CurrentlyProcessedFrame);
-	check(CurrentlyProcessedState.IsValid());
 	check(CurrentNodeIndex >= 0);
 
-
 	FStateTreeExecutionState& Exec = GetExecState();
 
-	const FActiveFrameID ActiveFrameID = CurrentlyProcessedFrame->FrameID;
-	check(ActiveFrameID.IsValid());
-	const FStateTreeIndex16 TaskIndex = FStateTreeIndex16(CurrentNodeIndex);
-	const EStateTreeRunStatus RunStatus = ExecutionContext::Private::GetRunStatusFromFinishTaskType(FinishType);
+	const UStateTree* CurrentStateTree = CurrentlyProcessedFrame->StateTree;
+	const ETaskCompletionStatus TaskStatus = ExecutionContext::Private::CastToTaskStatus(FinishType);
+
 	if (CurrentlyProcessedState.IsValid())
 	{
+		check(CurrentStateTree->States.IsValidIndex(CurrentlyProcessedState.Index));
+		const FCompactStateTreeState& State = CurrentStateTree->States[CurrentlyProcessedState.Index];
+
 		const int32 ActiveStateIndex = CurrentlyProcessedFrame->ActiveStates.IndexOfReverse(CurrentlyProcessedState);
 		check(ActiveStateIndex != INDEX_NONE);
-		const FActiveStateID ActiveStateID = CurrentlyProcessedFrame->ActiveStates.StateIDs[ActiveStateIndex];
-		constexpr FFinishedTask::EReasonType Reason = FFinishedTask::EReasonType::StateTask;
-		constexpr bool bTickProcessed = false;
-		Exec.FinishedTasks.Emplace(ActiveFrameID, ActiveStateID, TaskIndex, RunStatus, Reason, bTickProcessed);
+		
+		const int32 StateTaskIndex = CurrentNodeIndex - State.TasksBegin;
+		check(StateTaskIndex >= 0);
+
+		UE::StateTree::FTasksCompletionStatus StateTasksStatus = const_cast<FStateTreeExecutionFrame*>(CurrentlyProcessedFrame)->ActiveTasksStatus.GetStatus(State);
+		StateTasksStatus.SetStatusWithPriority(StateTaskIndex, TaskStatus);
+		Exec.bHasPendingCompletedState = StateTasksStatus.IsCompleted();
 	}
 	else
 	{
-		constexpr FFinishedTask::EReasonType Reason = FFinishedTask::EReasonType::GlobalTask;
-		constexpr bool bTickProcessed = false;
-		Exec.FinishedTasks.Emplace(ActiveFrameID, FActiveStateID::Invalid, TaskIndex, RunStatus, Reason, bTickProcessed);
+		// global frame
+		const int32 FrameTaskIndex = CurrentNodeIndex - CurrentStateTree->GlobalTasksBegin;
+		check(FrameTaskIndex >= 0);
+		UE::StateTree::FTasksCompletionStatus GlobalTasksStatus = const_cast<FStateTreeExecutionFrame*>(CurrentlyProcessedFrame)->ActiveTasksStatus.GetStatus(CurrentStateTree);
+		GlobalTasksStatus.SetStatusWithPriority(FrameTaskIndex, TaskStatus);
+		Exec.bHasPendingCompletedState = GlobalTasksStatus.IsCompleted();
 	}
 }
 
 // Deprecated
+PRAGMA_DISABLE_DEPRECATION_WARNINGS
 void FStateTreeExecutionContext::FinishTask(const UE::StateTree::FFinishedTask& Task, EStateTreeFinishTaskType FinishType)
 {
-	using namespace UE::StateTree;
-
-	if (!IsValid())
-	{
-		STATETREE_LOG(Warning, TEXT("%hs: StateTree context is not initialized properly ('%s' using StateTree '%s')"),
-			__FUNCTION__, *GetNameSafe(&Owner), *GetFullNameSafe(&RootStateTree));
-		return;
-	}
-
 	FStateTreeExecutionState& Exec = GetExecState();
-
-	if (CurrentNodeIndex == Task.TaskIndex.AsInt32() && CurrentlyProcessedFrame && Task.FrameID == CurrentlyProcessedFrame->FrameID)
+	FStateTreeExecutionFrame* Frame = Exec.FindActiveFrame(Task.FrameID);
+	if (Frame == nullptr)
 	{
-		FFinishedTask& NewTask = Exec.FinishedTasks.Add_GetRef(Task);
-		// Deprecated path. Need to reset the flag.
-		NewTask.bTickProcessed = false;
-	}
-	else if (IsFinishedTaskValid(Task))
-	{
-		FFinishedTask& CompletedTask = Exec.FinishedTasks.Add_GetRef(Task);
-		MarkStateCompleted(CompletedTask);
-		// Deprecated path. Need to reset the flag to the original value.
-		CompletedTask.bTickProcessed = false;
-		ScheduleNextTick();
+		return;
 	}
-}
-
-void FStateTreeExecutionContext::FinishTaskInternal(UE::StateTree::FFinishedTask::EReasonType Reason, UE::StateTree::FActiveFrameID ActiveFrameID, UE::StateTree::FActiveStateID ActiveStateID, int32 TaskIndex, EStateTreeRunStatus RunStatus)
-{
-	using namespace UE::StateTree;
 
-	FStateTreeExecutionState& Exec = GetExecState();
-
-	// Add it for later. In case the user called TickTasks and will call TickTransitions later.
-	const FStateTreeIndex16 TaskIndex16 = TaskIndex >= 0 ? FStateTreeIndex16(TaskIndex) : FStateTreeIndex16::Invalid;
-	
-	constexpr bool bTickProcessed = true;
-	UE::StateTree::FFinishedTask& Task = Exec.FinishedTasks.Emplace_GetRef(ActiveFrameID, ActiveStateID, TaskIndex16, RunStatus, Reason, bTickProcessed);
-	MarkStateCompleted(Task);
-}
-
-bool FStateTreeExecutionContext::IsFinishedTaskValid(const UE::StateTree::FFinishedTask& Task) const
-{
 	using namespace UE::StateTree;
 
-	const FStateTreeExecutionState& Exec = GetExecState();
-
-	// Is frame valid
-	const FStateTreeExecutionFrame* ExecutionFrame = Exec.FindActiveFrame(Task.FrameID);
-	if (ExecutionFrame == nullptr)
-	{
-		return false;
-	}
-
-	const UStateTree* StateTree = ExecutionFrame->StateTree;
-	check(StateTree);
-
-	if (Task.Reason == FFinishedTask::EReasonType::StateTask)
+	const UE::StateTree::ETaskCompletionStatus Status = ExecutionContext::Private::CastToTaskStatus(Task.RunStatus);
+	if (Task.Reason == FFinishedTask::EReasonType::GlobalTask)
 	{
-		if (!ensure(Task.StateID.IsValid() && Task.TaskIndex.IsValid()))
+		if (Frame->bIsGlobalFrame)
 		{
-			return false;
+			Frame->ActiveTasksStatus.GetStatus(Frame->StateTree).SetStatusWithPriority(Task.TaskIndex.AsInt32(), Status);
 		}
-
-		const int32 ActiveStateIndex = ExecutionFrame->ActiveStates.IndexOfReverse(Task.StateID);
-		if (ActiveStateIndex == INDEX_NONE)
-		{
-			return false;
-		}
-
-		const int32 TreeStateIndex = ExecutionFrame->ActiveStates[ActiveStateIndex].Index;
-		check(StateTree->GetStates().IsValidIndex(TreeStateIndex));
-		const FCompactStateTreeState& State = StateTree->GetStates()[TreeStateIndex];
-		const int32 BeginTaskIndex = State.TasksBegin;
-		const int32 EndTaskIndex = State.TasksBegin + State.TasksNum;
-		const bool bIsValidTask = Task.TaskIndex.AsInt32() >= BeginTaskIndex && Task.TaskIndex.AsInt32() < EndTaskIndex;
-		ensureMsgf(bIsValidTask, TEXT("The finished task has a task index that doesn't match the state. Have you recompiled the asset runtime?"));
-		return bIsValidTask;
 	}
-	else if (Task.Reason == FFinishedTask::EReasonType::GlobalTask)
-	{
-		if (!ensure(ExecutionFrame->bIsGlobalFrame && Task.TaskIndex.IsValid()))
-		{
-			return false;
-		}
-
-		const int32 BeginTaskIndex = StateTree->GlobalTasksBegin;
-		const int32 EndTaskIndex = StateTree->GlobalTasksBegin + StateTree->GlobalTasksNum;
-		const bool bIsValidTask = Task.TaskIndex.AsInt32() >= BeginTaskIndex && Task.TaskIndex.AsInt32() < EndTaskIndex;
-		ensureMsgf(bIsValidTask, TEXT("The finished task has a task index that is not global. Have you recompiled the asset runtime?"));
-		return bIsValidTask;
-	}
-	else if (Task.Reason == FFinishedTask::EReasonType::InternalTransition)
-	{
-		if (!ensure(Task.StateID.IsValid()))
-		{
-			return false;
-		}
-
-		return ExecutionFrame->ActiveStates.Contains(Task.StateID);
-	}
-	return false;
-}
-
-void FStateTreeExecutionContext::UpdateCompletedStateList()
-{
-	constexpr bool bMarkTaskProcessed = false;
-	UpdateCompletedStateList(bMarkTaskProcessed);
-}
-
-void FStateTreeExecutionContext::UpdateCompletedStateList(bool bMarkTaskProcessed)
-{
-	bCompletedStatesUpdated = true;
-
-	// Recreate the completed states list from the pending list
-	CompletedStates.Empty();
-
-	FStateTreeExecutionState& Exec = GetExecState();
-	for (int32 Index = Exec.FinishedTasks.Num() - 1; Index >= 0; --Index)
-	{
-		UE::StateTree::FFinishedTask& Task = Exec.FinishedTasks[Index];
-		// Because the user can tick only the transition (and not the task), we should not consider any finished tasks that were not ticked.
-		if (Task.bTickProcessed || bMarkTaskProcessed)
-		{
-			if (IsFinishedTaskValid(Task))
-			{
-				MarkStateCompleted(Task);
-			}
-			else
-			{
-				Exec.FinishedTasks.RemoveAtSwap(Index);
-			}
-		}
-	}
-}
-
-void FStateTreeExecutionContext::MarkStateCompleted(UE::StateTree::FFinishedTask& NewFinishedTask)
-{
-	NewFinishedTask.bTickProcessed = true;
-
-	bool bAdd = true;
-	// The priority is global tasks and failure. Only keep 1 FinishedTask per Frame per State.
-	//If the global task exist and a state (of the same tree) completes, only keep the global
-	for (int32 CompletedStateIndex = CompletedStates.Num() - 1; CompletedStateIndex >= 0; --CompletedStateIndex)
+	else
 	{
-		FCompletedState& ExistingCompletedState = CompletedStates[CompletedStateIndex];
-		if (ExistingCompletedState.FrameID == NewFinishedTask.FrameID)
+		const int32 FoundIndex = Frame->ActiveStates.IndexOfReverse(Task.StateID);
+		if (FoundIndex != INDEX_NONE)
 		{
-			// If the existing one is global and the new one is global, only update the finish type.
-			// If the existing one is global and the new one is state, discard new.
-			// If the existing one is state and the new one is global, discard the existing and add the new.
-			// If the existing one is state and match the new one, only update the finish type.
-			
-			if (ExistingCompletedState.Reason == UE::StateTree::FFinishedTask::EReasonType::GlobalTask)
+			const FStateTreeStateHandle StateHandle = Frame->ActiveStates[FoundIndex];
+			const FCompactStateTreeState* State = Frame->StateTree->GetStateFromHandle(StateHandle);
+			if (State != nullptr)
 			{
-				if (NewFinishedTask.Reason == UE::StateTree::FFinishedTask::EReasonType::GlobalTask)
+				if (Task.Reason == FFinishedTask::EReasonType::InternalTransition)
 				{
-					ExistingCompletedState.RunStatus = UE::StateTree::ExecutionContext::Private::GetPriorityRunStatus(ExistingCompletedState.RunStatus, NewFinishedTask.RunStatus);
-				}
-				bAdd = false;
-				break;
-			}
-			else
-			{
-				check(ExistingCompletedState.Reason == UE::StateTree::FFinishedTask::EReasonType::StateTask
-					|| ExistingCompletedState.Reason == UE::StateTree::FFinishedTask::EReasonType::InternalTransition);
-				if (NewFinishedTask.Reason == UE::StateTree::FFinishedTask::EReasonType::GlobalTask)
-				{
-					CompletedStates.RemoveAtSwap(CompletedStateIndex);
-					break;
+					Frame->ActiveTasksStatus.GetStatus(*State).SetCompletionStatus(Status);
 				}
-				else if (ExistingCompletedState.StateID == NewFinishedTask.StateID)
+				else
 				{
-					ExistingCompletedState.RunStatus = UE::StateTree::ExecutionContext::Private::GetPriorityRunStatus(ExistingCompletedState.RunStatus, NewFinishedTask.RunStatus);
-					bAdd = false;
-					break;
+					check(Task.Reason == FFinishedTask::EReasonType::StateTask);
+					Frame->ActiveTasksStatus.GetStatus(*State).SetStatusWithPriority(Task.TaskIndex.AsInt32(), Status);
 				}
 			}
 		}
 	}
+}
 
-	if (bAdd)
-	{
-		CompletedStates.Add(FCompletedState(NewFinishedTask));
-	}
+// Deprecated
+bool FStateTreeExecutionContext::IsFinishedTaskValid(const UE::StateTree::FFinishedTask& Task) const
+{
+	return false;
 }
 
-EStateTreeRunStatus FStateTreeExecutionContext::GetGlobalTasksCompletedStatesStatus() const
+// Deprecated
+void FStateTreeExecutionContext::UpdateCompletedStateList()
+{
+}
+
+// Deprecated
+void FStateTreeExecutionContext::MarkStateCompleted(UE::StateTree::FFinishedTask& NewFinishedTask)
 {
-	EStateTreeRunStatus Result = EStateTreeRunStatus::Running;
-	for (const FCompletedState& CompletedState : CompletedStates)
-	{
-		if (CompletedState.Reason == UE::StateTree::FFinishedTask::EReasonType::GlobalTask)
-		{
-			// return failure in priority
-			if (CompletedState.RunStatus == EStateTreeRunStatus::Failed)
-			{
-				return EStateTreeRunStatus::Failed;
-			}
-			Result = CompletedState.RunStatus;
-		}
-	}
-	return Result;
 }
+PRAGMA_DISABLE_DEPRECATION_WARNINGS
 
 void FStateTreeExecutionContext::UpdateInstanceData(TConstArrayView<FStateTreeExecutionFrame> CurrentActiveFrames, TArrayView<FStateTreeExecutionFrame> NextActiveFrames)
 {
@@ -2778,7 +2673,6 @@ EStateTreeRunStatus FStateTreeExecutionContext::EnterState(FStateTreeTransitionR
 	CaptureNewStateEvents(Exec.ActiveFrames, Transition.NextActiveFrames, Transition.NextActiveFrameEvents);
 
 	Exec.StateChangeCount++;
-	UE::StateTree::ExecutionContext::Private::RemoveTickProcessedFinishedTasks(Exec);
 	Exec.EnterStateFailedFrameIndex = FStateTreeIndex16::Invalid; // This will make all tasks to be accepted.
 	Exec.EnterStateFailedTaskIndex = FStateTreeIndex16::Invalid; // This will make all tasks to be accepted.
 	
@@ -2839,6 +2733,9 @@ EStateTreeRunStatus FStateTreeExecutionContext::EnterState(FStateTreeTransitionR
 					__FUNCTION__, *GetStateStatusString(Exec), *GetNameSafe(&Owner), *GetFullNameSafe(&RootStateTree));
 				break;
 			}
+			CurrentFrame.ActiveTasksStatus.Push(CurrentStateTree, State);
+
+			UE::StateTree::FTasksCompletionStatus CurrentStateTasksStatus = CurrentFrame.ActiveTasksStatus.GetStatus(State);
 			CurrentFrame.NumCurrentlyActiveStates = static_cast<uint8>(CurrentFrame.ActiveStates.Num());
 
 			FCurrentlyProcessedStateScope StateScope(*this, CurrentHandle);
@@ -2919,12 +2816,13 @@ EStateTreeRunStatus FStateTreeExecutionContext::EnterState(FStateTreeTransitionR
 			}
 
 			// Activate tasks on current state.
-			for (int32 TaskIndex = State.TasksBegin; TaskIndex < (State.TasksBegin + State.TasksNum); TaskIndex++)
+			for (int32 StateTaskIndex = 0; StateTaskIndex < State.TasksNum; ++StateTaskIndex)
 			{
-				const FStateTreeTaskBase& Task = NextFrame.StateTree->Nodes[TaskIndex].Get<const FStateTreeTaskBase>();
+				const int32 AssetTaskIndex = State.TasksBegin + StateTaskIndex;
+				const FStateTreeTaskBase& Task = NextFrame.StateTree->Nodes[AssetTaskIndex].Get<const FStateTreeTaskBase>();
 				const FStateTreeDataView TaskInstanceView = GetDataView(CurrentParentFrame, CurrentFrame, Task.InstanceDataHandle);
 
-				FNodeInstanceDataScope DataScope(*this, &Task, TaskIndex, Task.InstanceDataHandle, TaskInstanceView);
+				FNodeInstanceDataScope DataScope(*this, &Task, AssetTaskIndex, Task.InstanceDataHandle, TaskInstanceView);
 
 				// Copy bound properties.
 				if (Task.BindingsBatch.IsValid())
@@ -2945,46 +2843,27 @@ EStateTreeRunStatus FStateTreeExecutionContext::EnterState(FStateTreeTransitionR
 				if (bIsEnteringState && bShouldCallEnterState)
 				{
 					STATETREE_LOG(Verbose, TEXT("%*sTask '%s'.EnterState()"), (FrameIndex + StateIndex + 1)*UE::StateTree::Debug::IndentSize, TEXT(""), *Task.Name.ToString());
-					UE_STATETREE_DEBUG_TASK_ENTER_STATE(this, NextFrame.StateTree, FStateTreeIndex16(TaskIndex));
+					UE_STATETREE_DEBUG_TASK_ENTER_STATE(this, NextFrame.StateTree, FStateTreeIndex16(AssetTaskIndex));
 
-					const int32 PreviousFinishedTasksNum = Exec.FinishedTasks.Num();
-					EStateTreeRunStatus TaskStatus = EStateTreeRunStatus::Unset;
+					EStateTreeRunStatus TaskRunStatus = EStateTreeRunStatus::Unset;
 					{
 						QUICK_SCOPE_CYCLE_COUNTER(StateTree_Task_EnterState);
 						CSV_SCOPED_TIMING_STAT_EXCLUSIVE(StateTree_Task_EnterState);
 					
-						TaskStatus = Task.EnterState(*this, CurrentTransition);
-					}
-
-					// Process FinishTask (weak, or this context)
-					bool bFinishTaskInternal = true;
-					for (int32 FinishedTaskIndex = PreviousFinishedTasksNum; FinishedTaskIndex < Exec.FinishedTasks.Num(); ++FinishedTaskIndex)
-					{
-						UE::StateTree::FFinishedTask& FinishedTask = Exec.FinishedTasks[FinishedTaskIndex];
-						// If same task, then update TaskStatus.
-						if (TaskIndex == FinishedTask.TaskIndex.AsInt32()
-							&& CurrentStateID == FinishedTask.StateID
-							&& CurrentFrameID == FinishedTask.FrameID)
-						{
-							// Delayed from FinishTask(), use the best priority from what the task returns and FinishTask().
-							TaskStatus = UE::StateTree::ExecutionContext::Private::GetPriorityRunStatus(TaskStatus, FinishedTask.RunStatus);
-							bFinishTaskInternal = false;
-							MarkStateCompleted(FinishedTask);
-						}
+						TaskRunStatus = Task.EnterState(*this, CurrentTransition);
 					}
 
-					UE_STATETREE_DEBUG_TASK_EVENT(this, TaskIndex, TaskInstanceView, EStateTreeTraceEventType::OnEntered, TaskStatus);
+					UE::StateTree::ETaskCompletionStatus TaskStatus = UE::StateTree::ExecutionContext::Private::CastToTaskStatus(TaskRunStatus);
+					TaskStatus = CurrentStateTasksStatus.SetStatusWithPriority(StateTaskIndex, TaskStatus);
 
-					if (TaskStatus != EStateTreeRunStatus::Running && bFinishTaskInternal)
-					{
-						FinishTaskInternal(UE::StateTree::FFinishedTask::EReasonType::StateTask, CurrentFrame.FrameID, CurrentFrame.ActiveStates.StateIDs[StateIndex], TaskIndex, TaskStatus);
-					}
+					TaskRunStatus = UE::StateTree::ExecutionContext::Private::CastToRunStatus(TaskStatus);
+					UE_STATETREE_DEBUG_TASK_EVENT(this, AssetTaskIndex, TaskInstanceView, EStateTreeTraceEventType::OnEntered, TaskRunStatus);
 
-					Result = UE::StateTree::ExecutionContext::Private::GetPriorityRunStatus(Result, TaskStatus);
+					Result = UE::StateTree::ExecutionContext::Private::GetPriorityRunStatus(Result, TaskRunStatus);
 					if (Result == EStateTreeRunStatus::Failed)
 					{
 						Exec.EnterStateFailedFrameIndex = FStateTreeIndex16(FrameIndex);
-						Exec.EnterStateFailedTaskIndex = FStateTreeIndex16(TaskIndex);
+						Exec.EnterStateFailedTaskIndex = FStateTreeIndex16(AssetTaskIndex);
 						break;
 					}
 				}
@@ -3269,12 +3148,6 @@ void FStateTreeExecutionContext::ExitState(const FStateTreeTransitionResult& Tra
 					}
 				}
 
-				// Clean up completed state to prevent FinishTask called in ExitState.
-				CompletedStates.RemoveAllSwap([StateID = CurrentFrame.ActiveStates.StateIDs[StateIndex]](const FCompletedState& Completed)
-					{
-						return Completed.StateID == StateID;
-					});
-
 				// Delegate Listeners Cleanup
 				GetExecState().DelegateActiveListeners.RemoveAll(CurrentFrame.ActiveStates.StateIDs[StateIndex]);
 			}
@@ -3290,12 +3163,6 @@ void FStateTreeExecutionContext::ExitState(const FStateTreeTransitionResult& Tra
 				CurrentTransition.ChangeType = ExitCall.ChangeType;
 				CallStopOnEvaluatorsAndGlobalTasks(CurrentParentFrame, CurrentFrame, CurrentTransition);
 
-				// Clean up completed state to prevent FinishTask called in ExitState.
-				CompletedStates.RemoveAllSwap([FrameID = CurrentFrame.FrameID](const FCompletedState& Completed) 
-					{
-						return Completed.FrameID == FrameID;
-					});
-
 				// Delegate Listeners Cleanup
 				GetExecState().DelegateActiveListeners.RemoveAll(CurrentFrame.FrameID);
 			}
@@ -3437,23 +3304,14 @@ EStateTreeRunStatus FStateTreeExecutionContext::TickEvaluatorsAndGlobalTasks(con
 	return Result;
 }
 
-EStateTreeRunStatus FStateTreeExecutionContext::TickEvaluatorsAndGlobalTasksForFrame(const float DeltaTime, const bool bTickGlobalTasks, const int32 FrameIndex, const FStateTreeExecutionFrame* CurrentParentFrame, const TNotNull<const FStateTreeExecutionFrame*> CurrentFrame)
+EStateTreeRunStatus FStateTreeExecutionContext::TickEvaluatorsAndGlobalTasksForFrame(const float DeltaTime, const bool bTickGlobalTasks, const int32 FrameIndex, const FStateTreeExecutionFrame* CurrentParentFrame, const TNotNull<FStateTreeExecutionFrame*> CurrentFrame)
 {
 	check(CurrentFrame->bIsGlobalFrame);
 
 	EStateTreeRunStatus Result = EStateTreeRunStatus::Running;
 
-	const bool bCopyBoundPropertiesOnNonTickedTask = UE::StateTree::ExecutionContext::Private::bCopyBoundPropertiesOnNonTickedTask;
-
 	const UStateTree* CurrentStateTree = CurrentFrame->StateTree;
 
-	FTickTaskArguments TickArgs;
-	TickArgs.DeltaTime = DeltaTime;
-	TickArgs.bIsGlobalTasks = true;
-	TickArgs.bShouldTickTasks = true;
-	TickArgs.ParentFrame = CurrentParentFrame;
-	TickArgs.Frame = CurrentFrame;
-
 	// Tick evaluators
 	for (int32 EvalIndex = CurrentStateTree->EvaluatorsBegin; EvalIndex < (CurrentStateTree->EvaluatorsBegin + CurrentStateTree->EvaluatorsNum); EvalIndex++)
 	{
@@ -3478,30 +3336,33 @@ EStateTreeRunStatus FStateTreeExecutionContext::TickEvaluatorsAndGlobalTasksForF
 
 	if (bTickGlobalTasks)
 	{
-		const FCompletedState* FrameCompletedState = CompletedStates.FindByPredicate([FrameID = CurrentFrame->FrameID](FCompletedState& Other)
-			{
-				return Other.FrameID == FrameID
-					&& Other.Reason == UE::StateTree::FFinishedTask::EReasonType::GlobalTask;
-			});
-		if (FrameCompletedState)
-		{
-			// Completed global task stops the frame execution.
-			Result = FrameCompletedState->RunStatus;
-		}
-		else
+		using namespace UE::StateTree;
+		using namespace UE::StateTree::ExecutionContext::Private;
+
+		FTasksCompletionStatus CurrentGlobalTasksStatus = CurrentFrame->ActiveTasksStatus.GetStatus(CurrentStateTree);
+		if (!CurrentGlobalTasksStatus.HasAnyFailed())
 		{
 			const bool bHasEvents = EventQueue && EventQueue->HasEvents();
 			if (bCopyBoundPropertiesOnNonTickedTask || CurrentStateTree->ShouldTickGlobalTasks(bHasEvents))
 			{
 				// Update Tasks data and tick if possible (ie. if no task has yet failed and bShouldTickTasks is true)
+				FTickTaskArguments TickArgs;
+				TickArgs.DeltaTime = DeltaTime;
 				TickArgs.TasksBegin = CurrentStateTree->GlobalTasksBegin;
-				TickArgs.TasksEnd = CurrentStateTree->GlobalTasksBegin + CurrentStateTree->GlobalTasksNum;
+				TickArgs.TasksNum = CurrentStateTree->GlobalTasksNum;
 				TickArgs.Indent = FrameIndex + 1;
-				const FTickTaskResult TickTasksResult = TickTasks(TickArgs);
-
-				Result = TickTasksResult.Status;
+				TickArgs.ParentFrame = CurrentParentFrame;
+				TickArgs.Frame = CurrentFrame;
+				TickArgs.TasksCompletionStatus = &CurrentGlobalTasksStatus;
+				TickArgs.bIsGlobalTasks = true;
+				TickArgs.bShouldTickTasks = true;
+				TickTasks(TickArgs);
 			}
 		}
+
+		// Completed global task stops the frame execution.
+		const ETaskCompletionStatus GlobalTaskStatus = CurrentGlobalTasksStatus.GetCompletionStatus();
+		Result = CastToRunStatus(GlobalTaskStatus);
 	}
 
 	return Result;
@@ -3529,6 +3390,7 @@ EStateTreeRunStatus FStateTreeExecutionContext::StartEvaluatorsAndGlobalTasks(FS
 			FCurrentlyProcessedFrameScope FrameScope(*this, CurrentParentFrame, CurrentFrame);
 			
 			const UStateTree* CurrentStateTree = CurrentFrame.StateTree;
+			UE::StateTree::FTasksCompletionStatus CurrentGlobalTasksStatus = CurrentFrame.ActiveTasksStatus.GetStatus(CurrentStateTree);
 
 			// Start evaluators
 			for (int32 EvalIndex = CurrentStateTree->EvaluatorsBegin; EvalIndex < (CurrentStateTree->EvaluatorsBegin + CurrentStateTree->EvaluatorsNum); EvalIndex++)
@@ -3556,11 +3418,12 @@ EStateTreeRunStatus FStateTreeExecutionContext::StartEvaluatorsAndGlobalTasks(FS
 			// Even if we call Enter/ExitState() on global tasks, they do not enter any specific state.
 			const FStateTreeTransitionResult Transition = {}; // Empty transition
 
-			for (int32 TaskIndex = CurrentStateTree->GlobalTasksBegin; TaskIndex < (CurrentStateTree->GlobalTasksBegin + CurrentStateTree->GlobalTasksNum); TaskIndex++)
+			for (int32 GlobalTaskIndex = 0; GlobalTaskIndex < CurrentStateTree->GlobalTasksNum; ++GlobalTaskIndex)
 			{
-				const FStateTreeTaskBase& Task =  CurrentStateTree->Nodes[TaskIndex].Get<const FStateTreeTaskBase>();
+				const int32 AssetTaskIndex = CurrentStateTree->GlobalTasksBegin + GlobalTaskIndex;
+				const FStateTreeTaskBase& Task =  CurrentStateTree->Nodes[AssetTaskIndex].Get<const FStateTreeTaskBase>();
 				const FStateTreeDataView TaskInstanceView = GetDataView(CurrentParentFrame, CurrentFrame, Task.InstanceDataHandle);
-				FNodeInstanceDataScope DataScope(*this, &Task, TaskIndex, Task.InstanceDataHandle, TaskInstanceView);
+				FNodeInstanceDataScope DataScope(*this, &Task, AssetTaskIndex, Task.InstanceDataHandle, TaskInstanceView);
 
 				// Copy bound properties.
 				if (Task.BindingsBatch.IsValid())
@@ -3576,44 +3439,25 @@ EStateTreeRunStatus FStateTreeExecutionContext::StartEvaluatorsAndGlobalTasks(FS
 				}
 
 				STATETREE_LOG(Verbose, TEXT("  Start: '%s'"), *Task.Name.ToString());
-				UE_STATETREE_DEBUG_TASK_ENTER_STATE(this, CurrentStateTree, FStateTreeIndex16(TaskIndex));
+				UE_STATETREE_DEBUG_TASK_ENTER_STATE(this, CurrentStateTree, FStateTreeIndex16(AssetTaskIndex));
 
-				const int32 PreviousFinishedTasksNum = Exec.FinishedTasks.Num();
-				EStateTreeRunStatus TaskStatus = EStateTreeRunStatus::Unset;
+				EStateTreeRunStatus TaskRunStatus = EStateTreeRunStatus::Unset;
 				{
 					QUICK_SCOPE_CYCLE_COUNTER(StateTree_Task_TreeStart);
 
-					TaskStatus = Task.EnterState(*this, Transition);
+					TaskRunStatus = Task.EnterState(*this, Transition);
 				}
 
-				// Process FinishTask (weak, or this context)
-				bool bFinishTaskInternal = true;
-				for (int32 FinishedTaskIndex = PreviousFinishedTasksNum; FinishedTaskIndex < Exec.FinishedTasks.Num(); ++FinishedTaskIndex)
-				{
-					// If same task, then update TaskStatus.
-					UE::StateTree::FFinishedTask& FinishedTask = Exec.FinishedTasks[FinishedTaskIndex];
-					if (TaskIndex == FinishedTask.TaskIndex.AsInt32()
-						&& CurrentFrameID == FinishedTask.FrameID)
-					{
-						ensureMsgf(!FinishedTask.StateID.IsValid(), TEXT("EnterState a global task. The state should be invalid."));
-						// Delayed from FinishTask(), use the best priority from what the task returns and mark the state completed.
-						TaskStatus = UE::StateTree::ExecutionContext::Private::GetPriorityRunStatus(TaskStatus, FinishedTask.RunStatus);
-						bFinishTaskInternal = false;
-						MarkStateCompleted(FinishedTask);
-					}
-				}
+				UE::StateTree::ETaskCompletionStatus TaskStatus = UE::StateTree::ExecutionContext::Private::CastToTaskStatus(TaskRunStatus);
+				TaskStatus = CurrentGlobalTasksStatus.SetStatusWithPriority(GlobalTaskIndex, TaskStatus);
 
-				UE_STATETREE_DEBUG_TASK_EVENT(this, TaskIndex, TaskInstanceView, EStateTreeTraceEventType::OnEntered, TaskStatus);
+				TaskRunStatus = UE::StateTree::ExecutionContext::Private::CastToRunStatus(TaskStatus);
+				UE_STATETREE_DEBUG_TASK_EVENT(this, AssetTaskIndex, TaskInstanceView, EStateTreeTraceEventType::OnEntered, TaskRunStatus);
 
-				if (TaskStatus != EStateTreeRunStatus::Running && bFinishTaskInternal)
-				{
-					FinishTaskInternal(UE::StateTree::FFinishedTask::EReasonType::GlobalTask, CurrentFrame.FrameID, UE::StateTree::FActiveStateID::Invalid, TaskIndex, TaskStatus);					
-				}
-
-				Result = UE::StateTree::ExecutionContext::Private::GetPriorityRunStatus(Result, TaskStatus);
+				Result = UE::StateTree::ExecutionContext::Private::GetPriorityRunStatus(Result, TaskRunStatus);
 				if (Result != EStateTreeRunStatus::Running)
 				{
-					OutLastInitializedTaskIndex = FStateTreeIndex16(TaskIndex);
+					OutLastInitializedTaskIndex = FStateTreeIndex16(AssetTaskIndex);
 				}
 			}
 		}
@@ -3744,7 +3588,7 @@ void FStateTreeExecutionContext::CallStopOnEvaluatorsAndGlobalTasks(const FState
 	}
 }
 
-EStateTreeRunStatus FStateTreeExecutionContext::StartTemporaryEvaluatorsAndGlobalTasks(const FStateTreeExecutionFrame* CurrentParentFrame, const FStateTreeExecutionFrame& CurrentFrame)
+EStateTreeRunStatus FStateTreeExecutionContext::StartTemporaryEvaluatorsAndGlobalTasks(const FStateTreeExecutionFrame* CurrentParentFrame, FStateTreeExecutionFrame& CurrentFrame)
 {
 	if (!CurrentFrame.bIsGlobalFrame)
 	{
@@ -3762,6 +3606,7 @@ EStateTreeRunStatus FStateTreeExecutionContext::StartTemporaryEvaluatorsAndGloba
 	EStateTreeRunStatus Result = EStateTreeRunStatus::Running;
 	const UE::StateTree::FActiveFrameID CurrentFrameID = CurrentFrame.FrameID;
 	const UStateTree* CurrentStateTree = CurrentFrame.StateTree;
+	UE::StateTree::FTasksCompletionStatus CurrentTasksStatus = CurrentFrame.ActiveTasksStatus.GetStatus(CurrentStateTree);
 
 	// Start evaluators
 	for (int32 EvalIndex = CurrentStateTree->EvaluatorsBegin; EvalIndex < (CurrentStateTree->EvaluatorsBegin + CurrentStateTree->EvaluatorsNum); EvalIndex++)
@@ -3795,9 +3640,10 @@ EStateTreeRunStatus FStateTreeExecutionContext::StartTemporaryEvaluatorsAndGloba
 	// Even if we call Enter/ExitState() on global tasks, they do not enter any specific state.
 	const FStateTreeTransitionResult Transition = {}; // Empty transition
 
-	for (int32 TaskIndex = CurrentStateTree->GlobalTasksBegin; TaskIndex < (CurrentStateTree->GlobalTasksBegin + CurrentStateTree->GlobalTasksNum); TaskIndex++)
+	for (int32 GlobalTaskIndex = 0; GlobalTaskIndex < CurrentStateTree->GlobalTasksNum; ++GlobalTaskIndex)
 	{
-		const FStateTreeTaskBase& Task =  CurrentStateTree->Nodes[TaskIndex].Get<const FStateTreeTaskBase>();
+		const int32 AssetTaskIndex = CurrentStateTree->GlobalTasksBegin + GlobalTaskIndex;
+		const FStateTreeTaskBase& Task =  CurrentStateTree->Nodes[AssetTaskIndex].Get<const FStateTreeTaskBase>();
 		// Ignore disabled task
 		if (Task.bTaskEnabled == false)
 		{
@@ -3808,11 +3654,11 @@ EStateTreeRunStatus FStateTreeExecutionContext::StartTemporaryEvaluatorsAndGloba
 		FStateTreeDataView TaskDataView = GetDataViewOrTemporary(CurrentParentFrame, CurrentFrame, Task.InstanceDataHandle);
 		if (!TaskDataView.IsValid())
 		{
-			TaskDataView = AddTemporaryInstance(CurrentFrame, FStateTreeIndex16(TaskIndex), Task.InstanceDataHandle, CurrentFrame.StateTree->DefaultInstanceData.GetStruct(Task.InstanceTemplateIndex.Get()));
+			TaskDataView = AddTemporaryInstance(CurrentFrame, FStateTreeIndex16(AssetTaskIndex), Task.InstanceDataHandle, CurrentFrame.StateTree->DefaultInstanceData.GetStruct(Task.InstanceTemplateIndex.Get()));
 			check(TaskDataView.IsValid())
 		}
 
-		FNodeInstanceDataScope DataScope(*this, &Task, TaskIndex, Task.InstanceDataHandle, TaskDataView);
+		FNodeInstanceDataScope DataScope(*this, &Task, AssetTaskIndex, Task.InstanceDataHandle, TaskDataView);
 
 		// Copy bound properties with the temporary buffer.
 		if (Task.BindingsBatch.IsValid())
@@ -3821,41 +3667,21 @@ EStateTreeRunStatus FStateTreeExecutionContext::StartTemporaryEvaluatorsAndGloba
 		}
 
 		STATETREE_LOG(Verbose, TEXT("  Start: '%s'"), *Task.Name.ToString());
-		UE_STATETREE_DEBUG_TASK_ENTER_STATE(this, CurrentStateTree, FStateTreeIndex16(TaskIndex));
+		UE_STATETREE_DEBUG_TASK_ENTER_STATE(this, CurrentStateTree, FStateTreeIndex16(AssetTaskIndex));
 
-		const int32 PreviousFinishedTasksNum = Exec.FinishedTasks.Num();
-		EStateTreeRunStatus TaskStatus = EStateTreeRunStatus::Unset;
+		EStateTreeRunStatus TaskRunStatus = EStateTreeRunStatus::Unset;
 		{
 			QUICK_SCOPE_CYCLE_COUNTER(StateTree_Task_TreeStart);
-			TaskStatus = Task.EnterState(*this, Transition);
+			TaskRunStatus = Task.EnterState(*this, Transition);
 		}
 
-		// Process FinishTask (weak, or this context)
-		bool bFinishTaskInternal = true;
-		for (int32 FinishedTaskIndex = PreviousFinishedTasksNum; FinishedTaskIndex < Exec.FinishedTasks.Num(); ++FinishedTaskIndex)
-		{
-			// If same task, then update TaskStatus.
-			UE::StateTree::FFinishedTask& FinishedTask = Exec.FinishedTasks[FinishedTaskIndex];
-			if (TaskIndex == FinishedTask.TaskIndex.AsInt32()
-				&& CurrentFrameID == FinishedTask.FrameID)
-			{
-				ensureMsgf(!FinishedTask.StateID.IsValid(), TEXT("EnterState a global task. The state should be invalid."));
-				// Delayed from FinishTask(), use the best priority from what the task returns and mark the state completed.
-				TaskStatus = UE::StateTree::ExecutionContext::Private::GetPriorityRunStatus(TaskStatus, FinishedTask.RunStatus);
-				bFinishTaskInternal = false;
-				MarkStateCompleted(FinishedTask);
-			}
-		}
+		UE::StateTree::ETaskCompletionStatus TaskStatus = UE::StateTree::ExecutionContext::Private::CastToTaskStatus(TaskRunStatus);
+		TaskStatus = CurrentTasksStatus.SetStatusWithPriority(GlobalTaskIndex, TaskStatus);
 
-		UE_STATETREE_DEBUG_TASK_EVENT(this, TaskIndex, TaskDataView, EStateTreeTraceEventType::OnEntered, TaskStatus);
+		TaskRunStatus = UE::StateTree::ExecutionContext::Private::CastToRunStatus(TaskStatus);
+		UE_STATETREE_DEBUG_TASK_EVENT(this, AssetTaskIndex, TaskDataView, EStateTreeTraceEventType::OnEntered, TaskRunStatus);
 
-		if (TaskStatus != EStateTreeRunStatus::Running && bFinishTaskInternal)
-		{
-			FinishTaskInternal(UE::StateTree::FFinishedTask::EReasonType::GlobalTask, CurrentFrame.FrameID, UE::StateTree::FActiveStateID::Invalid, TaskIndex, TaskStatus);			
-		}
-
-		Result = UE::StateTree::ExecutionContext::Private::GetPriorityRunStatus(Result, TaskStatus);
-		if (Result != EStateTreeRunStatus::Running)
+		if (TaskRunStatus != EStateTreeRunStatus::Running)
 		{
 			break;
 		}
@@ -3943,29 +3769,29 @@ EStateTreeRunStatus FStateTreeExecutionContext::TickTasks(const float DeltaTime)
 	using namespace UE::StateTree;
 
 	FStateTreeExecutionState& Exec = GetExecState();
+	Exec.bHasPendingCompletedState = false;
 
 	if (Exec.ActiveFrames.IsEmpty())
 	{
 		return EStateTreeRunStatus::Failed;
 	}
 
-	EStateTreeRunStatus Result = EStateTreeRunStatus::Running;
 	int32 NumTotalEnabledTasks = 0;
-
 	const bool bCopyBoundPropertiesOnNonTickedTask = ExecutionContext::Private::bCopyBoundPropertiesOnNonTickedTask;
 
 	FTickTaskArguments TickArgs;
 	TickArgs.DeltaTime = DeltaTime;
 	TickArgs.bIsGlobalTasks = false;
-	TickArgs.bShouldTickTasks = CompletedStates.Num() == 0;
+	TickArgs.bShouldTickTasks = true;
 
 	STATETREE_CLOG(Exec.ActiveFrames.Num() > 0, VeryVerbose, TEXT("Ticking Tasks"));
 
-	for (int32 FrameIndex = 0; FrameIndex < Exec.ActiveFrames.Num(); FrameIndex++)
+	for (int32 FrameIndex = 0; FrameIndex < Exec.ActiveFrames.Num(); ++FrameIndex)
 	{
 		TickArgs.ParentFrame = FrameIndex > 0 ? &Exec.ActiveFrames[FrameIndex - 1] : nullptr;
 		TickArgs.Frame = &Exec.ActiveFrames[FrameIndex];
 		const UStateTree* CurrentStateTree = TickArgs.Frame->StateTree;
+
 		FCurrentlyProcessedFrameScope FrameScope(*this, TickArgs.ParentFrame, *TickArgs.Frame);
 
 		if (UE::StateTree::ExecutionContext::Private::bTickGlobalNodesFollowingTreeHierarchy)
@@ -3977,100 +3803,120 @@ EStateTreeRunStatus FStateTreeExecutionContext::TickTasks(const float DeltaTime)
 				if (FrameResult != EStateTreeRunStatus::Running)
 				{
 					Exec.RequestedStop = ExecutionContext::Private::GetPriorityRunStatus(Exec.RequestedStop, FrameResult);
+					TickArgs.bShouldTickTasks = false;
 					break;
 				}
 			}
 		}
 
-		for (int32 StateIndex = 0; StateIndex < TickArgs.Frame->ActiveStates.Num(); StateIndex++)
+		for (int32 StateIndex = 0; StateIndex < TickArgs.Frame->ActiveStates.Num(); ++StateIndex)
 		{
 			const FStateTreeStateHandle CurrentHandle = TickArgs.Frame->ActiveStates[StateIndex];
+			const FCompactStateTreeState& CurrentState = CurrentStateTree->States[CurrentHandle.Index];
+			UE::StateTree::FTasksCompletionStatus CurrentCompletionStatus = TickArgs.Frame->ActiveTasksStatus.GetStatus(CurrentState);
+
 			TickArgs.StateID = TickArgs.Frame->ActiveStates.StateIDs[StateIndex];
-			const FCompactStateTreeState& State = CurrentStateTree->States[CurrentHandle.Index];
+			TickArgs.TasksCompletionStatus = &CurrentCompletionStatus;
 
 			FCurrentlyProcessedStateScope StateScope(*this, CurrentHandle);
 			UE_STATETREE_DEBUG_SCOPED_STATE(this, CurrentHandle);
 
-			STATETREE_CLOG(State.TasksNum > 0, VeryVerbose, TEXT("%*sState '%s'")
-				, (FrameIndex + StateIndex + 1)*UE::StateTree::Debug::IndentSize, TEXT("")
+			STATETREE_CLOG(CurrentState.TasksNum > 0, VeryVerbose, TEXT("%*sState '%s'")
+				, (FrameIndex + StateIndex + 1) * UE::StateTree::Debug::IndentSize, TEXT("")
 				, *DebugGetStatePath(Exec.ActiveFrames, TickArgs.Frame, StateIndex));
 
-			if (State.Type == EStateTreeStateType::Linked || State.Type == EStateTreeStateType::LinkedAsset)
+			if (CurrentState.Type == EStateTreeStateType::Linked || CurrentState.Type == EStateTreeStateType::LinkedAsset)
 			{
-				if (State.ParameterDataHandle.IsValid() && State.ParameterBindingsBatch.IsValid())
+				if (CurrentState.ParameterDataHandle.IsValid() && CurrentState.ParameterBindingsBatch.IsValid())
 				{
-					const FStateTreeDataView StateParamsDataView = GetDataView(TickArgs.ParentFrame, *TickArgs.Frame, State.ParameterDataHandle);
-					CopyBatchOnActiveInstances(TickArgs.ParentFrame, *TickArgs.Frame, StateParamsDataView, State.ParameterBindingsBatch);
+					const FStateTreeDataView StateParamsDataView = GetDataView(TickArgs.ParentFrame, *TickArgs.Frame, CurrentState.ParameterDataHandle);
+					CopyBatchOnActiveInstances(TickArgs.ParentFrame, *TickArgs.Frame, StateParamsDataView, CurrentState.ParameterBindingsBatch);
 				}
 			}
 
 			const bool bHasEvents = EventQueue && EventQueue->HasEvents();
-			if (bCopyBoundPropertiesOnNonTickedTask || State.ShouldTickTasks(bHasEvents))
+			bool bRequestLoopStop = false;
+			if (bCopyBoundPropertiesOnNonTickedTask || CurrentState.ShouldTickTasks(bHasEvents))
 			{
 				// Update Tasks data and tick if possible (ie. if no task has yet failed and bShouldTickTasks is true)
-				TickArgs.TasksBegin = State.TasksBegin;
-				TickArgs.TasksEnd = State.TasksBegin + State.TasksNum;
+				TickArgs.TasksBegin = CurrentState.TasksBegin;
+				TickArgs.TasksNum = CurrentState.TasksNum;
 				TickArgs.Indent = (FrameIndex + StateIndex + 1);
 				const FTickTaskResult TickTasksResult = TickTasks(TickArgs);
 
-				TickArgs.bShouldTickTasks = TickTasksResult.bShouldTickTasks;
-				Result = UE::StateTree::ExecutionContext::Private::GetPriorityRunStatus(Result, TickTasksResult.Status);
-
-				if (!bCopyBoundPropertiesOnNonTickedTask && !TickTasksResult.bShouldTickTasks)
-				{
-					// Stop ticking
-					break;
-				}
+				// Keep updating the binding but do not call tick on tasks if there's a failure.
+				TickArgs.bShouldTickTasks = TickTasksResult.bShouldTickTasks
+					&& !CurrentCompletionStatus.HasAnyFailed();
+				// If a failure and we do not copy then bindings, then we can stop.
+				bRequestLoopStop = !bCopyBoundPropertiesOnNonTickedTask && !TickTasksResult.bShouldTickTasks;
 			}
 
-			if (Exec.RequestedStop != EStateTreeRunStatus::Unset)
+			NumTotalEnabledTasks += CurrentState.EnabledTasksNum;
+
+			if (bRequestLoopStop)
 			{
-				// Stop the state tasks execution
-				return Result;
+				break;
 			}
+		}
+	}
+
+	// Collect the result after every tasks has the chance to tick.
+	//An async or delegate might complete a global or "previous" task (in a different order).
+	EStateTreeRunStatus FrameResult = EStateTreeRunStatus::Running;
+	EStateTreeRunStatus StateResult = EStateTreeRunStatus::Running;
+	for (int32 FrameIndex = 0; FrameIndex < Exec.ActiveFrames.Num(); ++FrameIndex)
+	{
+		using namespace UE::StateTree::ExecutionContext::Private;
 
-			NumTotalEnabledTasks += State.EnabledTasksNum;
+		const FStateTreeExecutionFrame& CurrentFrame = Exec.ActiveFrames[FrameIndex];
+		const UStateTree* CurrentStateTree = CurrentFrame.StateTree;
+		if (CurrentFrame.bIsGlobalFrame)
+		{
+			const ETaskCompletionStatus GlobalTasksStatus = CurrentFrame.ActiveTasksStatus.GetStatus(CurrentStateTree).GetCompletionStatus();
+			FrameResult = GetPriorityRunStatus(FrameResult, CastToRunStatus(GlobalTasksStatus));
+		}
+		for (int32 StateIndex = 0; StateIndex < CurrentFrame.ActiveStates.Num() && StateResult != EStateTreeRunStatus::Failed; ++StateIndex)
+		{
+			const FStateTreeStateHandle CurrentHandle = CurrentFrame.ActiveStates[StateIndex];
+			const FCompactStateTreeState& State = CurrentStateTree->States[CurrentHandle.Index];
+			const ETaskCompletionStatus StateTasksStatus = CurrentFrame.ActiveTasksStatus.GetStatus(State).GetCompletionStatus();
+			StateResult = GetPriorityRunStatus(StateResult, CastToRunStatus(StateTasksStatus));
 		}
 	}
 
-	// Delayed finished task but no task were ticked to return the value.
-	if (Result == EStateTreeRunStatus::Running && CompletedStates.Num() > 0)
+	//@todo frame should only close their own
+	if (FrameResult != EStateTreeRunStatus::Running)
+	{
+		Exec.RequestedStop = ExecutionContext::Private::GetPriorityRunStatus(Exec.RequestedStop, FrameResult);
+	}
+	else if (NumTotalEnabledTasks == 0 && StateResult == EStateTreeRunStatus::Running && FrameResult == EStateTreeRunStatus::Running)
 	{
-		// Return the "worst" run status from the completed states.
-		//Since no states/tasks ticked, the result didn't have the chance to get updated.
-		for (const FCompletedState& CompletedState : CompletedStates)
+		// No enabled tasks, done ticking.
+		//Complete the the bottom state in the bottom frame (to trigger the completion transitions).
+		if (ensureMsgf(Exec.ActiveFrames.Num() > 0, TEXT("No task is allowed to clear/stop/transition. Those actions should be delayed inside the execution context.")))
 		{
-			if (CompletedState.Reason == UE::StateTree::FFinishedTask::EReasonType::GlobalTask)
+			FStateTreeExecutionFrame& LastFrame = Exec.ActiveFrames.Last();
+			const int32 NumberOfActiveState = LastFrame.ActiveStates.Num();
+			if (NumberOfActiveState == 0)
 			{
-				// The value was set here in TickTasks. RequestedStop should already have the correct value.
-				//Or the value was set in GlobalTickTasks.
-				//Don't take any chances and set it back.
-				Exec.RequestedStop = ExecutionContext::Private::GetPriorityRunStatus(Exec.RequestedStop, CompletedState.RunStatus);
+				LastFrame.ActiveTasksStatus.GetStatus(LastFrame.StateTree).SetCompletionStatus(ETaskCompletionStatus::Succeeded);
 			}
 			else
 			{
-				Result = UE::StateTree::ExecutionContext::Private::GetPriorityRunStatus(Result, CompletedState.RunStatus);
+				const FStateTreeStateHandle CurrentHandle = LastFrame.ActiveStates[NumberOfActiveState - 1];
+				const FCompactStateTreeState& State = LastFrame.StateTree->States[CurrentHandle.Index];
+				LastFrame.ActiveTasksStatus.GetStatus(State).SetCompletionStatus(ETaskCompletionStatus::Succeeded);
 			}
 		}
-	}
-	else if (NumTotalEnabledTasks == 0 && CompletedStates.Num() == 0)
-	{
-		// No enabled tasks, done ticking. Check from the bottom state in the bottom frame to trigger completion transitions
-		if (Exec.ActiveFrames.Num() > 0)
-		{
-			const FStateTreeExecutionFrame& LastFrame = Exec.ActiveFrames.Last();
-			const UE::StateTree::FActiveFrameID CompletedFrameID = LastFrame.FrameID;
-			const UE::StateTree::FActiveStateID CompletedStateID = LastFrame.ActiveStates.StateIDs[LastFrame.ActiveStates.Num() - 1];
-			FinishTaskInternal(UE::StateTree::FFinishedTask::EReasonType::InternalTransition, CompletedFrameID, CompletedStateID, INDEX_NONE, EStateTreeRunStatus::Succeeded);
-		}
 		else
 		{
-			Exec.RequestedStop = EStateTreeRunStatus::Succeeded;
+			Exec.RequestedStop = EStateTreeRunStatus::Stopped;
 		}
-		Result = EStateTreeRunStatus::Succeeded;
+		StateResult = EStateTreeRunStatus::Succeeded;
 	}
 
-	return Result;
+	Exec.bHasPendingCompletedState = StateResult != EStateTreeRunStatus::Running;
+	return StateResult;
 }
 
 FStateTreeExecutionContext::FTickTaskResult FStateTreeExecutionContext::TickTasks(const FTickTaskArguments& Args)
@@ -4078,8 +3924,8 @@ FStateTreeExecutionContext::FTickTaskResult FStateTreeExecutionContext::TickTask
 	using namespace UE::StateTree;
 
 	check(Args.Frame);
+	check(Args.TasksCompletionStatus);
 
-	EStateTreeRunStatus Result = EStateTreeRunStatus::Unset;
 	bool bShouldTickTasks = Args.bShouldTickTasks;
 
 	FStateTreeExecutionState& Exec = GetExecState();
@@ -4088,9 +3934,10 @@ FStateTreeExecutionContext::FTickTaskResult FStateTreeExecutionContext::TickTask
 	const FActiveFrameID CurrentFrameID = Args.Frame->FrameID;
 	check(CurrentStateTree);
 
-	for (int32 TaskIndex = Args.TasksBegin; TaskIndex < Args.TasksEnd; ++TaskIndex)
+	for (int32 OwnerTaskIndex = 0; OwnerTaskIndex < Args.TasksNum; ++OwnerTaskIndex)
 	{
-		const FStateTreeTaskBase& Task = CurrentStateTree->Nodes[TaskIndex].Get<const FStateTreeTaskBase>();
+		const int32 AssetTaskIndex = Args.TasksBegin + OwnerTaskIndex;
+		const FStateTreeTaskBase& Task = CurrentStateTree->Nodes[AssetTaskIndex].Get<const FStateTreeTaskBase>();
 
 		// Ignore disabled task
 		if (Task.bTaskEnabled == false)
@@ -4100,10 +3947,12 @@ FStateTreeExecutionContext::FTickTaskResult FStateTreeExecutionContext::TickTask
 		}
 
 		const FStateTreeDataView TaskInstanceView = GetDataView(Args.ParentFrame, *Args.Frame, Task.InstanceDataHandle);
-		FNodeInstanceDataScope DataScope(*this, &Task, TaskIndex, Task.InstanceDataHandle, TaskInstanceView);
+		FNodeInstanceDataScope DataScope(*this, &Task, AssetTaskIndex, Task.InstanceDataHandle, TaskInstanceView);
 
 		const bool bHasEvents = EventQueue && EventQueue->HasEvents();
+		const bool bIsTaskRunning = Args.TasksCompletionStatus->IsRunning(OwnerTaskIndex);
 		const bool bNeedsTick = bShouldTickTasks
+			&& bIsTaskRunning
 			&& (Task.bShouldCallTick || (bHasEvents && Task.bShouldCallTickOnlyOnEvents));
 		STATETREE_LOG(VeryVerbose, TEXT("%*s  Tick: '%s' %s"), Args.Indent * Debug::IndentSize, TEXT("")
 			, *Task.Name.ToString()
@@ -4121,93 +3970,44 @@ FStateTreeExecutionContext::FTickTaskResult FStateTreeExecutionContext::TickTask
 
 		if (!bNeedsTick)
 		{
+			// Task didn't tick because it failed.
+			//The following tasks should not tick but we might still need to update their bindings.
+			if (!bIsTaskRunning && bShouldTickTasks && Args.TasksCompletionStatus->HasAnyFailed())
+			{
+				bShouldTickTasks = false;
+			}
 			continue;
 		}
 
-		//UE_STATETREE_DEBUG_TASK_EVENT(this, TaskIndex, TaskDataView, EStateTreeTraceEventType::OnTickingTask, EStateTreeRunStatus::Running);
-		UE_STATETREE_DEBUG_TASK_TICK(this, CurrentStateTree, TaskIndex);
+		//UE_STATETREE_DEBUG_TASK_EVENT(this, AssetTaskIndex, TaskDataView, EStateTreeTraceEventType::OnTickingTask, EStateTreeRunStatus::Running);
+		UE_STATETREE_DEBUG_TASK_TICK(this, CurrentStateTree, AssetTaskIndex);
 
-		const int32 PreviousFinishedTasksNum = Exec.FinishedTasks.Num();
-		EStateTreeRunStatus TaskResult = EStateTreeRunStatus::Unset;
+		EStateTreeRunStatus TaskRunStatus = EStateTreeRunStatus::Unset;
 		{
 			QUICK_SCOPE_CYCLE_COUNTER(StateTree_Task_Tick);
 			CSV_SCOPED_TIMING_STAT_EXCLUSIVE(StateTree_Task_Tick);
 
-			TaskResult = Task.Tick(*this, Args.DeltaTime);
-		}
-
-		// Process FinishTask (weak context, or this context)
-		bool bFinishTaskInternal = true;
-		for (int32 FinishedTaskIndex = PreviousFinishedTasksNum; FinishedTaskIndex < Exec.FinishedTasks.Num(); ++FinishedTaskIndex)
-		{
-			FFinishedTask& FinishedTask = Exec.FinishedTasks[FinishedTaskIndex];
-			if (!IsFinishedTaskValid(FinishedTask))
-			{
-				continue;
-			}
-
-			// If same task, then update TaskStatus.
-			// If state, then continue unless it's a failure.
-			// If global task, then mark the completed execution.
-			if (TaskIndex == FinishedTask.TaskIndex.AsInt32()
-				&& (Args.StateID == FinishedTask.StateID || Args.bIsGlobalTasks)
-				&& CurrentFrameID == FinishedTask.FrameID)
-			{
-				if (Args.bIsGlobalTasks)
-				{
-					ensureMsgf(!FinishedTask.StateID.IsValid(), TEXT("Ticking a global task. The state should be invalid."));
-				}
-				// Delayed from FinishTask(), use the best priority from what the task returns and mark the state completed.
-				TaskResult = ExecutionContext::Private::GetPriorityRunStatus(TaskResult, FinishedTask.RunStatus);
-				bFinishTaskInternal = false;
-				MarkStateCompleted(FinishedTask);
-			}
-			else
-			{
-				MarkStateCompleted(FinishedTask);
-				if (Args.bIsGlobalTasks)
-				{
-					if (FinishedTask.Reason == FFinishedTask::EReasonType::GlobalTask)
-					{
-						// Mark it completed and continue ticking the other tasks.
-						// If it's a failure,. then the loop will stop bellow. We do the trace first.
-						Result = ExecutionContext::Private::GetPriorityRunStatus(Result, FinishedTask.RunStatus);
-					}
-					// Else will be processed in TickTasks()
-				}
-				else
-				{
-					if (FinishedTask.Reason == FFinishedTask::EReasonType::GlobalTask)
-					{
-						Exec.RequestedStop = ExecutionContext::Private::GetPriorityRunStatus(Exec.RequestedStop, FinishedTask.RunStatus);
-					}
-					else
-					{
-						// If it's a failure, then the loops will stop bellow. We do the trace first.
-						Result = ExecutionContext::Private::GetPriorityRunStatus(Result, FinishedTask.RunStatus);
-					}
-				}
-			}
+			TaskRunStatus = Task.Tick(*this, Args.DeltaTime);
 		}
 
-		UE_STATETREE_DEBUG_TASK_EVENT(this, TaskIndex, TaskInstanceView,
-			TaskResult != EStateTreeRunStatus::Running ? EStateTreeTraceEventType::OnTaskCompleted : EStateTreeTraceEventType::OnTicked,
-			TaskResult);
+		// Set the new status and fetch back the status with priority.
+		//In case an async task completes the same task.
+		//Or in case FinishTask() inside the Task.Tick()
+		ETaskCompletionStatus TaskStatus = UE::StateTree::ExecutionContext::Private::CastToTaskStatus(TaskRunStatus);
+		TaskStatus = Args.TasksCompletionStatus->SetStatusWithPriority(OwnerTaskIndex, TaskStatus);
+		TaskRunStatus = ExecutionContext::Private::CastToRunStatus(TaskStatus);
 
-		if (TaskResult != EStateTreeRunStatus::Running && bFinishTaskInternal)
-		{
-			const FFinishedTask::EReasonType Reason = Args.bIsGlobalTasks ? FFinishedTask::EReasonType::GlobalTask : FFinishedTask::EReasonType::StateTask;
-			ensureMsgf(Args.bIsGlobalTasks != Args.StateID.IsValid(), TEXT("The state ID needs to be valid when the we are ticking state tasks."));
-			FinishTaskInternal(Reason, CurrentFrameID, Args.StateID, TaskIndex, TaskResult);
-		}
+		UE_STATETREE_DEBUG_TASK_EVENT(this, AssetTaskIndex, TaskInstanceView,
+			TaskRunStatus != EStateTreeRunStatus::Running ? EStateTreeTraceEventType::OnTaskCompleted : EStateTreeTraceEventType::OnTicked,
+			TaskRunStatus);
 
-		Result = ExecutionContext::Private::GetPriorityRunStatus(Result, TaskResult);
-		if (Result == EStateTreeRunStatus::Failed)
+		if (TaskRunStatus == EStateTreeRunStatus::Failed)
 		{
 			bShouldTickTasks = false;
 		}
 	}
-	return FTickTaskResult{Result, bShouldTickTasks};
+
+	return FTickTaskResult{bShouldTickTasks};
 }
 
 bool FStateTreeExecutionContext::TestAllConditions(const FStateTreeExecutionFrame* CurrentParentFrame, const FStateTreeExecutionFrame& CurrentFrame, const int32 ConditionsOffset, const int32 ConditionsNum)
@@ -4559,6 +4359,7 @@ void FStateTreeExecutionContext::SetupNextTransition(const FStateTreeExecutionFr
 	FStateTreeExecutionFrame& NewFrame = NextTransition.NextActiveFrames.AddDefaulted_GetRef();
 	NewFrame.StateTree = CurrentFrame.StateTree;
 	NewFrame.RootState = CurrentFrame.RootState;
+	NewFrame.ActiveTasksStatus = CurrentFrame.ActiveTasksStatus;
 
 	if (NextState == FStateTreeStateHandle::Invalid)
 	{
@@ -4572,6 +4373,14 @@ void FStateTreeExecutionContext::SetupNextTransition(const FStateTreeExecutionFr
 
 bool FStateTreeExecutionContext::TriggerTransitions()
 {
+	//1. Process transition requests. Keep the single request with the highest priority.
+	//2. Process tick/event/delegate transitions and tasks. TriggerTransitions, from bottom to top.
+	// If delayed,
+	//	If delayed completed, then process.
+	//	Else add them to the delayed transition list.
+	//3. If no transition, Process completion transitions, from top to bottom.
+	//4. If transition occurs, check if there are any frame (sub-tree) that completed.
+
 	CSV_SCOPED_TIMING_STAT_EXCLUSIVE(StateTree_TriggerTransition);
 	UE_STATETREE_DEBUG_SCOPED_PHASE(this, EStateTreeUpdatePhase::TriggerTransitions);
 
@@ -4601,6 +4410,10 @@ bool FStateTreeExecutionContext::TriggerTransitions()
 			}
 		}
 	}
+
+	//@todo should only clear once when the transition is successful.
+	//to prevent 2 async requests and the first requests fails for X reason.
+	//they will be identified by a Frame/StateID so it's fine if they stay in the array.
 	InstanceData.ResetTransitionRequests();
 
 	//
@@ -4971,148 +4784,152 @@ PRAGMA_ENABLE_DEPRECATION_WARNINGS
 	bool bProcessSubTreeCompletion = true;
 
 	if (NextTransition.Priority == EStateTreeTransitionPriority::None
-		&& Exec.LastTickStatus != EStateTreeRunStatus::Running)
-	{
-		// Sort CompletedStates, only keep the smaller frame and the active state of the smaller frame.
-		int32 CompletedFrameStartIndex = INT_MAX;
-		int32 CompletedStateStartIndex = INT_MAX;
-		EStateTreeRunStatus CompletedExpectedRunStatus = Exec.LastTickStatus;
-		for (const FCompletedState& CompletedState : CompletedStates)
-		{
-			// Global task are not triggered in the transition loop.
-			if (CompletedState.Reason != UE::StateTree::FFinishedTask::EReasonType::GlobalTask)
-			{
-				const int32 NewFrameIndex = Exec.IndexOfActiveFrame(CompletedState.FrameID);
-				check(NewFrameIndex != INDEX_NONE); // the array is supposed to be clear every loop
-				const int32 NewStateIndex = Exec.ActiveFrames[NewFrameIndex].ActiveStates.IndexOfReverse(CompletedState.StateID);
-				const int32 NewLoopStateIndex = NewStateIndex == INDEX_NONE ? UINT_MAX : NewStateIndex;
-				if (NewFrameIndex < CompletedFrameStartIndex
-					|| (NewFrameIndex == CompletedFrameStartIndex && NewLoopStateIndex < CompletedStateStartIndex))
+		&& (Exec.LastTickStatus != EStateTreeRunStatus::Running || Exec.bHasPendingCompletedState))
+	{
+		// Find the pending completed frame/state. Don't cache the result because this function is reentrant.
+		//Stop at the first completion.
+		int32 FrameIndexToStart = -1; 
+		int32 StateIndexToStart = -1;
+		EStateTreeRunStatus CurrentStatus = EStateTreeRunStatus::Unset;
+		for (int32 FrameIndex = 0; FrameIndex < Exec.ActiveFrames.Num(); ++FrameIndex)
+		{
+			const FStateTreeExecutionFrame& CurrentFrame = Exec.ActiveFrames[FrameIndex];
+			const UStateTree* CurrentStateTree = CurrentFrame.StateTree;
+			for (int32 StateIndex = 0; StateIndex < CurrentFrame.ActiveStates.Num(); ++StateIndex)
+			{
+				const FStateTreeStateHandle CurrentHandle = CurrentFrame.ActiveStates[StateIndex];
+				const FCompactStateTreeState& State = CurrentStateTree->States[CurrentHandle.Index];
+				const UE::StateTree::ETaskCompletionStatus StateTasksStatus = CurrentFrame.ActiveTasksStatus.GetStatus(State).GetCompletionStatus();
+				if (StateTasksStatus != UE::StateTree::ETaskCompletionStatus::Running)
 				{
-					CompletedFrameStartIndex = NewFrameIndex;
-					CompletedStateStartIndex = NewLoopStateIndex;
-					CompletedExpectedRunStatus = CompletedState.RunStatus;
+					CurrentStatus = UE::StateTree::ExecutionContext::Private::CastToRunStatus(StateTasksStatus);
+					FrameIndexToStart = FrameIndex;
+					StateIndexToStart = StateIndex;
+					break;
 				}
 			}
-		}
 
-		UE::StateTree::ExecutionContext::Private::RemoveTickProcessedFinishedTasks(Exec);
-		CompletedStates.Empty();
+			if (CurrentStatus != EStateTreeRunStatus::Unset)
+			{
+				break;
+			}
+		}
 
-		// Start from the last completed state if specified, else only check the last state from the last frame.
-		const int32 FrameStartIndex = FMath::Min(CompletedFrameStartIndex, Exec.ActiveFrames.Num() - 1);
-		for (int32 FrameIndex = FrameStartIndex; FrameIndex >= 0; FrameIndex--)
+		// Possible if the completed task was a global task
+		//@todo remove this once the frame global task only completes it's own frame.
+		if (CurrentStatus != EStateTreeRunStatus::Unset)
 		{
-			FStateTreeExecutionFrame* CurrentParentFrame = FrameIndex > 0 ? &Exec.ActiveFrames[FrameIndex - 1] : nullptr;
-			FStateTreeExecutionFrame& CurrentFrame = Exec.ActiveFrames[FrameIndex];
-			const UStateTree* CurrentStateTree = CurrentFrame.StateTree;
+			const bool bIsCurrentStatusSucceeded = CurrentStatus == EStateTreeRunStatus::Succeeded;
+			const bool bIsCurrentStatusFailed = CurrentStatus == EStateTreeRunStatus::Failed;
+			const bool bIsCurrentStatusStopped = CurrentStatus == EStateTreeRunStatus::Stopped;
+			checkf(bIsCurrentStatusSucceeded || bIsCurrentStatusFailed || bIsCurrentStatusStopped, TEXT("Running is not accepted in the CurrentStatus loop."));
 
-			FCurrentlyProcessedFrameScope FrameScope(*this, CurrentParentFrame, CurrentFrame);
+			const EStateTreeTransitionTrigger CompletionTrigger = bIsCurrentStatusSucceeded ? EStateTreeTransitionTrigger::OnStateSucceeded : EStateTreeTransitionTrigger::OnStateFailed;
 
-			int32 StateStartIndex = CurrentFrame.ActiveStates.Num() - 1; // This is ok, even if the ActiveStates is 0, -1 will skip the whole state loop below.
-			if (FrameIndex == FrameStartIndex && CompletedStateStartIndex < CurrentFrame.ActiveStates.Num())
+			// Start from the last completed state and move up to the first state.
+			for (int32 FrameIndex = FrameIndexToStart; FrameIndex >= 0; --FrameIndex)
 			{
-				StateStartIndex = CompletedStateStartIndex;
-			}
+				FStateTreeExecutionFrame* CurrentParentFrame = FrameIndex > 0 ? &Exec.ActiveFrames[FrameIndex - 1] : nullptr;
+				FStateTreeExecutionFrame& CurrentFrame = Exec.ActiveFrames[FrameIndex];
+				const UStateTree* CurrentStateTree = CurrentFrame.StateTree;
 
-			// Use worst case of LastTickStatus and FoundCompletedState.
-			const bool bLastTickWasSuccess = Exec.LastTickStatus == EStateTreeRunStatus::Succeeded;
-			const EStateTreeRunStatus ExpectedRunStatus = bLastTickWasSuccess && FrameIndex == FrameStartIndex ? CompletedExpectedRunStatus : Exec.LastTickStatus;
-			const EStateTreeTransitionTrigger CompletionTrigger = ExpectedRunStatus == EStateTreeRunStatus::Succeeded ? EStateTreeTransitionTrigger::OnStateSucceeded : EStateTreeTransitionTrigger::OnStateFailed;
+				FCurrentlyProcessedFrameScope FrameScope(*this, CurrentParentFrame, CurrentFrame);
 
-			// Check completion transitions
-			for (int32 StateIndex = StateStartIndex; StateIndex >= 0; StateIndex--)
-			{
-				const FStateTreeStateHandle StateHandle = CurrentFrame.ActiveStates[StateIndex];
-				const FCompactStateTreeState& State = CurrentStateTree->States[StateHandle.Index];
+				const int32 CurrentStateIndexToStart = FrameIndex == FrameIndexToStart ? StateIndexToStart : CurrentFrame.ActiveStates.Num() - 1;
 
-				if (State.ShouldTickCompletionTransitions(CompletionTrigger == EStateTreeTransitionTrigger::OnStateSucceeded))
+				// Check completion transitions
+				for (int32 StateIndex = CurrentStateIndexToStart; StateIndex >= 0; --StateIndex)
 				{
-					FCurrentlyProcessedStateScope StateScope(*this, StateHandle);
-					UE_STATETREE_DEBUG_SCOPED_STATE_PHASE(this, StateHandle, EStateTreeUpdatePhase::TriggerTransitions);
+					const FStateTreeStateHandle StateHandle = CurrentFrame.ActiveStates[StateIndex];
+					const FCompactStateTreeState& State = CurrentStateTree->States[StateHandle.Index];
 
-					for (uint8 TransitionCounter = 0; TransitionCounter < State.TransitionsNum; ++TransitionCounter)
+					if (State.ShouldTickCompletionTransitions(bIsCurrentStatusSucceeded, bIsCurrentStatusFailed))
 					{
-						// All transition conditions must pass
-						const int16 TransitionIndex = State.TransitionsBegin + TransitionCounter;
-						const FCompactStateTransition& Transition = CurrentStateTree->Transitions[TransitionIndex];
+						FCurrentlyProcessedStateScope StateScope(*this, StateHandle);
+						UE_STATETREE_DEBUG_SCOPED_STATE_PHASE(this, StateHandle, EStateTreeUpdatePhase::TriggerTransitions);
 
-						// Skip disabled transitions
-						if (Transition.bTransitionEnabled == false)
+						for (uint8 TransitionCounter = 0; TransitionCounter < State.TransitionsNum; ++TransitionCounter)
 						{
-							continue;
-						}
+							// All transition conditions must pass
+							const int16 TransitionIndex = State.TransitionsBegin + TransitionCounter;
+							const FCompactStateTransition& Transition = CurrentStateTree->Transitions[TransitionIndex];
 
-						if (EnumHasAnyFlags(Transition.Trigger, CompletionTrigger))
-						{
-							bool bPassed = false;
+							// Skip disabled transitions
+							if (!Transition.bTransitionEnabled)
 							{
-								UE_STATETREE_DEBUG_TRANSITION_EVENT(this, FStateTreeTransitionSource(CurrentFrame.StateTree, FStateTreeIndex16(TransitionIndex), Transition.State, Transition.Priority), EStateTreeTraceEventType::OnEvaluating);
-								UE_STATETREE_DEBUG_SCOPED_PHASE(this, EStateTreeUpdatePhase::TransitionConditions);
-								bPassed = TestAllConditions(CurrentParentFrame, CurrentFrame, Transition.ConditionsBegin, Transition.ConditionsNum);
+								continue;
 							}
 
-							if (bPassed)
+
+							const bool bTransitionAccepted = !bIsCurrentStatusStopped 
+								? EnumHasAnyFlags(Transition.Trigger, CompletionTrigger)
+								: Transition.Trigger == EStateTreeTransitionTrigger::OnStateCompleted;
+							if (bTransitionAccepted)
 							{
-								// No delay allowed on completion conditions.
-								// No priority on completion transitions, use the priority to signal that state is selected.
-								UE_STATETREE_DEBUG_TRANSITION_EVENT(this, FStateTreeTransitionSource(CurrentFrame.StateTree, FStateTreeIndex16(TransitionIndex), Transition.State, Transition.Priority), EStateTreeTraceEventType::OnRequesting);
-								if (RequestTransition(CurrentFrame, Transition.State, EStateTreeTransitionPriority::Normal, /*TransitionEvent*/nullptr, Transition.Fallback))
+								bool bPassed = false;
 								{
-									NextTransitionSource = FStateTreeTransitionSource(CurrentFrame.StateTree, FStateTreeIndex16(TransitionIndex), Transition.State, Transition.Priority);
-									break;
+									UE_STATETREE_DEBUG_TRANSITION_EVENT(this, FStateTreeTransitionSource(CurrentFrame.StateTree, FStateTreeIndex16(TransitionIndex), Transition.State, Transition.Priority), EStateTreeTraceEventType::OnEvaluating);
+									UE_STATETREE_DEBUG_SCOPED_PHASE(this, EStateTreeUpdatePhase::TransitionConditions);
+									bPassed = TestAllConditions(CurrentParentFrame, CurrentFrame, Transition.ConditionsBegin, Transition.ConditionsNum);
+								}
+
+								if (bPassed)
+								{
+									// No delay allowed on completion conditions.
+									// No priority on completion transitions, use the priority to signal that state is selected.
+									UE_STATETREE_DEBUG_TRANSITION_EVENT(this, FStateTreeTransitionSource(CurrentFrame.StateTree, FStateTreeIndex16(TransitionIndex), Transition.State, Transition.Priority), EStateTreeTraceEventType::OnRequesting);
+									if (RequestTransition(CurrentFrame, Transition.State, EStateTreeTransitionPriority::Normal, /*TransitionEvent*/nullptr, Transition.Fallback))
+									{
+										NextTransitionSource = FStateTreeTransitionSource(CurrentFrame.StateTree, FStateTreeIndex16(TransitionIndex), Transition.State, Transition.Priority);
+										break;
+									}
 								}
 							}
 						}
-					}
 
-					if (NextTransition.Priority != EStateTreeTransitionPriority::None) //-V547
-					{
-						break;
+						if (NextTransition.Priority != EStateTreeTransitionPriority::None) //-V547
+						{
+							break;
+						}
 					}
 				}
-			}
 
-			// if a valid completion transition has already been found, the remaining transitions in parent frames won't have a higher priority than the found one
-			// so skip the remainder. this also prevented false positive warnings and ensures from STDebugger
-			if (NextTransition.Priority != EStateTreeTransitionPriority::None) //-V547
-			{
-				break;
+				// if a valid completion transition has already been found, the remaining transitions in parent frames won't have a higher priority than the found one
+				// so skip the remainder. this also prevented false positive warnings and ensures from STDebugger
+				if (NextTransition.Priority != EStateTreeTransitionPriority::None) //-V547
+				{
+					break;
+				}
 			}
-		}
 
-		// Handle the case where no transition was found.
-		if (NextTransition.Priority == EStateTreeTransitionPriority::None) //-V547
-		{
-			STATETREE_LOG(Verbose, TEXT("Could not trigger completion transition, jump back to root state."));
-			UE_STATETREE_DEBUG_LOG_EVENT(this, Warning, TEXT("Could not trigger completion transition, jump back to root state."));
+			// Handle the case where no transition was found.
+			if (NextTransition.Priority == EStateTreeTransitionPriority::None) //-V547
+			{
+				STATETREE_LOG(Verbose, TEXT("Could not trigger completion transition, jump back to root state."));
+				UE_STATETREE_DEBUG_LOG_EVENT(this, Warning, TEXT("Could not trigger completion transition, jump back to root state."));
 
-			check(!Exec.ActiveFrames.IsEmpty());
-			FStateTreeExecutionFrame& RootFrame = Exec.ActiveFrames[0];
-			FCurrentlyProcessedFrameScope RootFrameScope(*this, nullptr, RootFrame);
-			FCurrentlyProcessedStateScope RootStateScope(*this, FStateTreeStateHandle::Root);
+				check(!Exec.ActiveFrames.IsEmpty());
+				FStateTreeExecutionFrame& RootFrame = Exec.ActiveFrames[0];
+				FCurrentlyProcessedFrameScope RootFrameScope(*this, nullptr, RootFrame);
+				FCurrentlyProcessedStateScope RootStateScope(*this, FStateTreeStateHandle::Root);
 				
-			if (RequestTransition(RootFrame, FStateTreeStateHandle::Root, EStateTreeTransitionPriority::Normal))
-			{
-				NextTransitionSource = FStateTreeTransitionSource(RootFrame.StateTree, EStateTreeTransitionSourceType::Internal, FStateTreeStateHandle::Root, EStateTreeTransitionPriority::Normal);
-			}
-			else
-			{
-				STATETREE_LOG(Warning, TEXT("Failed to select root state. Stopping the tree with failure."));
-				UE_STATETREE_DEBUG_LOG_EVENT(this, Error, TEXT("Failed to select root state. Stopping the tree with failure."));
+				if (RequestTransition(RootFrame, FStateTreeStateHandle::Root, EStateTreeTransitionPriority::Normal))
+				{
+					NextTransitionSource = FStateTreeTransitionSource(RootFrame.StateTree, EStateTreeTransitionSourceType::Internal, FStateTreeStateHandle::Root, EStateTreeTransitionPriority::Normal);
+				}
+				else
+				{
+					STATETREE_LOG(Warning, TEXT("Failed to select root state. Stopping the tree with failure."));
+					UE_STATETREE_DEBUG_LOG_EVENT(this, Error, TEXT("Failed to select root state. Stopping the tree with failure."));
 
-				SetupNextTransition(RootFrame, FStateTreeStateHandle::Failed, EStateTreeTransitionPriority::Critical);
+					SetupNextTransition(RootFrame, FStateTreeStateHandle::Failed, EStateTreeTransitionPriority::Critical);
 
-				// In this case we don't want to complete subtrees, we want to force the whole tree to stop.
-				bProcessSubTreeCompletion = false;
+					// In this case we don't want to complete subtrees, we want to force the whole tree to stop.
+					bProcessSubTreeCompletion = false;
+				}
 			}
 		}
 	}
-	else
-	{
-		UE::StateTree::ExecutionContext::Private::RemoveTickProcessedFinishedTasks(Exec);
-		CompletedStates.Empty();
-	}
 
 	// Check if the transition was succeed/failed, if we're on a sub-tree, complete the subtree instead of transition.
 	if (NextTransition.TargetState.IsCompletionState() && bProcessSubTreeCompletion)
@@ -5126,7 +4943,7 @@ PRAGMA_ENABLE_DEPRECATION_WARNINGS
 		{
 			const FStateTreeExecutionFrame& SourceFrame = Exec.ActiveFrames[SourceFrameIndex];
 			const int32 ParentFrameIndex = SourceFrameIndex - 1;
-			const FStateTreeExecutionFrame& ParentFrame = Exec.ActiveFrames[ParentFrameIndex];
+			FStateTreeExecutionFrame& ParentFrame = Exec.ActiveFrames[ParentFrameIndex];
 			const FStateTreeStateHandle ParentLinkedState = ParentFrame.ActiveStates.Last();
 
 			if (ParentLinkedState.IsValid())
@@ -5136,9 +4953,10 @@ PRAGMA_ENABLE_DEPRECATION_WARNINGS
 					*GetSafeStateName(ParentFrame, ParentLinkedState), *GetSafeStateName(SourceFrame, NextTransition.SourceState), *UEnum::GetDisplayValueAsText(RunStatus).ToString());
 
 				// Set the parent linked state as last completed state, and update tick status to the status from the transition.
-				constexpr int32 TaskIndex = INDEX_NONE;
-				const UE::StateTree::FActiveStateID ParentLastActiveStateID = ParentFrame.ActiveStates.StateIDs[ParentFrame.ActiveStates.Num() - 1];
-				FinishTaskInternal(UE::StateTree::FFinishedTask::EReasonType::InternalTransition, ParentFrame.FrameID, ParentLastActiveStateID, TaskIndex, RunStatus);
+				const UE::StateTree::ETaskCompletionStatus TaskStatus = UE::StateTree::ExecutionContext::Private::CastToTaskStatus(RunStatus);
+				const FCompactStateTreeState& State = ParentFrame.StateTree->States[ParentLinkedState.Index];
+				ParentFrame.ActiveTasksStatus.GetStatus(State).SetCompletionStatus(TaskStatus);
+				Exec.bHasPendingCompletedState = true;
 				Exec.LastTickStatus = RunStatus;
 
 				// Clear the transition and return that no transition took place.
@@ -5261,6 +5079,7 @@ bool FStateTreeExecutionContext::SelectState(const FStateTreeExecutionFrame& Cur
 		}
 	}
 
+	ensureMsgf(LastFrame.ActiveStates.Num() >= FirstNewStateIndex, TEXT("ActiveTasksStatus won't be in sync with the amount of states."));
 	LastFrame.ActiveStates.SetNum(FirstNewStateIndex);
 
 	// Existing state's data is safe to access during select.
@@ -5468,6 +5287,7 @@ bool FStateTreeExecutionContext::SelectStateInternal(
 			__FUNCTION__, *GetSafeStateName(CurrentFrame, NextStateHandle), *GetStateStatusString(Exec), *GetNameSafe(&Owner), *GetFullNameSafe(CurrentFrame.StateTree));
 		return false;
 	}
+	CurrentFrame.ActiveTasksStatus.Push(CurrentStateTree, NextState);
 
 	// Check if we're still tracking on the current active frame and state.
 	// If we are, update the NumCurrentlyActiveStates to indicate that this state's instance data can be accessed. 
@@ -5531,6 +5351,10 @@ bool FStateTreeExecutionContext::SelectStateInternal(
 				NewFrame.FrameID = UE::StateTree::FActiveFrameID(Storage.GenerateUniqueId());
 				NewFrame.ExternalDataBaseIndex = CurrentFrame.ExternalDataBaseIndex;
 
+				const FCompactStateTreeFrame* FrameInfo = NewFrame.StateTree->GetFrameFromHandle(NewFrame.RootState);
+				ensureAlwaysMsgf(FrameInfo, TEXT("The compiled data is invalid. It should contains the information for the new root frame."));
+				NewFrame.ActiveTasksStatus = FrameInfo ? FStateTreeTasksCompletionStatus(*FrameInfo) : FStateTreeTasksCompletionStatus();
+
 				// Check and prevent recursion.
 				const bool bNewFrameAlreadySelected = OutSelectionResult.GetSelectedFrames().ContainsByPredicate([&NewFrame](const FStateTreeExecutionFrame& Frame) {
 					return Frame.IsSameFrame(NewFrame);
@@ -5611,6 +5435,10 @@ bool FStateTreeExecutionContext::SelectStateInternal(
 			NewFrame.RootState = FStateTreeStateHandle::Root;
 			NewFrame.bIsGlobalFrame = true;
 
+			const FCompactStateTreeFrame* FrameInfo = NewFrame.StateTree->GetFrameFromHandle(NewFrame.RootState);
+			ensureAlwaysMsgf(FrameInfo, TEXT("The compiled data is invalid. It should contains the information for the root frame."));
+			NewFrame.ActiveTasksStatus = FrameInfo ? FStateTreeTasksCompletionStatus(*FrameInfo) : FStateTreeTasksCompletionStatus();
+
 			// Check and prevent recursion.
 			const bool bNewFrameAlreadySelected = OutSelectionResult.GetSelectedFrames().ContainsByPredicate([&NewFrame](const FStateTreeExecutionFrame& Frame) {
 				return Frame.IsSameFrame(NewFrame);
@@ -6291,7 +6119,7 @@ bool FStateTreeWeakExecutionContext::UnbindDelegate(const FStateTreeDelegateList
 
 bool FStateTreeWeakExecutionContext::FinishTask(const FStateTreeWeakTaskRef& Task, EStateTreeFinishTaskType FinishType) const
 {
-	bool bCompleted = false;
+	bool bSucceed = false;
 
 	TStrongObjectPtr<UObject> OwnerPinned = Owner.Pin();
 	TStrongObjectPtr<const UStateTree> StateTreePinned = StateTree.Pin();
@@ -6306,22 +6134,56 @@ bool FStateTreeWeakExecutionContext::FinishTask(const FStateTreeWeakTaskRef& Tas
 		if (IsFrameOrStateActive(Exec))
 		{
 			FStateTreeStrongTaskRef TaskPin = Task.Pin();
-			const FStateTreeExecutionFrame* ExecutionFrame = Exec.FindActiveFrame(FrameID);
+			FStateTreeExecutionFrame* ExecutionFrame = Exec.FindActiveFrame(FrameID);
 			if (TaskPin.IsValid() && ExecutionFrame)
 			{
 				const bool bIsSameTree = ExecutionFrame->StateTree == TaskPin.GetStateTree();
 				if (ensureMsgf(bIsSameTree, TEXT("User error. The task is not from the same execution tree")))
 				{
+					using namespace UE::StateTree;
+
 					const int32 GlobalTaskBeginIndex = ExecutionFrame->StateTree->GlobalTasksBegin;
 					const int32 GlobalTaskEndIndex = ExecutionFrame->StateTree->GlobalTasksBegin + ExecutionFrame->StateTree->GlobalTasksNum;
-					const bool bIsGlobalTask = TaskPin.GetTaskIndex().AsInt32() >= GlobalTaskBeginIndex && TaskPin.GetTaskIndex().AsInt32() < GlobalTaskEndIndex;
-					const UE::StateTree::FFinishedTask::EReasonType Reason = bIsGlobalTask ? UE::StateTree::FFinishedTask::EReasonType::GlobalTask : UE::StateTree::FFinishedTask::EReasonType::StateTask;
-					const EStateTreeRunStatus RunStatus = UE::StateTree::ExecutionContext::Private::GetRunStatusFromFinishTaskType(FinishType);
-					constexpr bool bTickProcessed = false;
-					UE::StateTree::FFinishedTask CompletedTask(FrameID, StateID, TaskPin.GetTaskIndex(), RunStatus, Reason, bTickProcessed);
-					Exec.FinishedTasks.Add(CompletedTask);
-					ScheduleNextTick(OwnerPinned.Get(), StateTreePinned.Get(), *StoragePinned);
-					bCompleted = true;
+					const int32 AssetTaskIndex = TaskPin.GetTaskIndex().AsInt32();
+					const bool bIsGlobalTask = AssetTaskIndex >= GlobalTaskBeginIndex && AssetTaskIndex < GlobalTaskEndIndex;
+					const ETaskCompletionStatus TaskStatus = UE::StateTree::ExecutionContext::Private::CastToTaskStatus(FinishType);
+
+					bool bCompleted = false;
+					if (bIsGlobalTask)
+					{
+						const int32 FrameTaskIndex = AssetTaskIndex - GlobalTaskBeginIndex;
+						FTasksCompletionStatus GlobalTasksStatus = ExecutionFrame->ActiveTasksStatus.GetStatus(ExecutionFrame->StateTree);
+						GlobalTasksStatus.SetStatusWithPriority(FrameTaskIndex, TaskStatus);
+						bCompleted = GlobalTasksStatus.IsCompleted();
+						bSucceed = true;
+					}
+					else
+					{
+						// Find the task's state.
+						for (int32 StateIndex = 0; StateIndex < ExecutionFrame->ActiveStates.Num(); ++StateIndex)
+						{
+							const FStateTreeStateHandle CurrentHandle = ExecutionFrame->ActiveStates[StateIndex];
+							const FCompactStateTreeState& State = ExecutionFrame->StateTree->States[CurrentHandle.Index];
+							const int32 StateTaskBeginIndex = State.TasksBegin;
+							const int32 StateTaskEndIndex = State.TasksBegin + State.TasksNum;
+							const bool bIsStateTask = AssetTaskIndex >= StateTaskBeginIndex && AssetTaskIndex < StateTaskEndIndex;
+							if (bIsStateTask)
+							{
+								const int32 StateTaskIndex = AssetTaskIndex - State.TasksBegin;
+								FTasksCompletionStatus StateTasksStatus = ExecutionFrame->ActiveTasksStatus.GetStatus(State);
+								StateTasksStatus.SetStatusWithPriority(StateTaskIndex, TaskStatus);
+								bCompleted = StateTasksStatus.IsCompleted();
+								bSucceed = true;
+								break;
+							}
+						}
+					}
+
+					if (bCompleted)
+					{
+						Exec.bHasPendingCompletedState = true;
+						ScheduleNextTick(OwnerPinned.Get(), StateTreePinned.Get(), *StoragePinned);
+					}
 				}
 			}
 		}
@@ -6329,7 +6191,7 @@ bool FStateTreeWeakExecutionContext::FinishTask(const FStateTreeWeakTaskRef& Tas
 		StoragePinned->ReleaseWriteAccess();
 	}
 
-	return bCompleted;
+	return bSucceed;
 }
 
 
diff --git a/Engine/Plugins/Runtime/StateTree/Source/StateTreeModule/Private/StateTreeExecutionTypes.cpp b/Engine/Plugins/Runtime/StateTree/Source/StateTreeModule/Private/StateTreeExecutionTypes.cpp
--- a/Engine/Plugins/Runtime/StateTree/Source/StateTreeModule/Private/StateTreeExecutionTypes.cpp
+++ b/Engine/Plugins/Runtime/StateTree/Source/StateTreeModule/Private/StateTreeExecutionTypes.cpp
@@ -241,6 +241,7 @@ FRecordedStateTreeExecutionFrame::FRecordedStateTreeExecutionFrame(const FStateT
 	bIsGlobalFrame = ExecutionFrame.bIsGlobalFrame;
 }
 
+//Deprecated
 //----------------------------------------------------------------------//
 // FFinishedTask
 //----------------------------------------------------------------------//
diff --git a/Engine/Plugins/Runtime/StateTree/Source/StateTreeModule/Private/StateTreeTasksStatus.cpp b/Engine/Plugins/Runtime/StateTree/Source/StateTreeModule/Private/StateTreeTasksStatus.cpp
new file mode 100644
--- /dev/null
+++ b/Engine/Plugins/Runtime/StateTree/Source/StateTreeModule/Private/StateTreeTasksStatus.cpp
@@ -0,0 +1,232 @@
+// Copyright Epic Games, Inc. All Rights Reserved.
+
+#include "StateTreeTasksStatus.h"
+#include "StateTree.h"
+#include "StateTreeTypes.h"
+
+#include UE_INLINE_GENERATED_CPP_BY_NAME(StateTreeTasksStatus)
+
+namespace UE::StateTree::Private
+{
+	inline uint32* GetInlinedBufferPtr(FStateTreeTasksCompletionStatus::FMaskType*& Buffer)
+	{
+		return reinterpret_cast<uint32*>(&Buffer);
+	}
+}
+
+FStateTreeTasksCompletionStatus::FStateTreeTasksCompletionStatus(const FCompactStateTreeFrame& Frame)
+{
+	BufferNum = Frame.NumberOfTasksStatusMasks;
+	MallocBufferIfNeeded();
+}
+
+FStateTreeTasksCompletionStatus::~FStateTreeTasksCompletionStatus()
+{
+	if (!UseInlineBuffer())
+	{
+		FMemory::Free(Buffer);
+	}
+	Buffer = nullptr;
+	BufferNum = 0;
+}
+
+FStateTreeTasksCompletionStatus::FStateTreeTasksCompletionStatus(const FStateTreeTasksCompletionStatus& Other)
+	: BufferNum(Other.BufferNum)
+{
+	CopyBuffer(Other);
+}
+
+FStateTreeTasksCompletionStatus::FStateTreeTasksCompletionStatus(FStateTreeTasksCompletionStatus&& Other)
+	: Buffer(Other.Buffer)
+	, BufferNum(Other.BufferNum)
+{
+	Other.Buffer = nullptr;
+	Other.BufferNum = 0;
+}
+
+FStateTreeTasksCompletionStatus& FStateTreeTasksCompletionStatus::operator=(const FStateTreeTasksCompletionStatus& Other)
+{
+	check(this != &Other);
+	this->~FStateTreeTasksCompletionStatus();
+	BufferNum = Other.BufferNum;
+	CopyBuffer(Other);
+	return *this;
+}
+
+FStateTreeTasksCompletionStatus& FStateTreeTasksCompletionStatus::operator=(FStateTreeTasksCompletionStatus&& Other)
+{
+	check(this != &Other);
+	this->~FStateTreeTasksCompletionStatus();
+	BufferNum = Other.BufferNum;
+	Buffer = Other.Buffer;
+	Other.Buffer = nullptr;
+	Other.BufferNum = 0;
+	return *this;
+}
+
+void FStateTreeTasksCompletionStatus::MallocBufferIfNeeded()
+{
+	if (!UseInlineBuffer())
+	{
+		const int32 NumberOfBytes = sizeof(FStateTreeTasksCompletionStatus::FMaskType) * BufferNum * 2;
+		Buffer = (FMaskType*)FMemory::Malloc(sizeof(FMaskType) * NumberOfBytes);
+		FMemory::Memzero(Buffer, NumberOfBytes);
+	}
+}
+
+void FStateTreeTasksCompletionStatus::CopyBuffer(const FStateTreeTasksCompletionStatus& Other)
+{
+	if (!Other.UseInlineBuffer())
+	{
+		const int32 NumberOfBytes = sizeof(FStateTreeTasksCompletionStatus::FMaskType) * BufferNum * 2;
+		Buffer = (FStateTreeTasksCompletionStatus::FMaskType*)FMemory::Malloc(NumberOfBytes);
+		FMemory::Memcpy(Buffer, Other.Buffer, NumberOfBytes);
+	}
+	else
+	{
+		Buffer = Other.Buffer;
+	}
+}
+
+template<typename TTasksCompletionStatusType>
+TTasksCompletionStatusType FStateTreeTasksCompletionStatus::GetStatusInternal(FMaskType Mask, uint8 BufferIndex, uint8 BitsOffset, EStateTreeTaskCompletionType Control)
+{
+	const bool bIsValid = BufferNum > BufferIndex;
+	if (bIsValid && UseInlineBuffer())
+	{
+		check(BufferIndex == 0);
+
+		return TTasksCompletionStatusType(
+			UE::StateTree::Private::GetInlinedBufferPtr(Buffer),
+			UE::StateTree::Private::GetInlinedBufferPtr(Buffer) + 1,
+			Mask,
+			BitsOffset,
+			Control
+		);
+	}
+
+	if (!bIsValid)
+	{
+		check(false);
+		// In case of invalid data (and the check continues), we prefer to not set any task completion than writing in random memory.
+		// Because the mask is 0, no bit will be tested or set. The state tree will never complete.
+		BufferNum = 1;
+		return TTasksCompletionStatusType(
+			UE::StateTree::Private::GetInlinedBufferPtr(Buffer),
+			UE::StateTree::Private::GetInlinedBufferPtr(Buffer) + 1,
+			0,
+			0,
+			EStateTreeTaskCompletionType::Any
+		);
+	}
+
+	constexpr int32 NumberOfBuffers = 2;
+	return TTasksCompletionStatusType(
+		Buffer + (BufferIndex * NumberOfBuffers),
+		Buffer + (BufferIndex * NumberOfBuffers) + 1,
+		Mask,
+		BitsOffset,
+		Control
+	);
+}
+
+UE::StateTree::FTasksCompletionStatus FStateTreeTasksCompletionStatus::GetStatus(const FCompactStateTreeState& State)
+{
+	return GetStatusInternal<UE::StateTree::FTasksCompletionStatus>(
+		State.CompletionTasksMask,
+		State.CompletionTasksMaskBufferIndex,
+		State.CompletionTasksMaskBitsOffset,
+		State.CompletionTasksControl
+		);
+}
+
+UE::StateTree::FConstTasksCompletionStatus FStateTreeTasksCompletionStatus::GetStatus(const FCompactStateTreeState& State) const
+{
+	return const_cast<FStateTreeTasksCompletionStatus*>(this)->GetStatusInternal<UE::StateTree::FConstTasksCompletionStatus>(
+		State.CompletionTasksMask,
+		State.CompletionTasksMaskBufferIndex,
+		State.CompletionTasksMaskBitsOffset,
+		State.CompletionTasksControl
+		);
+}
+
+UE::StateTree::FTasksCompletionStatus FStateTreeTasksCompletionStatus::GetStatus(TNotNull<const UStateTree*> StateTree)
+{
+	constexpr int32 BufferIndex = 0;
+	constexpr int32 BitOffset = 0;
+	return GetStatusInternal<UE::StateTree::FTasksCompletionStatus>(
+		StateTree->CompletionGlobalTasksMask,
+		BufferIndex,
+		BitOffset,
+		StateTree->CompletionGlobalTasksControl
+		);
+}
+
+UE::StateTree::FConstTasksCompletionStatus FStateTreeTasksCompletionStatus::GetStatus(TNotNull<const UStateTree*> StateTree) const
+{
+	constexpr int32 BufferIndex = 0;
+	constexpr int32 BitOffset = 0;
+	return const_cast<FStateTreeTasksCompletionStatus*>(this)->GetStatusInternal<UE::StateTree::FConstTasksCompletionStatus>(
+		StateTree->CompletionGlobalTasksMask,
+		BufferIndex,
+		BitOffset,
+		StateTree->CompletionGlobalTasksControl)
+		;
+}
+
+void FStateTreeTasksCompletionStatus::Push(TNotNull<const UStateTree*> StateTree, const FCompactStateTreeState& State)
+{
+	check(BufferNum > State.CompletionTasksMaskBufferIndex);
+	GetStatus(State).ResetStatus();
+}
+
+bool FStateTreeTasksCompletionStatus::Serialize(FArchive& Ar)
+{
+	if (Ar.IsLoading())
+	{
+		int8 NewBufferNum = 0;
+		Ar << NewBufferNum;
+		if (NewBufferNum != NewBufferNum)
+		{
+			this->~FStateTreeTasksCompletionStatus();
+			BufferNum = NewBufferNum;
+			MallocBufferIfNeeded();
+		}
+
+		if (UseInlineBuffer())
+		{
+			Ar << *(UE::StateTree::Private::GetInlinedBufferPtr(Buffer));
+			Ar << *(UE::StateTree::Private::GetInlinedBufferPtr(Buffer) + 1);
+		}
+		else
+		{
+			for (int32 Index = 0; Index < BufferNum*2; ++Index)
+			{
+				Ar << Buffer[Index];
+			}
+		}
+	}
+	else if (Ar.IsSaving())
+	{
+		Ar << BufferNum;
+		if (UseInlineBuffer())
+		{
+			Ar << *(UE::StateTree::Private::GetInlinedBufferPtr(Buffer));
+			Ar << *(UE::StateTree::Private::GetInlinedBufferPtr(Buffer) + 1);
+		}
+		else
+		{
+			for (int32 Index = 0; Index < BufferNum*2; ++Index)
+			{
+				Ar << Buffer[Index];
+			}
+		}
+	}
+	return true;
+}
+
+bool FStateTreeTasksCompletionStatus::NetSerialize(FArchive& Ar, class UPackageMap* Map, bool& bOutSuccess)
+{
+	bOutSuccess = true;
+	return Serialize(Ar);
+}
\ No newline at end of file
diff --git a/Engine/Plugins/Runtime/StateTree/Source/StateTreeModule/Public/StateTree.h b/Engine/Plugins/Runtime/StateTree/Source/StateTreeModule/Public/StateTree.h
--- a/Engine/Plugins/Runtime/StateTree/Source/StateTreeModule/Public/StateTree.h
+++ b/Engine/Plugins/Runtime/StateTree/Source/StateTreeModule/Public/StateTree.h
@@ -103,31 +103,52 @@ class STATETREEMODULE_API UStateTree : public UDataAsset
 
 public:
 	/** @return Default instance data. */
-	const FStateTreeInstanceData& GetDefaultInstanceData() const { return DefaultInstanceData; }
+	const FStateTreeInstanceData& GetDefaultInstanceData() const
+	{
+		return DefaultInstanceData;
+	}
 
 	/** @return Shared instance data. */
 	TSharedPtr<FStateTreeInstanceData> GetSharedInstanceData() const;
 
 	/** @return Number of context data views required for StateTree execution (Tree params, context data, External data). */
-	int32 GetNumContextDataViews() const { return NumContextData; }
+	int32 GetNumContextDataViews() const
+	{
+		return NumContextData;
+	}
 
 	/** @return List of external data required by the state tree */
-	TConstArrayView<FStateTreeExternalDataDesc> GetExternalDataDescs() const { return ExternalDataDescs; }
+	TConstArrayView<FStateTreeExternalDataDesc> GetExternalDataDescs() const
+	{
+		return ExternalDataDescs;
+	}
 
 	/** @return List of context data enforced by the schema that must be provided through the execution context. */
-	TConstArrayView<FStateTreeExternalDataDesc> GetContextDataDescs() const { return ContextDataDescs; }
+	TConstArrayView<FStateTreeExternalDataDesc> GetContextDataDescs() const
+	{
+		return ContextDataDescs;
+	}
 
 	/** @return true if the other StateTree has compatible context data. */
 	bool HasCompatibleContextData(const UStateTree& Other) const;
 	
 	/** @return List of default parameters of the state tree. Default parameter values can be overridden at runtime by the execution context. */
-	const FInstancedPropertyBag& GetDefaultParameters() const { return Parameters; }
+	const FInstancedPropertyBag& GetDefaultParameters() const
+	{
+		return Parameters;
+	}
 
 	/** @return true if the tree asset can be used at runtime. */
 	bool IsReadyToRun() const;
 
 	/** @return schema that was used to compile the StateTree. */
-	const UStateTreeSchema* GetSchema() const { return Schema; }
+	const UStateTreeSchema* GetSchema() const
+	{
+		return Schema;
+	}
+
+	/** @return Pointer to a frame or null if frame not found */
+	const FCompactStateTreeFrame* GetFrameFromHandle(const FStateTreeStateHandle StateHandle) const;
 
 	/** @return Pointer to a state or null if state not found */ 
 	const FCompactStateTreeState* GetStateFromHandle(const FStateTreeStateHandle StateHandle) const;
@@ -271,6 +292,10 @@ private:
 	UPROPERTY(Instanced)
 	TObjectPtr<UStateTreeSchema> Schema = nullptr;
 
+	/** Runtime frames */
+	UPROPERTY()
+	TArray<FCompactStateTreeFrame> Frames;
+
 	/** Runtime states, root state at index 0 */
 	UPROPERTY()
 	TArray<FCompactStateTreeState> States;
@@ -326,6 +351,10 @@ private:
 	UPROPERTY(Transient)
 	TArray<FStateTreeExternalDataDesc> ExternalDataDescs;
 
+	/** Mask used to test the global tasks completion. */
+	UPROPERTY()
+	uint32 CompletionGlobalTasksMask = 0;
+
 	/** Number of context data, include parameters and all context data. */
 	UPROPERTY()
 	uint16 NumContextData = 0;
@@ -350,6 +379,10 @@ private:
 	UPROPERTY()
 	uint16 GlobalTasksNum = 0;
 
+	/** How the global tasks control the completion of the frame. */
+	UPROPERTY()
+	EStateTreeTaskCompletionType CompletionGlobalTasksControl = EStateTreeTaskCompletionType::Any;
+
 	/** The parameter data type used by the schema. */
 	UPROPERTY()
 	EStateTreeParameterDataType ParameterDataType = EStateTreeParameterDataType::GlobalParameterData;
@@ -393,6 +426,7 @@ private:
 
 	friend struct FStateTreeInstance;
 	friend struct FStateTreeExecutionContext;
+	friend struct FStateTreeTasksCompletionStatus;
 	friend struct FStateTreeMinimalExecutionContext;
 	friend struct FStateTreeReadOnlyExecutionContext;
 	friend struct FStateTreeWeakExecutionContext;
diff --git a/Engine/Plugins/Runtime/StateTree/Source/StateTreeModule/Public/StateTreeExecutionContext.h b/Engine/Plugins/Runtime/StateTree/Source/StateTreeModule/Public/StateTreeExecutionContext.h
--- a/Engine/Plugins/Runtime/StateTree/Source/StateTreeModule/Public/StateTreeExecutionContext.h
+++ b/Engine/Plugins/Runtime/StateTree/Source/StateTreeModule/Public/StateTreeExecutionContext.h
@@ -645,6 +645,7 @@ public:
 	 */
 	void FinishTask(const FStateTreeTaskBase& Task, EStateTreeFinishTaskType FinishType);
 	
+PRAGMA_DISABLE_DEPRECATION_WARNINGS
 	UE_DEPRECATED(5.6, "Use the weak context to finish a task async or FinishTask(FStateTreeTaskBase, EStateTreeFinishTaskType) to finish the current task.")
 	/**
 	 * Finishes a task.
@@ -652,6 +653,7 @@ public:
 	 * If called outside of the tick processing, then the request is buffered and handled on the next tick.
 	 */
 	void FinishTask(const UE::StateTree::FFinishedTask& Task, EStateTreeFinishTaskType FinishType);
+PRAGMA_ENABLE_DEPRECATION_WARNINGS
 
 	/** @return data view of the specified handle relative to given frame. */
 	static FStateTreeDataView GetDataViewFromInstanceStorage(FStateTreeInstanceStorage& InstanceDataStorage, FStateTreeInstanceStorage* CurrentlyProcessedSharedInstanceStorage, const FStateTreeExecutionFrame* ParentFrame, const FStateTreeExecutionFrame& CurrentFrame, const FStateTreeDataHandle Handle);
@@ -781,7 +783,13 @@ protected:
 	/**
 	 * Tick evaluators and global tasks of a specific frame.
 	 */
-	EStateTreeRunStatus TickEvaluatorsAndGlobalTasksForFrame(const float DeltaTime, const bool bTickGlobalTasks, const int32 FrameIndex, const FStateTreeExecutionFrame* CurrentParentFrame, const TNotNull<const FStateTreeExecutionFrame*> CurrentFrame);
+	EStateTreeRunStatus TickEvaluatorsAndGlobalTasksForFrame(const float DeltaTime, const bool bTickGlobalTasks, const int32 FrameIndex, const FStateTreeExecutionFrame* CurrentParentFrame, const TNotNull<FStateTreeExecutionFrame*> CurrentFrame);
+
+	UE_DEPRECATED(5.6, "Use the non const version of TickEvaluatorsAndGlobalTasksForFrame")
+	EStateTreeRunStatus TickEvaluatorsAndGlobalTasksForFrame(const float DeltaTime, const bool bTickGlobalTasks, const int32 FrameIndex, const FStateTreeExecutionFrame* CurrentParentFrame, const FStateTreeExecutionFrame* CurrentFrame)
+	{
+		return TickEvaluatorsAndGlobalTasksForFrame(DeltaTime, bTickGlobalTasks, FrameIndex, CurrentParentFrame, TNotNull<FStateTreeExecutionFrame*>(const_cast<FStateTreeExecutionFrame*>(CurrentFrame)));
+	}
 
 	/**
 	 * Starts evaluators and global tasks.
@@ -800,7 +808,13 @@ protected:
 	void CallStopOnEvaluatorsAndGlobalTasks(const FStateTreeExecutionFrame* ParentFrame, const FStateTreeExecutionFrame& Frame, const FStateTreeTransitionResult& Transition, const FStateTreeIndex16 LastInitializedTaskIndex = FStateTreeIndex16());
 
 	/** Starts temporary instances of global evaluators and tasks for a given frame. */
-	EStateTreeRunStatus StartTemporaryEvaluatorsAndGlobalTasks(const FStateTreeExecutionFrame* CurrentParentFrame, const FStateTreeExecutionFrame& CurrentFrame);
+	EStateTreeRunStatus StartTemporaryEvaluatorsAndGlobalTasks(const FStateTreeExecutionFrame* CurrentParentFrame, FStateTreeExecutionFrame& CurrentFrame);
+
+	UE_DEPRECATED(5.6, "Use the non const version of StartTemporaryEvaluatorsAndGlobalTasks")
+	EStateTreeRunStatus StartTemporaryEvaluatorsAndGlobalTasks(const FStateTreeExecutionFrame* CurrentParentFrame, const FStateTreeExecutionFrame& CurrentFrame)
+	{
+		return StartTemporaryEvaluatorsAndGlobalTasks(CurrentParentFrame, const_cast<FStateTreeExecutionFrame&>(CurrentFrame));
+	}
 
 	/** Stops temporary global evaluators and tasks for the provided frame. */
 	void StopTemporaryEvaluatorsAndGlobalTasks(const FStateTreeExecutionFrame* CurrentParentFrame, const FStateTreeExecutionFrame& CurrentFrame);
@@ -813,7 +827,6 @@ protected:
 
 	struct FTickTaskResult
 	{
-		EStateTreeRunStatus Status;
 		bool bShouldTickTasks = true;
 	};
 	struct FTickTaskArguments
@@ -821,15 +834,16 @@ protected:
 		FTickTaskArguments() = default;
 		float DeltaTime = 0.f;
 		int32 TasksBegin = 0;
-		int32 TasksEnd = 0;
+		int32 TasksNum = 0;
 		int32 Indent = 0;
 		const FStateTreeExecutionFrame* ParentFrame = nullptr;
-		const FStateTreeExecutionFrame* Frame = nullptr;
+		FStateTreeExecutionFrame* Frame = nullptr;
 		UE::StateTree::FActiveStateID StateID;
+		UE::StateTree::FTasksCompletionStatus* TasksCompletionStatus = nullptr;
 		bool bIsGlobalTasks = false;
 		bool bShouldTickTasks = true;
 	};
-	/** Ticks tasks for a specific state or frame. */
+	/** Ticks tasks and updates the bindings for a specific state or frame. */
 	FTickTaskResult TickTasks(const FTickTaskArguments& Args);
 
 	/** Common functionality shared by the tick methods. */
@@ -885,24 +899,30 @@ protected:
 	 */
 	bool TriggerTransitions();
 
-	/** Adds the state to the completed list and adds task to the finished list. */
-	void FinishTaskInternal(UE::StateTree::FFinishedTask::EReasonType Reason, UE::StateTree::FActiveFrameID ActiveFrameID, UE::StateTree::FActiveStateID ActiveStateID, int32 TaskIndex, EStateTreeRunStatus RunStatus);
-
+PRAGMA_DISABLE_DEPRECATION_WARNINGS
+	UE_DEPRECATED(5.6, "Replaced with FStateTreeTasksCompletionStatus")
 	/** Confirms that the frame and state ID are valid and the task index is correct. */
 	bool IsFinishedTaskValid(const UE::StateTree::FFinishedTask& Task) const;
+PRAGMA_ENABLE_DEPRECATION_WARNINGS
 
+	UE_DEPRECATED(5.6, "Replaced with FStateTreeTasksCompletionStatus")
 	/**
 	 * Removes stale entries and fills the completed states list.
 	 * The tree can be in a different state from when a finished task was added to the pending list and the finished task may not be valid.
 	 * @param bMarkTaskProcessed when true, marks the tasks as processed. When false, only used the already marked tasks.
 	 */
 	void UpdateCompletedStateList(bool bMarkTaskProcessed);
-	UE_DEPRECATED(5.6, "Use UpdateCompletedStateList with the bMarkTaskProcessed argument")
+
+	UE_DEPRECATED(5.6, "Replaced with FStateTreeTasksCompletionStatus")
 	void UpdateCompletedStateList();
 
+PRAGMA_DISABLE_DEPRECATION_WARNINGS
+	UE_DEPRECATED(5.6, "Replaced with FStateTreeTasksCompletionStatus")
 	/** Adds the state to the completed list from a finished task. */
 	void MarkStateCompleted(UE::StateTree::FFinishedTask& FinisedTask);
+PRAGMA_ENABLE_DEPRECATION_WARNINGS
 
+	UE_DEPRECATED(5.6, "Replaced with FStateTreeTasksCompletionStatus")
 	/** In the list of completed states, returns the run status when it's completed by a global task. */
 	EStateTreeRunStatus GetGlobalTasksCompletedStatesStatus() const;
 
@@ -1051,28 +1071,6 @@ protected:
 	TOptional<EStateTreeRunStatus> CompletedStateRunStatus;
 #endif
 
-	/**
-	 * UpdateCompletedStateList is called once per context. Either from TickTasks or TickTransitions.
-	 * User can call TickTasks and wait some frame before calling TickTransitions.
-	 */
-	bool bCompletedStatesUpdated = false;
-
-	/** State that was completed in this frame. */
-	struct FCompletedState : UE::StateTree::FFinishedTask
-	{
-		FCompletedState() = default;
-		FCompletedState(UE::StateTree::FFinishedTask Other)
-			: UE::StateTree::FFinishedTask(Other)
-		{
-		};
-	};
-
-	/**
-	 * Completed states while ticking.
-	 * Build from the PendingCompletedStates list and maintained by the CompleteTask function.
-	 */
-	TArray<FCompletedState, TInlineAllocator<4>> CompletedStates;
-
 	/** Next transition, used by RequestTransition(). */
 	FStateTreeTransitionResult NextTransition;
 
diff --git a/Engine/Plugins/Runtime/StateTree/Source/StateTreeModule/Public/StateTreeExecutionTypes.h b/Engine/Plugins/Runtime/StateTree/Source/StateTreeModule/Public/StateTreeExecutionTypes.h
--- a/Engine/Plugins/Runtime/StateTree/Source/StateTreeModule/Public/StateTreeExecutionTypes.h
+++ b/Engine/Plugins/Runtime/StateTree/Source/StateTreeModule/Public/StateTreeExecutionTypes.h
@@ -6,6 +6,7 @@
 #include "StateTreeEvents.h"
 #include "StateTreeExecutionExtension.h"
 #include "StateTreeStatePath.h"
+#include "StateTreeTasksStatus.h"
 #include "UObject/ObjectKey.h"
 
 #include "StateTreeExecutionTypes.generated.h"
@@ -66,6 +67,7 @@ enum class EStateTreeRunStatus : uint8
 
 
 /** Status describing how a task finished. */
+UENUM()
 enum class EStateTreeFinishTaskType : uint8
 {	
 	/** The task execution failed. */
@@ -311,7 +313,8 @@ struct STATETREEMODULE_API FStateTreeActiveStates
 {
 	GENERATED_BODY()
 
-	static constexpr uint8 MaxStates = 8;	// Max number of active states
+	/** Max number of active states. */
+	static constexpr uint8 MaxStates = 8;
 
 	FStateTreeActiveStates() = default;
 	
@@ -345,7 +348,7 @@ struct STATETREEMODULE_API FStateTreeActiveStates
 		{
 			return false;
 		}
-		
+
 		States[NumStates] = StateHandle;
 		StateIDs[NumStates] = StateID;
 		++NumStates;
@@ -475,16 +478,28 @@ struct STATETREEMODULE_API FStateTreeActiveStates
 	}
 	
 	/** Returns last state in the array, or invalid state if the array is empty. */
-	FStateTreeStateHandle Last() const { return NumStates > 0 ? States[NumStates - 1] : FStateTreeStateHandle::Invalid; }
+	FStateTreeStateHandle Last() const
+	{
+		return NumStates > 0 ? States[NumStates - 1] : FStateTreeStateHandle::Invalid;
+	}
 	
 	/** Returns number of states in the array. */
-	int32 Num() const { return NumStates; }
+	int32 Num() const
+	{
+		return NumStates;
+	}
 
 	/** Returns true if the index is within array bounds. */
-	bool IsValidIndex(const int32 Index) const { return Index >= 0 && Index < (int32)NumStates; }
+	bool IsValidIndex(const int32 Index) const
+	{
+		return Index >= 0 && Index < (int32)NumStates;
+	}
 	
 	/** Returns true if the array is empty. */
-	bool IsEmpty() const { return NumStates == 0; } 
+	bool IsEmpty() const
+	{
+		return NumStates == 0;
+	}
 
 	/** Returns a specified state in the array. */
 	FORCEINLINE const FStateTreeStateHandle& operator[](const int32 Index) const
@@ -690,7 +705,9 @@ PRAGMA_ENABLE_DEPRECATION_WARNINGS
 namespace UE::StateTree
 {
 /** Describes a finished task waiting to be processed by an execution context. */
-struct FFinishedTask
+struct
+UE_DEPRECATED(5.6, "Replaced with FStateTreeTasksCompletionStatus")
+FFinishedTask
 {
 	enum class EReasonType : uint8
 	{
@@ -866,12 +883,16 @@ struct STATETREEMODULE_API FStateTreeExecutionFrame
 
 	/** The root state of the frame (e.g. Root state or a subtree). */
 	UPROPERTY()
-	FStateTreeStateHandle RootState = FStateTreeStateHandle::Root; 
+	FStateTreeStateHandle RootState = FStateTreeStateHandle::Root;
 	
 	/** Active states in this frame */
 	UPROPERTY()
 	FStateTreeActiveStates ActiveStates;
 
+	/** Flag to track the completion of a global task or a task from a state in the ActiveStates. */
+	UPROPERTY(Transient)
+	FStateTreeTasksCompletionStatus ActiveTasksStatus;
+
 	/** Unique frame id for this frame. Can be used to identify the frame. */
 	UE::StateTree::FActiveFrameID FrameID;
 
@@ -990,9 +1011,6 @@ public:
 	UPROPERTY()
 	TArray<FStateTreeTransitionDelayedState> DelayedTransitions;
 
-	/** Pending finished tasks that need processing. */
-	TArray<UE::StateTree::FFinishedTask> FinishedTasks;
-
 	/** Used by state tree random-based operations. */
 	UPROPERTY()
 	FRandomStream RandomStream;
@@ -1034,7 +1052,7 @@ public:
 	UPROPERTY()
 	FStateTreeIndex16 LastExitedNodeIndex = FStateTreeIndex16::Invalid;
 
-	/** Result of last tick */
+	/** Result of last TickTasks */
 	UPROPERTY()
 	EStateTreeRunStatus LastTickStatus = EStateTreeRunStatus::Failed;
 
@@ -1054,18 +1072,28 @@ public:
 	UPROPERTY()
 	uint16 StateChangeCount = 0;
 
+	/** A task that completed a state or a global task that completed a global frame. */
+	UPROPERTY()
+	bool bHasPendingCompletedState = false;
+
 #if WITH_EDITORONLY_DATA
-	UE_DEPRECATED(5.6, "Use DelayedCompletedTasks for trigerring completion transitions.")
+PRAGMA_DISABLE_DEPRECATION_WARNINGS
+	UE_DEPRECATED(5.6, "Replaced with FStateTreeTasksCompletionStatus")
+	/** Pending finished tasks that need processing. */
+	TArray<UE::StateTree::FFinishedTask> FinishedTasks;
+
+	UE_DEPRECATED(5.6, "Use FinishTask to completed a state.")
 	/** Handle of the state that was first to report state completed (success or failure), used to trigger completion transitions. */
 	UPROPERTY()
 	FStateTreeIndex16 CompletedFrameIndex = FStateTreeIndex16::Invalid; 
 	
-	UE_DEPRECATED(5.6, "Use DelayedCompletedTasks for trigerring completion transitions.")
+	UE_DEPRECATED(5.6, "Use FinishTask to completed a state.")
 	UPROPERTY()
 	FStateTreeStateHandle CompletedStateHandle = FStateTreeStateHandle::Invalid;
 
 	UE_DEPRECATED(5.6, "CurrentExecutionContext is not needed anymore. Use FrameID and StateID.")
 	FStateTreeExecutionContext* CurrentExecutionContext = nullptr;
+PRAGMA_ENABLE_DEPRECATION_WARNINGS
 #endif //WITH_EDITORONLY_DATA
 };
 
diff --git a/Engine/Plugins/Runtime/StateTree/Source/StateTreeModule/Public/StateTreeTaskBase.h b/Engine/Plugins/Runtime/StateTree/Source/StateTreeModule/Public/StateTreeTaskBase.h
--- a/Engine/Plugins/Runtime/StateTree/Source/StateTreeModule/Public/StateTreeTaskBase.h
+++ b/Engine/Plugins/Runtime/StateTree/Source/StateTreeModule/Public/StateTreeTaskBase.h
@@ -27,6 +27,10 @@ struct STATETREEMODULE_API FStateTreeTaskBase : public FStateTreeNodeBase
 		, bShouldAffectTransitions(false)
 		, bConsideredForScheduling(true)
 		, bTaskEnabled(true)
+#if WITH_EDITORONLY_DATA
+		, bConsideredForCompletion(true)
+		, bCanEditConsideredForCompletion(true)
+#endif
 	{
 	}
 	
@@ -135,6 +139,19 @@ struct STATETREEMODULE_API FStateTreeTaskBase : public FStateTreeNodeBase
 
 	UPROPERTY()
 	EStateTreeTransitionPriority TransitionHandlingPriority = EStateTreeTransitionPriority::Normal;
+
+#if WITH_EDITORONLY_DATA
+	/**
+	 * True if the task is considered for completion.
+	 * False if the task runs in the background without affecting the state completion.
+	 */
+	UPROPERTY()
+	uint8 bConsideredForCompletion : 1;
+
+	/** True if the user can edit bConsideredForCompletion in the editor. */
+	UPROPERTY()
+	uint8 bCanEditConsideredForCompletion : 1;
+#endif
 };
 
 /**
diff --git a/Engine/Plugins/Runtime/StateTree/Source/StateTreeModule/Public/StateTreeTasksStatus.h b/Engine/Plugins/Runtime/StateTree/Source/StateTreeModule/Public/StateTreeTasksStatus.h
new file mode 100644
--- /dev/null
+++ b/Engine/Plugins/Runtime/StateTree/Source/StateTreeModule/Public/StateTreeTasksStatus.h
@@ -0,0 +1,339 @@
+// Copyright Epic Games, Inc. All Rights Reserved.
+
+#pragma once
+
+#include "Misc/NotNull.h"
+#include "StateTreeTasksStatus.generated.h"
+
+class FArchive;
+struct FCompactStateTreeFrame;
+struct FCompactStateTreeState;
+class UPackageMap;
+class UStateTree;
+
+UENUM()
+enum class EStateTreeTaskCompletionType : uint8
+{
+	/** All tasks need to complete for the group to completes. */
+	All,
+	/** Any task completes the group. */
+	Any,
+};
+
+namespace UE::StateTree
+{
+	enum class ETaskCompletionStatus : uint8
+	{
+		/** The task is running. */
+		Running = 0,
+		/**
+		 * The task stopped without a particular reason. (ie. Aborted).
+		 * Use for backward compatibility. Prefer Succeeded and Failed.
+		 */
+		Stopped = 1,
+		/** The task stopped with a success. */
+		Succeeded = 2,
+		/** The task stopped with a failure. */
+		Failed = 3,
+	};
+	/** Number of entry in the ETaskCompletionStatus. */
+	static constexpr int32 NumberOfTaskStatus = 4;
+
+	/** */
+	template<typename T>
+	struct TTasksCompletionStatus
+	{
+	public:
+		static constexpr int32 MaxNumTasks = sizeof(T)*8;
+
+	public:
+		TTasksCompletionStatus(TNotNull<T*> InFirstCompletionBits, TNotNull<T*> InSecondCompletionBits, T InCompletionMask, int32 InBitIndex, EStateTreeTaskCompletionType InTaskControl)
+			: FirstCompletionBits(InFirstCompletionBits)
+			, SecondCompletionBits(InSecondCompletionBits)
+			, CompletionMask(InCompletionMask)
+			, BitIndex(InBitIndex)
+			, TaskControl(InTaskControl)
+		{}
+
+	public:
+		/**
+		 * @param StateTaskIndex The index of the task in the state.
+		 * @return the status of a task.
+		 */
+		[[nodiscard]] ETaskCompletionStatus GetStatus(int32 StateTaskIndex) const
+		{
+			StateTaskIndex += BitIndex;
+			check(StateTaskIndex < MaxNumTasks);
+			return static_cast<ETaskCompletionStatus>(GetStatusInternal(StateTaskIndex));
+		}
+
+		/**
+		 * The completion status of all tasks or any task, in priority order.
+		 * If any of the tasks fail, the completion status will be Failed regardless of any other tasks.
+		 * In case the type is Any, it will return Succeeded before Stopped and before Running.
+		 * @return the completion status of all tasks.
+		 */
+		[[nodiscard]] ETaskCompletionStatus GetCompletionStatus() const
+		{
+			const T FirstCompletionBitsMasked = (*FirstCompletionBits) & CompletionMask;
+			const T SecondCompletionBitsMasked = (*SecondCompletionBits) & CompletionMask;
+			if ((SecondCompletionBitsMasked & FirstCompletionBitsMasked) != 0)
+			{
+				return ETaskCompletionStatus::Failed;
+			}
+			if (TaskControl == EStateTreeTaskCompletionType::All)
+			{
+				if (SecondCompletionBitsMasked == CompletionMask)
+				{
+					return ETaskCompletionStatus::Succeeded;
+				}
+				if (FirstCompletionBitsMasked == CompletionMask)
+				{
+					return ETaskCompletionStatus::Stopped;
+				}
+				// if we have a mix of Succeeded and Stopped, return Succeeded.
+				return (FirstCompletionBitsMasked | SecondCompletionBitsMasked) == CompletionMask ? ETaskCompletionStatus::Succeeded : ETaskCompletionStatus::Running;
+			}
+			else
+			{
+				if (SecondCompletionBitsMasked != 0)
+				{
+					return ETaskCompletionStatus::Succeeded;
+				}
+				return FirstCompletionBitsMasked != 0 ? ETaskCompletionStatus::Stopped : ETaskCompletionStatus::Running;
+			}
+		}
+
+		/**
+		 * @param StateTaskIndex The index of the task in the state.
+		 * @return true when the task status is running.
+		 */
+		[[nodiscard]] bool IsRunning(int32 StateTaskIndex) const
+		{
+			StateTaskIndex += BitIndex;
+			check(StateTaskIndex < MaxNumTasks);
+			const T TaskMask = T(1) << StateTaskIndex;
+			return (((*FirstCompletionBits) | (*SecondCompletionBits)) & TaskMask) == 0;
+		}
+
+		/**
+		 * @param StateTaskIndex The index of the task in the state.
+		 * @return true when the task status is failed.
+		 */
+		[[nodiscard]] bool HasFailed(int32 StateTaskIndex) const
+		{
+			StateTaskIndex += BitIndex;
+			check(StateTaskIndex < MaxNumTasks);
+			const T TaskMask = T(1) << StateTaskIndex;
+			return ((*FirstCompletionBits) & (*SecondCompletionBits) & TaskMask) != 0;
+		}
+
+		/** @return true when there's any failure or all tasks succeeded or stopped. */
+		[[nodiscard]] bool HasAllCompleted() const
+		{
+			const bool bHasAnyFailure = ((*FirstCompletionBits) & (*SecondCompletionBits) & CompletionMask) != 0;
+			const T CompletionResult = (((*FirstCompletionBits) | (*SecondCompletionBits)) & CompletionMask);
+			return bHasAnyFailure || CompletionResult == CompletionMask;
+		}
+
+		/** @return true when there's any failure or any tasks succeeded or stopped. */
+		[[nodiscard]] bool HasAnyCompleted() const
+		{
+			return (((*FirstCompletionBits) | (*SecondCompletionBits)) & CompletionMask) != 0;
+		}
+
+		/** @return true when there's any failure. */
+		[[nodiscard]] bool HasAnyFailed() const
+		{
+			return ((*FirstCompletionBits) & (*SecondCompletionBits) & CompletionMask) != 0;
+		}
+
+		/** @return true when there's any failure or respect the task control. */
+		[[nodiscard]] bool IsCompleted() const
+		{
+			return TaskControl == EStateTreeTaskCompletionType::All ? HasAllCompleted() : HasAnyCompleted();
+		}
+
+		/**
+		 * Set the status of a task.
+		 * @param StateTaskIndex The index of the task in the state.
+		 */
+		void SetStatus(int32 StateTaskIndex, ETaskCompletionStatus NewStatus)
+		{
+			StateTaskIndex += BitIndex;
+			check(StateTaskIndex < MaxNumTasks);
+			SetStatusInternal(StateTaskIndex, NewStatus);
+		}
+
+		/**
+		 * Set the status of a task respecting the previous status value.
+		 * The priority is Failed, Succeeded, Stopped, Running.
+		 * @param StateTaskIndex The index of the task in the state.
+		 */
+		ETaskCompletionStatus SetStatusWithPriority(int32 StateTaskIndex, ETaskCompletionStatus NewStatus)
+		{
+			StateTaskIndex += BitIndex;
+			check(StateTaskIndex < MaxNumTasks);
+			uint8 CurrentStatus = GetStatusInternal(StateTaskIndex);
+			if (static_cast<uint8>(NewStatus) > CurrentStatus)
+			{
+				SetStatusInternal(StateTaskIndex, NewStatus);
+				CurrentStatus = static_cast<uint8>(NewStatus);
+			}
+			return static_cast<ETaskCompletionStatus>(CurrentStatus);
+		}
+
+		/** Set the status of all tasks in the completion mask. */
+		void SetCompletionStatus(ETaskCompletionStatus NewStatus)
+		{
+			const T ClearMask = ~CompletionMask;
+			(*FirstCompletionBits) &= ClearMask;
+			(*SecondCompletionBits) &= ClearMask;
+			if ((static_cast<uint8>(NewStatus) & 0x1) != 0)
+			{
+				(*FirstCompletionBits) |= CompletionMask;
+			}
+			if ((static_cast<uint8>(NewStatus) & 0x2) != 0)
+			{
+				(*SecondCompletionBits) |= CompletionMask;
+			}
+		}
+
+		/** Set the status of all tasks to running. */
+		void ResetStatus()
+		{
+			const T ClearMask = ~CompletionMask;
+			(*FirstCompletionBits) &= ClearMask;
+			(*SecondCompletionBits) &= ClearMask;
+		}
+
+	private:
+		uint8 GetStatusInternal(int32 Index) const
+		{
+			T Result = ((*SecondCompletionBits) >> Index) & 0x1;
+			Result <<= 1;
+			Result += ((*FirstCompletionBits) >> Index) & 0x1;
+			return static_cast<uint8>(Result);
+		}
+
+		void SetStatusInternal(int32 Index, ETaskCompletionStatus NewStatus)
+		{
+			// Clear the current value, then set the new value.
+			const T ClearMask = ~(T(1) << Index);
+			(*FirstCompletionBits) &= ClearMask;
+			(*SecondCompletionBits) &= ClearMask;
+			(*FirstCompletionBits) |= (static_cast<T>(NewStatus) & 0x1) << Index;
+			(*SecondCompletionBits) |= ((static_cast<T>(NewStatus) & 0x2) >> 1) << Index;
+		}
+
+		/** The first buffer. 00 is Running and 01 is Stopped. */
+		T* FirstCompletionBits;
+		/** The second buffer. 10 is Succeeded and 11 is Failed. */
+		T* SecondCompletionBits;
+		/** The mask that represents each task considered by the state/frame for completion. */
+		T CompletionMask;
+		/** The offset, in bits, of the first task inside the mask. */
+		int32 BitIndex;
+		/** How the mask is tested to complete the state/frame. */
+		EStateTreeTaskCompletionType TaskControl;
+	};
+
+	using FTasksCompletionStatus = TTasksCompletionStatus<uint32>;
+	using FConstTasksCompletionStatus = TTasksCompletionStatus<const uint32>;
+}
+
+/**
+ * Container for task status for all the active states and global tasks.
+ * Each task needs 2 bits of information. The information is in 2 different int32. 1 bit per int32 instead of 2bits inside the same int32.
+ * The int32 are sequential: The first masks takes the 2 first int32 and (if needed) the second mask takes the 3rd and 4th int32.
+ * A state/global has at least 1 entry (2bits), even if there are no tasks on the state/global. It is to represent if the state completes (ex: when no tasks are ticked the state completes).
+ * The bits, from different states, are packed until there are too many to fit inside the same int32.
+ * When it can't be packed, all the bits from that overflow state are moved to the next int32.
+ * When possible (when the number of tasks is below 32), the buffer will be inlined. Otherwise, it will use dynamic memory.
+ * We allocate the worst case scenario when the frame is created.
+ * 
+ * Ex: For the tree:
+ *     (0)Global task: 8 tasks. (1)State Root: 6 tasks. (2)StateA: no task (it takes one bit). (3)StateB: 10 tasks.
+ *     (4)State: 8 tasks, not enough space go to the next int32. (5)State: 1 task, takes one bit.
+ *     The mask looks like:
+ * [------33333333333211111100000000|-----------------------544444444]
+ *     The buffer will be 4 int32.
+ *     The first 2 int32 are for the (0)Global task, (1)State Root, (2)State, (3)StateB
+ *     The next 2 int32 are for the (4)State, (5)State
+ * 
+ * The first bit of each buffer is combined to represent the ETaskCompletionStatus.
+ * Ex: (on int8 instead of int32 to be shorter in this description)
+ *     There are 3 tasks. The bits 0-2 are used. The other bits are never set/read. We use the "completion mask" to filter them.
+ * [00001100|00001010]
+ *     Task 1 has the value 00 (0 from the first buffer and 0 from the second buffer). It's Running.
+ *     Task 2 has the value 01 (0 from the first buffer and 1 from the second buffer). It's Stopped.
+ *     Task 3 has the value 10 (1 from the first buffer and 0 from the second buffer). It's Succeeded.
+ */
+USTRUCT()
+struct FStateTreeTasksCompletionStatus
+{
+	GENERATED_BODY()
+
+	using FMaskType = uint32;
+	static constexpr int32 MaxNumberOfTasksPerGroup = sizeof(FMaskType)*8;
+	/** 32 global global tasks + 32 tasks per states for a max of 8 states(max 8 states per frame). */
+	static constexpr int32 MaxTotalAmountOfTasks = MaxNumberOfTasksPerGroup + MaxNumberOfTasksPerGroup * 8;
+
+public:
+	STATETREEMODULE_API FStateTreeTasksCompletionStatus() = default;
+	STATETREEMODULE_API FStateTreeTasksCompletionStatus(const FCompactStateTreeFrame& Frame);
+	STATETREEMODULE_API ~FStateTreeTasksCompletionStatus();
+
+	STATETREEMODULE_API FStateTreeTasksCompletionStatus(const FStateTreeTasksCompletionStatus&);
+	STATETREEMODULE_API FStateTreeTasksCompletionStatus(FStateTreeTasksCompletionStatus&&);
+	STATETREEMODULE_API FStateTreeTasksCompletionStatus& operator=(const FStateTreeTasksCompletionStatus&);
+	STATETREEMODULE_API FStateTreeTasksCompletionStatus& operator=(FStateTreeTasksCompletionStatus&&);
+
+public:
+	[[nodiscard]] STATETREEMODULE_API UE::StateTree::FTasksCompletionStatus GetStatus(const FCompactStateTreeState& State);
+	[[nodiscard]] STATETREEMODULE_API UE::StateTree::FConstTasksCompletionStatus GetStatus(const FCompactStateTreeState& State) const;
+
+	[[nodiscard]] STATETREEMODULE_API UE::StateTree::FTasksCompletionStatus GetStatus(TNotNull<const UStateTree*> StateTree);
+	[[nodiscard]] STATETREEMODULE_API UE::StateTree::FConstTasksCompletionStatus GetStatus(TNotNull<const UStateTree*> StateTree) const;
+
+	STATETREEMODULE_API void Push(TNotNull<const UStateTree*> StateTree, const FCompactStateTreeState& State);
+
+	STATETREEMODULE_API bool Serialize(FArchive& Ar);
+	STATETREEMODULE_API bool NetSerialize(FArchive& Ar, UPackageMap* Map, bool& bOutSuccess);
+
+private:
+	static constexpr int32 MaxNumberOfTaskForInlineBuffer = sizeof(FMaskType*) * 8 / 2;
+
+	inline bool UseInlineBuffer() const
+	{
+		return BufferNum <= 1;
+	}
+
+	void MallocBufferIfNeeded();
+	void CopyBuffer(const FStateTreeTasksCompletionStatus& Other);
+
+	template<typename TTasksCompletionStatusType>
+	inline TTasksCompletionStatusType GetStatusInternal(FMaskType Mask, uint8 BufferIndex, uint8 BitsOffset, EStateTreeTaskCompletionType Control);
+
+	/**
+	 * Dynamic or inlined container for a FMaskType.
+	 * The actual memory used is double of BufferNum because each task has 2 bits of information.
+	 */
+	FMaskType* Buffer = nullptr;
+	/**
+	 * Number of requested FMaskType.
+	 * If <= 1, it will use Buffer as an inlined buffer.
+	 */
+	uint8 BufferNum = 0;
+};
+
+template<>
+struct TStructOpsTypeTraits<FStateTreeTasksCompletionStatus> : public TStructOpsTypeTraitsBase2<FStateTreeTasksCompletionStatus>
+{
+	enum
+	{
+		WithCopy = true,
+		WithSerializer = true,
+		WithNetSerializer = true,
+	};
+};
diff --git a/Engine/Plugins/Runtime/StateTree/Source/StateTreeModule/Public/StateTreeTypes.h b/Engine/Plugins/Runtime/StateTree/Source/StateTreeModule/Public/StateTreeTypes.h
--- a/Engine/Plugins/Runtime/StateTree/Source/StateTreeModule/Public/StateTreeTypes.h
+++ b/Engine/Plugins/Runtime/StateTree/Source/StateTreeModule/Public/StateTreeTypes.h
@@ -11,6 +11,7 @@
 #include "GameplayTagContainer.h"
 #include "StateTreeDelegate.h"
 #include "StateTreeIndexTypes.h"
+#include "StateTreeTasksStatus.h"
 #include "StateTreeTypes.generated.h"
 
 STATETREEMODULE_API DECLARE_LOG_CATEGORY_EXTERN(LogStateTree, Warning, All);
@@ -643,6 +644,27 @@ struct STATETREEMODULE_API FCompactStateTransition
 	uint8 bConsumeEventOnSelect : 1 = true;
 };
 
+/**
+ *  Runtime representation of a StateTree frame.
+ */
+USTRUCT()
+struct STATETREEMODULE_API FCompactStateTreeFrame
+{
+	GENERATED_BODY()
+
+	/** The root state of the frame (e.g. Root state or a subtree). */
+	UPROPERTY()
+	FStateTreeStateHandle RootState = FStateTreeStateHandle::Invalid;
+
+	/**
+	 * The max number of masks needed by the frame.
+	 * It is the worst case of all the possible active state combinations for the frame.
+	 * It includes the global tasks.
+	 */
+	UPROPERTY()
+	uint8 NumberOfTasksStatusMasks = 0;
+};
+
 /**
  *  Runtime representation of a StateTree state.
  */
@@ -681,12 +703,18 @@ struct STATETREEMODULE_API FCompactStateTreeState
 		return bHasTickTriggerTransitions || (bHasEvent && bHasEventTriggerTransitions) || (bHasBroadcastedDelegates && bHasDelegateTriggerTransitions);
 	}
 
-	/** @return True if the state has any transitions that tick on completion. */
-	bool ShouldTickCompletionTransitions(bool bSucceeded) const
+	/**
+	 * A state can complete with Stopped, Succeeded, and Failed.
+	 * The state can contain transitions that trigger on any completes, or only if the state succeeded, or only if the state failed.
+	 * @param bSucceeded Whenever the state completion succeeded.
+	 * @param bFailed Whenever the state completion failed.
+	 * @return True if the state has any transitions that tick on completion.
+	 */
+	bool ShouldTickCompletionTransitions(bool bSucceeded, bool bFailed) const
 	{
 		return bHasCompletedTriggerTransitions
 			|| (bHasSucceededTriggerTransitions && bSucceeded)
-			|| (bHasFailedTriggerTransitions && !bSucceeded);
+			|| (bHasFailedTriggerTransitions && bFailed);
 	}
 
 	/** Description of an event required to enter the state. */
@@ -761,6 +789,29 @@ struct STATETREEMODULE_API FCompactStateTreeState
 	UPROPERTY()
 	float CustomTickRate = 0.f;
 
+	/** Mask used to test the tasks completion. */
+	UPROPERTY()
+	uint32 CompletionTasksMask = 0;
+
+	/**
+	 * Index in the mask buffer used by the state. 
+	 * CompletionTasksMaskBufferIndex = (Final Task Bit) / (size of mask:32)
+	 */
+	UPROPERTY()
+	uint8 CompletionTasksMaskBufferIndex = 0;
+	
+	/**
+	 * Offset, in bits, of the first flag inside the mask.
+	 * Used to access the bit from the state task index.
+	 * CompletionTasksMaskBitsOffset = (Final Task Bit) % (size of mask:32)
+	 */
+	UPROPERTY()
+	uint8 CompletionTasksMaskBitsOffset = 0;
+
+	/** How the tasks control the completion of the state. */
+	UPROPERTY()
+	EStateTreeTaskCompletionType CompletionTasksControl = EStateTreeTaskCompletionType::Any;
+
 	/** Number of enter conditions */
 	UPROPERTY()
 	uint8 EnterConditionsNum = 0;
diff --git a/Engine/Plugins/Runtime/StateTree/Source/StateTreeTestSuite/Private/StateTreeTaskStateTest.cpp b/Engine/Plugins/Runtime/StateTree/Source/StateTreeTestSuite/Private/StateTreeTaskStateTest.cpp
new file mode 100644
--- /dev/null
+++ b/Engine/Plugins/Runtime/StateTree/Source/StateTreeTestSuite/Private/StateTreeTaskStateTest.cpp
@@ -0,0 +1,430 @@
+// Copyright Epic Games, Inc. All Rights Reserved.
+
+#include "StateTreeTest.h"
+#include "StateTreeTestBase.h"
+#include "StateTreeTestTypes.h"
+
+#include "StateTreeCompilerLog.h"
+#include "StateTreeEditorData.h"
+#include "StateTreeCompiler.h"
+#include "Conditions/StateTreeCommonConditions.h"
+
+#define LOCTEXT_NAMESPACE "AITestSuite_StateTreeTask"
+
+UE_DISABLE_OPTIMIZATION_SHIP
+
+namespace UE::StateTree::Tests
+{
+
+namespace Private
+{
+} // namespace Private
+
+struct FStateTreeTest_TasksCompletion_AllAny : FStateTreeTestBase
+{
+	virtual bool InstantTest() override
+	{
+		// Main asset
+		UStateTree& StateTree = NewStateTree();
+		UStateTreeState* RootState = nullptr;
+		{
+			UStateTreeEditorData& EditorData = *Cast<UStateTreeEditorData>(StateTree.EditorData);
+
+			UStateTreeState& Root = EditorData.AddSubTree(FName(TEXT("Tree1StateRoot")));
+			RootState = &Root;
+			Root.TasksCompletion = EStateTreeTaskCompletionType::All;
+			for (int32 TaskIndex = 0; TaskIndex < FStateTreeTasksCompletionStatus::MaxNumberOfTasksPerGroup; ++TaskIndex)
+			{
+				TStateTreeEditorNode<FTestTask_Stand>& Task = Root.AddTask<FTestTask_Stand>(FName(*FString::Printf(TEXT("Tree1StateRootTask_%d"), TaskIndex)));
+				Task.GetNode().TicksToCompletion = TaskIndex + 1;
+			}
+			UStateTreeState& Tree1StateA = Root.AddChildState("Tree1StateA");
+			Tree1StateA.TasksCompletion = EStateTreeTaskCompletionType::All;
+			for (int32 TaskIndex = 0; TaskIndex < FStateTreeTasksCompletionStatus::MaxNumberOfTasksPerGroup; ++TaskIndex)
+			{
+				TStateTreeEditorNode<FTestTask_Stand>& Task = Tree1StateA.AddTask<FTestTask_Stand>(FName(*FString::Printf(TEXT("Tree1StateATask_%d"), TaskIndex)));
+				Task.GetNode().TicksToCompletion = TaskIndex + 1;
+			}
+
+			{
+				FStateTreeCompilerLog Log;
+				FStateTreeCompiler Compiler(Log);
+				const bool bResult = Compiler.Compile(StateTree);
+				AITEST_TRUE("StateTree should get compiled", bResult);
+			}
+		}
+
+		// Test All
+		{
+			EStateTreeRunStatus Status = EStateTreeRunStatus::Unset;
+			FStateTreeInstanceData InstanceData;
+			FTestStateTreeExecutionContext Exec(StateTree, StateTree, InstanceData);
+			const bool bInitSucceeded = Exec.IsValid();
+			AITEST_TRUE("StateTree should init", bInitSucceeded);
+
+			{
+				Status = Exec.Start();
+				AITEST_EQUAL("Start should complete with Running", Status, EStateTreeRunStatus::Running);
+				Exec.LogClear();
+			}
+
+			for (int32 TickIndex = 0; TickIndex < FStateTreeTasksCompletionStatus::MaxNumberOfTasksPerGroup; ++TickIndex)
+			{
+				Status = Exec.Tick(0.1f);
+				AITEST_EQUAL(TEXT("Tick should complete with Running"), Status, EStateTreeRunStatus::Running);
+				AITEST_TRUE(TEXT("In correct states"), Exec.ExpectInActiveStates("Tree1StateRoot", "Tree1StateA"));
+
+				for (int32 TaskIndex = 0; TaskIndex < FStateTreeTasksCompletionStatus::MaxNumberOfTasksPerGroup; ++TaskIndex)
+				{
+					bool bTickedRoot = Exec.Expect(*FString::Printf(TEXT("Tree1StateRootTask_%d"), TaskIndex), TEXT("Tick"));
+					if (TaskIndex < TickIndex)
+					{
+						AITEST_FALSE(*FString::Printf(TEXT("Should not tick Task %d, %d"), TickIndex, TaskIndex), bTickedRoot);
+					}
+					else
+					{
+						AITEST_TRUE(*FString::Printf(TEXT("Should tick Task %d, %d"), TickIndex, TaskIndex), bTickedRoot);
+					}
+					bool bTickedA = Exec.Expect(*FString::Printf(TEXT("Tree1StateATask_%d"), TaskIndex), TEXT("Tick"));
+					if (TaskIndex < TickIndex)
+					{
+						AITEST_FALSE(*FString::Printf(TEXT("Should not tick Task A %d, %d"), TickIndex, TaskIndex), bTickedA);
+					}
+					else
+					{
+						AITEST_TRUE(*FString::Printf(TEXT("Should tick Task A %d, %d"), TickIndex, TaskIndex), bTickedA);
+					}
+				}
+				const bool bStateSucceeded = Exec.Expect("Tree1StateRootTask_0", "StateCompleted");
+				const bool bLastTick = TickIndex == FStateTreeTasksCompletionStatus::MaxNumberOfTasksPerGroup - 1;
+				AITEST_EQUAL(TEXT("Tick should not complete the task."), bStateSucceeded, bLastTick);
+				Exec.LogClear();
+			}
+			{
+				Status = Exec.Tick(0.1f);
+				AITEST_EQUAL(TEXT("Tick should complete with Running"), Status, EStateTreeRunStatus::Running);
+				AITEST_TRUE(TEXT("In correct states"), Exec.ExpectInActiveStates("Tree1StateRoot", "Tree1StateA"));
+				bool bTicked = Exec.Expect("Tree1StateRootTask_0", TEXT("Tick"));
+				AITEST_TRUE(TEXT("Reset should allow new tick."), bTicked);
+			}
+
+			Exec.Stop();
+		}
+
+		// Test any
+		RootState->TasksCompletion = EStateTreeTaskCompletionType::Any;
+		{
+			FStateTreeCompilerLog Log;
+			FStateTreeCompiler Compiler(Log);
+			const bool bResult = Compiler.Compile(StateTree);
+			AITEST_TRUE("StateTree should get compiled", bResult);
+		}
+		{
+			EStateTreeRunStatus Status = EStateTreeRunStatus::Unset;
+			FStateTreeInstanceData InstanceData;
+			FTestStateTreeExecutionContext Exec(StateTree, StateTree, InstanceData);
+			const bool bInitSucceeded = Exec.IsValid();
+			AITEST_TRUE("StateTree should init", bInitSucceeded);
+
+			{
+				Status = Exec.Start();
+				AITEST_EQUAL("Start should complete with Running", Status, EStateTreeRunStatus::Running);
+				Exec.LogClear();
+			}
+
+			{
+				Status = Exec.Tick(0.1f);
+				AITEST_EQUAL(TEXT("Tick should complete with Running"), Status, EStateTreeRunStatus::Running);
+				AITEST_TRUE(TEXT("In correct states"), Exec.ExpectInActiveStates("Tree1StateRoot", "Tree1StateA"));
+				for (int32 TaskIndex = 0; TaskIndex < FStateTreeTasksCompletionStatus::MaxNumberOfTasksPerGroup; ++TaskIndex)
+				{
+					bool bTicked = Exec.Expect(*FString::Printf(TEXT("Tree1StateRootTask_%d"), TaskIndex), TEXT("Tick"));
+					AITEST_TRUE(*FString::Printf(TEXT("Should tick Task %d"), TaskIndex), bTicked);
+				}
+				const bool bStateSucceeded = Exec.Expect("Tree1StateRootTask_0", "StateCompleted");
+				AITEST_TRUE(TEXT("Tick should not complete the task."), bStateSucceeded);
+				Exec.LogClear();
+			}
+			{
+				Status = Exec.Tick(0.1f);
+				AITEST_EQUAL(TEXT("Tick should complete with Running"), Status, EStateTreeRunStatus::Running);
+				AITEST_TRUE(TEXT("In correct states"), Exec.ExpectInActiveStates("Tree1StateRoot", "Tree1StateA"));
+				bool bTicked = Exec.Expect("Tree1StateRootTask_0", TEXT("Tick"));
+				AITEST_TRUE(TEXT("Reset should allow new tick."), bTicked);
+			}
+
+			Exec.Stop();
+		}
+
+		return true;
+	}
+};
+IMPLEMENT_AI_INSTANT_TEST(FStateTreeTest_TasksCompletion_AllAny, "System.StateTree.TasksCompletion.AllAny");
+
+struct FStateTreeTest_TasksCompletion_FailureTasks : FStateTreeTestBase
+{
+	virtual bool InstantTest() override
+	{
+		constexpr int32 BadTask = 2;
+		// Main asset
+		UStateTree& StateTree = NewStateTree();
+		{
+			UStateTreeEditorData& EditorData = *Cast<UStateTreeEditorData>(StateTree.EditorData);
+
+			UStateTreeState& Root = EditorData.AddSubTree(FName(TEXT("Tree1StateRoot")));
+			Root.TasksCompletion = EStateTreeTaskCompletionType::All;
+			for (int32 TaskIndex = 0; TaskIndex < FStateTreeTasksCompletionStatus::MaxNumberOfTasksPerGroup; ++TaskIndex)
+			{
+				TStateTreeEditorNode<FTestTask_Stand>& Task = Root.AddTask<FTestTask_Stand>(FName(*FString::Printf(TEXT("Tree1StateRootTask_%d"), TaskIndex)));
+				Task.GetNode().TicksToCompletion = TaskIndex + 1;
+				if (TaskIndex == BadTask)
+				{
+					Task.GetNode().TickCompletionResult = EStateTreeRunStatus::Failed;
+				}
+			}
+
+			{
+				FStateTreeCompilerLog Log;
+				FStateTreeCompiler Compiler(Log);
+				const bool bResult = Compiler.Compile(StateTree);
+				AITEST_TRUE("StateTree should get compiled", bResult);
+			}
+		}
+
+		{
+			EStateTreeRunStatus Status = EStateTreeRunStatus::Unset;
+			FStateTreeInstanceData InstanceData;
+			FTestStateTreeExecutionContext Exec(StateTree, StateTree, InstanceData);
+			const bool bInitSucceeded = Exec.IsValid();
+			AITEST_TRUE("StateTree should init", bInitSucceeded);
+
+			{
+				Status = Exec.Start();
+				AITEST_EQUAL("Start should complete with Running", Status, EStateTreeRunStatus::Running);
+				Exec.LogClear();
+			}
+
+			for (int32 TickIndex = 0; TickIndex <= BadTask; ++TickIndex)
+			{
+				Status = Exec.Tick(0.1f);
+				AITEST_EQUAL(TEXT("Tick should complete with Running"), Status, EStateTreeRunStatus::Running);
+				AITEST_TRUE(TEXT("In correct states"), Exec.ExpectInActiveStates("Tree1StateRoot"));
+				const bool bLastTick = TickIndex == BadTask;
+				for (int32 TaskIndex = 0; TaskIndex < FStateTreeTasksCompletionStatus::MaxNumberOfTasksPerGroup; ++TaskIndex)
+				{
+					bool bTicked = Exec.Expect(*FString::Printf(TEXT("Tree1StateRootTask_%d"), TaskIndex), TEXT("Tick"));
+					if (TaskIndex < TickIndex || (bLastTick && TaskIndex > BadTask))
+					{
+						AITEST_FALSE(*FString::Printf(TEXT("Should not tick Task %d, %d"), TickIndex, TaskIndex), bTicked);
+					}
+					else
+					{
+						AITEST_TRUE(*FString::Printf(TEXT("Should tick Task %d, %d"), TickIndex, TaskIndex), bTicked);
+					}
+				}
+				const bool bStateSucceeded = Exec.Expect("Tree1StateRootTask_0", "StateCompleted");
+				AITEST_EQUAL(TEXT("Tick should not complete the task."), bStateSucceeded, bLastTick);
+				Exec.LogClear();
+			}
+			{
+				Status = Exec.Tick(0.1f);
+				AITEST_EQUAL(TEXT("Tick should complete with Running"), Status, EStateTreeRunStatus::Running);
+				AITEST_TRUE(TEXT("In correct states"), Exec.ExpectInActiveStates("Tree1StateRoot"));
+				bool bTicked = Exec.Expect("Tree1StateRootTask_0", TEXT("Tick"));
+				AITEST_TRUE(TEXT("Reset should allow new tick."), bTicked);
+			}
+
+			Exec.Stop();
+		}
+
+		return true;
+	}
+};
+IMPLEMENT_AI_INSTANT_TEST(FStateTreeTest_TasksCompletion_FailureTasks, "System.StateTree.TasksCompletion.Failure");
+
+// test set status with priority
+struct FStateTreeTest_TasksCompletion_Status : FStateTreeTestBase
+{
+	virtual bool InstantTest() override
+	{
+		FCompactStateTreeFrame Frame;
+		Frame.NumberOfTasksStatusMasks = 1;
+		FStateTreeTasksCompletionStatus Status = Frame;
+
+		constexpr int32 NumberOfTasks = 4;
+		FCompactStateTreeState State;
+		State.CompletionTasksControl = EStateTreeTaskCompletionType::All;
+		State.CompletionTasksMaskBitsOffset = 3;
+		State.CompletionTasksMaskBufferIndex = 0;
+		State.CompletionTasksMask = 1 << NumberOfTasks;
+		State.CompletionTasksMask -= 1;
+		State.CompletionTasksMask <<= State.CompletionTasksMaskBitsOffset;
+
+		UE::StateTree::FTasksCompletionStatus TestingStatus = Status.GetStatus(State);
+
+		auto TestEmpty = [this, &TestingStatus]()
+			{
+				AITEST_FALSE(TEXT("Empty is not completed."), TestingStatus.IsCompleted());
+				AITEST_FALSE(TEXT("Empty has not failure."), TestingStatus.HasAnyFailed());
+				AITEST_FALSE(TEXT("Empty is not completed."), TestingStatus.HasAnyCompleted());
+				AITEST_FALSE(TEXT("Empty is not completed."), TestingStatus.HasAllCompleted());
+				AITEST_EQUAL(TEXT("The completion status is running."), TestingStatus.GetCompletionStatus(), UE::StateTree::ETaskCompletionStatus::Running);
+				for (int32 Index = 0; Index < NumberOfTasks; ++Index)
+				{
+					AITEST_FALSE(TEXT("Empty is not completed."), TestingStatus.HasFailed(Index));
+					AITEST_TRUE(TEXT("Empty is running."), TestingStatus.IsRunning(Index));
+					AITEST_EQUAL(TEXT("Empty is running."), TestingStatus.GetStatus(Index), UE::StateTree::ETaskCompletionStatus::Running);
+				}
+				return true;
+			};
+
+		// Test new/empty completion status
+		{
+			if (!TestEmpty())
+			{
+				return false;
+			}
+		}
+
+		// Set tasks (1) to running. Does nothing.
+		{
+			TestingStatus.SetStatus(1, UE::StateTree::ETaskCompletionStatus::Running);
+			if (!TestEmpty())
+			{
+				return false;
+			}
+		}
+
+		auto Test1Expected = [this, &TestingStatus](UE::StateTree::ETaskCompletionStatus Expected)
+			{
+				AITEST_FALSE(TEXT("1 doesn't completed the state."), TestingStatus.IsCompleted());
+				AITEST_FALSE(TEXT("1 doesn't completed the failed."), TestingStatus.HasAnyFailed());
+				AITEST_TRUE(TEXT("1 does completed any state."), TestingStatus.HasAnyCompleted());
+				AITEST_FALSE(TEXT("1 doesn't completed all state."), TestingStatus.HasAllCompleted());
+				AITEST_EQUAL(TEXT("States running."), TestingStatus.GetCompletionStatus(), UE::StateTree::ETaskCompletionStatus::Running);
+				for (int32 Index = 0; Index < NumberOfTasks; ++Index)
+				{
+					AITEST_FALSE(TEXT("1 is not completed."), TestingStatus.HasFailed(Index));
+					if (1 == Index)
+					{
+						AITEST_FALSE(TEXT("1 is Stopped."), TestingStatus.IsRunning(Index));
+						AITEST_EQUAL(TEXT("1 is Stopped."), TestingStatus.GetStatus(Index), Expected);
+					}
+					else
+					{
+						AITEST_TRUE(TEXT("1 is Stopped."), TestingStatus.IsRunning(Index));
+						AITEST_EQUAL(TEXT("1 is Stopped."), TestingStatus.GetStatus(Index), UE::StateTree::ETaskCompletionStatus::Running);
+					}
+				}
+				return true;
+			};
+		// Set tasks (1) to stop.
+		{
+			TestingStatus.SetStatus(1, UE::StateTree::ETaskCompletionStatus::Stopped);
+			if (!Test1Expected(UE::StateTree::ETaskCompletionStatus::Stopped))
+			{
+				return false;
+			}
+		}
+		// Set tasks (1) to running. Does nothing.
+		{
+			ETaskCompletionStatus NewStatus = TestingStatus.SetStatusWithPriority(1, UE::StateTree::ETaskCompletionStatus::Running);
+			AITEST_EQUAL(TEXT("1 make the state Stopped."), NewStatus, UE::StateTree::ETaskCompletionStatus::Stopped);
+			if (!Test1Expected(UE::StateTree::ETaskCompletionStatus::Stopped))
+			{
+				return false;
+			}
+		}
+		// Set tasks (1) to Succeeded. Does nothing.
+		{
+			ETaskCompletionStatus NewStatus = TestingStatus.SetStatusWithPriority(1, UE::StateTree::ETaskCompletionStatus::Succeeded);
+			AITEST_EQUAL(TEXT("1 make the state Succeed."), NewStatus, UE::StateTree::ETaskCompletionStatus::Succeeded);
+			if (!Test1Expected(UE::StateTree::ETaskCompletionStatus::Succeeded))
+			{
+				return false;
+			}
+		}
+		// Set tasks (1) to stop. Does nothing.
+		{
+			ETaskCompletionStatus NewStatus = TestingStatus.SetStatusWithPriority(1, UE::StateTree::ETaskCompletionStatus::Stopped);
+			AITEST_EQUAL(TEXT("1 make the state Succeed."), NewStatus, UE::StateTree::ETaskCompletionStatus::Succeeded);
+			if (!Test1Expected(UE::StateTree::ETaskCompletionStatus::Succeeded))
+			{
+				return false;
+			}
+		}
+		// Set tasks (1) to running. Does nothing.
+		{
+			ETaskCompletionStatus NewStatus = TestingStatus.SetStatusWithPriority(1, UE::StateTree::ETaskCompletionStatus::Stopped);
+			AITEST_EQUAL(TEXT("1 make the state Succeed."), NewStatus, UE::StateTree::ETaskCompletionStatus::Succeeded);
+			if (!Test1Expected(UE::StateTree::ETaskCompletionStatus::Succeeded))
+			{
+				return false;
+			}
+		}
+		// Set tasks (1) to fail. Does nothing.
+		{
+			ETaskCompletionStatus NewStatus = TestingStatus.SetStatusWithPriority(1, UE::StateTree::ETaskCompletionStatus::Failed);
+			AITEST_EQUAL(TEXT("1 make the state Succeed."), NewStatus, UE::StateTree::ETaskCompletionStatus::Failed);
+			AITEST_TRUE(TEXT("1 completes the state."), TestingStatus.IsCompleted());
+			AITEST_TRUE(TEXT("1 completes the failed."), TestingStatus.HasAnyFailed());
+			AITEST_TRUE(TEXT("1 completes any state."), TestingStatus.HasAnyCompleted());
+			AITEST_TRUE(TEXT("1 completes all state."), TestingStatus.HasAllCompleted());
+			AITEST_EQUAL(TEXT("States failed."), TestingStatus.GetCompletionStatus(), UE::StateTree::ETaskCompletionStatus::Failed);
+			for (int32 Index = 0; Index < NumberOfTasks; ++Index)
+			{
+				if (1 == Index)
+				{
+					AITEST_TRUE(TEXT("1 is not completed."), TestingStatus.HasFailed(Index));
+					AITEST_FALSE(TEXT("1 is failed."), TestingStatus.IsRunning(Index));
+					AITEST_EQUAL(TEXT("1 is failed."), TestingStatus.GetStatus(Index), UE::StateTree::ETaskCompletionStatus::Failed);
+				}
+				else
+				{
+					AITEST_FALSE(TEXT("1 is not completed."), TestingStatus.HasFailed(Index));
+					AITEST_TRUE(TEXT("1 is running."), TestingStatus.IsRunning(Index));
+					AITEST_EQUAL(TEXT("1 is running."), TestingStatus.GetStatus(Index), UE::StateTree::ETaskCompletionStatus::Running);
+				}
+			}
+		}
+		// Set completion status
+		{
+			TestingStatus.SetCompletionStatus(UE::StateTree::ETaskCompletionStatus::Stopped);
+			AITEST_TRUE(TEXT("All complete the state."), TestingStatus.IsCompleted());
+			AITEST_FALSE(TEXT("All complete the failed."), TestingStatus.HasAnyFailed());
+			AITEST_TRUE(TEXT("All complete any state."), TestingStatus.HasAnyCompleted());
+			AITEST_TRUE(TEXT("All complete all state."), TestingStatus.HasAllCompleted());
+			AITEST_EQUAL(TEXT("States stopped."), TestingStatus.GetCompletionStatus(), UE::StateTree::ETaskCompletionStatus::Stopped);
+			for (int32 Index = 0; Index < NumberOfTasks; ++Index)
+			{
+				AITEST_FALSE(TEXT("1 is not completed."), TestingStatus.HasFailed(Index));
+				AITEST_FALSE(TEXT("1 is stopped."), TestingStatus.IsRunning(Index));
+				AITEST_EQUAL(TEXT("1 is stopped."), TestingStatus.GetStatus(Index), UE::StateTree::ETaskCompletionStatus::Stopped);
+			}
+		}
+		// Test GetCompletionStatus
+		{
+			TestingStatus.SetCompletionStatus(UE::StateTree::ETaskCompletionStatus::Running);
+			AITEST_EQUAL(TEXT("States stopped."), TestingStatus.GetCompletionStatus(), UE::StateTree::ETaskCompletionStatus::Running);
+		}
+		// Test GetCompletionStatus
+		{
+			TestingStatus.SetCompletionStatus(UE::StateTree::ETaskCompletionStatus::Succeeded);
+			AITEST_EQUAL(TEXT("States succeeded."), TestingStatus.GetCompletionStatus(), UE::StateTree::ETaskCompletionStatus::Succeeded);
+		}
+		// Test GetCompletionStatus
+		{
+			TestingStatus.SetCompletionStatus(UE::StateTree::ETaskCompletionStatus::Failed);
+			AITEST_EQUAL(TEXT("States failed."), TestingStatus.GetCompletionStatus(), UE::StateTree::ETaskCompletionStatus::Failed);
+		}
+
+		return true;
+	}
+};
+IMPLEMENT_AI_INSTANT_TEST(FStateTreeTest_TasksCompletion_Status, "System.StateTree.TasksCompletion.Status");
+
+} // namespace UE::StateTree::Tests
+
+UE_ENABLE_OPTIMIZATION_SHIP
+
+#undef LOCTEXT_NAMESPACE
diff --git a/Engine/Plugins/Runtime/StateTree/Source/StateTreeTestSuite/Private/StateTreeWeakContextTest.cpp b/Engine/Plugins/Runtime/StateTree/Source/StateTreeTestSuite/Private/StateTreeWeakContextTest.cpp
--- a/Engine/Plugins/Runtime/StateTree/Source/StateTreeTestSuite/Private/StateTreeWeakContextTest.cpp
+++ b/Engine/Plugins/Runtime/StateTree/Source/StateTreeTestSuite/Private/StateTreeWeakContextTest.cpp
@@ -131,14 +131,14 @@ struct FStateTreeTest_WeakContext_FinishTask : FStateTreeTestBase
 				TStateTreeEditorNode<FTestTask_PrintValue>& Tree1State2Task1 = State2.AddTask<FTestTask_PrintValue>(FName(TEXT("Tree1State2Task1")));
 				Tree1State2Task1.GetNode().CustomEnterStateFunc = [&WeakContext](FStateTreeExecutionContext& Context, const FTestTask_PrintValue* Task)
 					{
-						WeakContext.ContextTree1State1Task1 = Context.MakeWeakExecutionContext();
-						WeakContext.TaskTree1State1Task1 = Context.MakeWeakTaskRef(*Task);
+						WeakContext.ContextTree1State2Task1 = Context.MakeWeakExecutionContext();
+						WeakContext.TaskTree1State2Task1 = Context.MakeWeakTaskRef(*Task);
 					};
 				TStateTreeEditorNode<FTestTask_PrintValue>& Tree1State2Task2 = State2.AddTask<FTestTask_PrintValue>(FName(TEXT("Tree1State2Task2")));
 				Tree1State2Task2.GetNode().CustomEnterStateFunc = [&WeakContext](FStateTreeExecutionContext& Context, const FTestTask_PrintValue* Task)
 					{
-						WeakContext.ContextTree1State1Task2 = Context.MakeWeakExecutionContext();
-						WeakContext.TaskTree1State1Task2 = Context.MakeWeakTaskRef(*Task);
+						WeakContext.ContextTree1State2Task2 = Context.MakeWeakExecutionContext();
+						WeakContext.TaskTree1State2Task2 = Context.MakeWeakTaskRef(*Task);
 					};
 			}
 		}
@@ -214,15 +214,15 @@ struct FStateTreeTest_WeakContext_FinishTask : FStateTreeTestBase
 					.Then("Tree1State1Task2", TEXT("EnterState0"))
 					.Then("Tree1State2Task1", TEXT("EnterState0"))
 					.Then("Tree1State2Task2", TEXT("EnterState0")));
-				AITEST_FALSE(TEXT("FinishTask should not Tick"), Exec.Expect("Tree1RootTask1", TEXT("Tick0")));
+				AITEST_TRUE(TEXT("FinishTask should not Tick"), Exec.Expect("Tree1RootTask1", TEXT("Tick0")));
 				AITEST_FALSE(TEXT("FinishTask should not Tick"), Exec.Expect("Tree1RootTask2", TEXT("Tick0")));
-				AITEST_FALSE(TEXT("FinishTask should not Tick"), Exec.Expect("Tree1State1Task1", TEXT("Tick0")));
-				AITEST_FALSE(TEXT("FinishTask should not Tick"), Exec.Expect("Tree1State1Task2", TEXT("Tick0")));
-				AITEST_FALSE(TEXT("FinishTask should not Tick"), Exec.Expect("Tree1State2Task1", TEXT("Tick0")));
-				AITEST_FALSE(TEXT("FinishTask should not Tick"), Exec.Expect("Tree1State2Task2", TEXT("Tick0")));
+				AITEST_TRUE(TEXT("FinishTask should not Tick"), Exec.Expect("Tree1State1Task1", TEXT("Tick0")));
+				AITEST_TRUE(TEXT("FinishTask should not Tick"), Exec.Expect("Tree1State1Task2", TEXT("Tick0")));
+				AITEST_TRUE(TEXT("FinishTask should not Tick"), Exec.Expect("Tree1State2Task1", TEXT("Tick0")));
+				AITEST_TRUE(TEXT("FinishTask should not Tick"), Exec.Expect("Tree1State2Task2", TEXT("Tick0")));
 				Exec.LogClear();
 			}
-			// Test Weak::FinishTask of a task after the current task, during the tick tasks.
+			// Test Weak::FinishTask of a task after the current task, during the tick tasks. Complete with Failure. The following task should not run.
 			{
 				WeakContext.bWhen_TaskTree1State1Task1_Finish_TaskTree1State1Task2 = true;
 				FTestStateTreeExecutionContext Exec(StateTree1, StateTree1, InstanceData);
@@ -284,6 +284,7 @@ struct FStateTreeTest_WeakContext_FinishTask : FStateTreeTestBase
 					.Then("Tree1State1Task2", TEXT("ExitState0"))
 					.Then("Tree1GlobalTask2", TEXT("ExitState0"))
 					.Then("Tree1GlobalTask1", TEXT("ExitState0")));
+				AITEST_FALSE(TEXT("FinishTask should not Tick"), Exec.Expect("Tree1State2Task1", TEXT("Tick0")));
 				Exec.LogClear();
 				WeakContext.bWhen_Tree1GlobalTask1_Finish_Tree1GlobalTask2 = false;
 			}
@@ -303,8 +304,7 @@ struct FStateTreeTest_WeakContext_FinishTask : FStateTreeTestBase
 
 				AITEST_TRUE(TEXT("Tick should Tick"), Exec.Expect("Tree1GlobalTask1", TEXT("Tick0"))
 					.Then("Tree1State1Task2", TEXT("Tick0")));
-				AITEST_FALSE(TEXT("FinishTask should not Tick"), Exec.Expect("Tree1State2Task1", TEXT("Tick0")));
-				AITEST_TRUE(TEXT("Tick should Tick"), Exec.Expect("Tree1State2Task2", TEXT("ExitState0"))
+				AITEST_TRUE(TEXT("Tick should exit"), Exec.Expect("Tree1State2Task2", TEXT("ExitState0"))
 					.Then("Tree1GlobalTask1", TEXT("ExitState0")));
 				Exec.LogClear();
 				WeakContext.bWhen_TaskTree1State1Task2_Finish_Tree1GlobalTask2 = false;
