Fix warning macro appears a second time without the matching macro. Introduced with CL41245712

diff --git a/Engine/Plugins/Runtime/StateTree/Source/StateTreeModule/Private/StateTreeExecutionContext.cpp b/Engine/Plugins/Runtime/StateTree/Source/StateTreeModule/Private/StateTreeExecutionContext.cpp
--- a/Engine/Plugins/Runtime/StateTree/Source/StateTreeModule/Private/StateTreeExecutionContext.cpp
+++ b/Engine/Plugins/Runtime/StateTree/Source/StateTreeModule/Private/StateTreeExecutionContext.cpp
@@ -1791,7 +1791,7 @@ void FStateTreeExecutionContext::UpdateCompletedStateList()
 void FStateTreeExecutionContext::MarkStateCompleted(UE::StateTree::FFinishedTask& NewFinishedTask)
 {
 }
-PRAGMA_DISABLE_DEPRECATION_WARNINGS
+PRAGMA_ENABLE_DEPRECATION_WARNINGS
 
 void FStateTreeExecutionContext::UpdateInstanceData(TConstArrayView<FStateTreeExecutionFrame> CurrentActiveFrames, TArrayView<FStateTreeExecutionFrame> NextActiveFrames)
 {
diff --git a/Engine/Plugins/Runtime/StateTree/Source/StateTreeModule/Public/StateTreeExecutionContext.h b/Engine/Plugins/Runtime/StateTree/Source/StateTreeModule/Public/StateTreeExecutionContext.h
--- a/Engine/Plugins/Runtime/StateTree/Source/StateTreeModule/Public/StateTreeExecutionContext.h
+++ b/Engine/Plugins/Runtime/StateTree/Source/StateTreeModule/Public/StateTreeExecutionContext.h
@@ -785,12 +785,6 @@ protected:
 	 */
 	EStateTreeRunStatus TickEvaluatorsAndGlobalTasksForFrame(const float DeltaTime, const bool bTickGlobalTasks, const int32 FrameIndex, const FStateTreeExecutionFrame* CurrentParentFrame, const TNotNull<FStateTreeExecutionFrame*> CurrentFrame);
 
-	UE_DEPRECATED(5.6, "Use the non const version of TickEvaluatorsAndGlobalTasksForFrame")
-	EStateTreeRunStatus TickEvaluatorsAndGlobalTasksForFrame(const float DeltaTime, const bool bTickGlobalTasks, const int32 FrameIndex, const FStateTreeExecutionFrame* CurrentParentFrame, const FStateTreeExecutionFrame* CurrentFrame)
-	{
-		return TickEvaluatorsAndGlobalTasksForFrame(DeltaTime, bTickGlobalTasks, FrameIndex, CurrentParentFrame, TNotNull<FStateTreeExecutionFrame*>(const_cast<FStateTreeExecutionFrame*>(CurrentFrame)));
-	}
-
 	/**
 	 * Starts evaluators and global tasks.
 	 * @return run status returned by the global tasks.
