Fix identical sub-expressions warning. Introduced with CL41245712

diff --git a/Engine/Plugins/Runtime/StateTree/Source/StateTreeModule/Private/StateTreeTasksStatus.cpp b/Engine/Plugins/Runtime/StateTree/Source/StateTreeModule/Private/StateTreeTasksStatus.cpp
--- a/Engine/Plugins/Runtime/StateTree/Source/StateTreeModule/Private/StateTreeTasksStatus.cpp
+++ b/Engine/Plugins/Runtime/StateTree/Source/StateTreeModule/Private/StateTreeTasksStatus.cpp
@@ -186,7 +186,7 @@ bool FStateTreeTasksCompletionStatus::Serialize(FArchive& Ar)
 	{
 		int8 NewBufferNum = 0;
 		Ar << NewBufferNum;
-		if (NewBufferNum != NewBufferNum)
+		if (NewBufferNum != BufferNum)
 		{
 			this->~FStateTreeTasksCompletionStatus();
 			BufferNum = NewBufferNum;
