StateTree: Reorder the EStateTreeRunStatus to be aligned with ETaskCompletionStatus.
Fix EStateTreeFinishTaskType::Failed start at value 0 to make UObject that contains the enum valid. They are memset to 0, so the 0 value needs to be valid.

diff --git a/Engine/Plugins/Runtime/StateTree/Source/StateTreeModule/Private/StateTreeExecutionContext.cpp b/Engine/Plugins/Runtime/StateTree/Source/StateTreeModule/Private/StateTreeExecutionContext.cpp
--- a/Engine/Plugins/Runtime/StateTree/Source/StateTreeModule/Private/StateTreeExecutionContext.cpp
+++ b/Engine/Plugins/Runtime/StateTree/Source/StateTreeModule/Private/StateTreeExecutionContext.cpp
@@ -62,7 +62,6 @@ namespace UE::StateTree::ExecutionContext::Private
 				continue;
 			}
 
-
 			const int32 TransitionEnd = State->TransitionsBegin + State->TransitionsNum;
 			for (int32 TransitionIndex = State->TransitionsBegin; TransitionIndex <TransitionEnd; ++TransitionIndex)
 			{
@@ -113,66 +112,50 @@ namespace UE::StateTree::ExecutionContext::Private
 	EStateTreeRunStatus GetPriorityRunStatus(EStateTreeRunStatus A, EStateTreeRunStatus B)
 	{
 		static_assert((int32)EStateTreeRunStatus::Running == 0);
-		static_assert((int32)EStateTreeRunStatus::Failed == 1);
+		static_assert((int32)EStateTreeRunStatus::Stopped == 1);
 		static_assert((int32)EStateTreeRunStatus::Succeeded == 2);
-		static_assert((int32)EStateTreeRunStatus::Stopped == 3);
+		static_assert((int32)EStateTreeRunStatus::Failed == 3);
 		static_assert((int32)EStateTreeRunStatus::Unset == 4);
 		static_assert(NumEStateTreeRunStatus() == 5, "The number of entries in EStateTreeRunStatus changed. Update GetPriorityRunStatus.");
 
-		static constexpr int32 PriorityMatrix[] = { 1, 4, 3, 2, 0 };
+		static constexpr int32 PriorityMatrix[] = { 1, 2, 3, 4, 0 };
 		return PriorityMatrix[(uint8)A] > PriorityMatrix[(uint8)B] ? A : B;
 	}
 
 	UE::StateTree::ETaskCompletionStatus CastToTaskStatus(EStateTreeFinishTaskType FinishTask)
 	{
 		static_assert(NumEStateTreeFinishTaskType() == 2, "The number of entries in EStateTreeFinishTaskType changed. Update CastToTaskStatus.");
+
 		return FinishTask == EStateTreeFinishTaskType::Succeeded ? UE::StateTree::ETaskCompletionStatus::Succeeded : UE::StateTree::ETaskCompletionStatus::Failed;
 	}
 
 	EStateTreeRunStatus CastToRunStatus(EStateTreeFinishTaskType FinishTask)
 	{
-		static_assert((int32)EStateTreeRunStatus::Succeeded == (int32)EStateTreeFinishTaskType::Succeeded, "Succeeded does not match in both types.");
-		static_assert((int32)EStateTreeRunStatus::Failed == (int32)EStateTreeFinishTaskType::Failed, "Failed does not match in both types.");
 		static_assert(NumEStateTreeFinishTaskType() == 2, "The number of entries in EStateTreeFinishTaskType changed. Update CastToRunStatus.");
-		return (EStateTreeRunStatus)FinishTask;
+
+		return FinishTask == EStateTreeFinishTaskType::Succeeded ? EStateTreeRunStatus::Succeeded : EStateTreeRunStatus::Failed;
 	}
 
 	UE::StateTree::ETaskCompletionStatus CastToTaskStatus(EStateTreeRunStatus InStatus)
 	{
-		static_assert((int32)EStateTreeRunStatus::Running == 0);
-		static_assert((int32)EStateTreeRunStatus::Failed == 1);
-		static_assert((int32)EStateTreeRunStatus::Succeeded == 2);
-		static_assert((int32)EStateTreeRunStatus::Stopped == 3);
-		static_assert((int32)EStateTreeRunStatus::Unset == 4);
+		static_assert((int32)EStateTreeRunStatus::Running == (int32)UE::StateTree::ETaskCompletionStatus::Running);
+		static_assert((int32)EStateTreeRunStatus::Stopped == (int32)UE::StateTree::ETaskCompletionStatus::Stopped);
+		static_assert((int32)EStateTreeRunStatus::Succeeded == (int32)UE::StateTree::ETaskCompletionStatus::Succeeded);
+		static_assert((int32)EStateTreeRunStatus::Failed == (int32)UE::StateTree::ETaskCompletionStatus::Failed);
 		static_assert(NumEStateTreeRunStatus() == 5, "The number of entries in EStateTreeRunStatus changed. Update CastToTaskStatus.");
 
-		static constexpr UE::StateTree::ETaskCompletionStatus ConversionMatrix[] =
-			{
-				UE::StateTree::ETaskCompletionStatus::Running,
-				UE::StateTree::ETaskCompletionStatus::Failed,
-				UE::StateTree::ETaskCompletionStatus::Succeeded,
-				UE::StateTree::ETaskCompletionStatus::Stopped,
-				UE::StateTree::ETaskCompletionStatus::Running,
-			};
-		return ConversionMatrix[(int32)InStatus];
+		return InStatus != EStateTreeRunStatus::Unset ? (UE::StateTree::ETaskCompletionStatus)InStatus : UE::StateTree::ETaskCompletionStatus::Running;
 	}
-	
+
 	EStateTreeRunStatus CastToRunStatus(UE::StateTree::ETaskCompletionStatus InStatus)
 	{
-		//@todo reorder EStateTreeRunStatus enum to be 1 to 1 with UE::StateTree::ETaskCompletionStatus
-		static_assert((int32)UE::StateTree::ETaskCompletionStatus::Running == 0);
-		static_assert((int32)UE::StateTree::ETaskCompletionStatus::Stopped == 1);
-		static_assert((int32)UE::StateTree::ETaskCompletionStatus::Succeeded == 2);
-		static_assert((int32)UE::StateTree::ETaskCompletionStatus::Failed == 3);
+		static_assert((int32)EStateTreeRunStatus::Running == (int32)UE::StateTree::ETaskCompletionStatus::Running);
+		static_assert((int32)EStateTreeRunStatus::Stopped == (int32)UE::StateTree::ETaskCompletionStatus::Stopped);
+		static_assert((int32)EStateTreeRunStatus::Succeeded == (int32)UE::StateTree::ETaskCompletionStatus::Succeeded);
+		static_assert((int32)EStateTreeRunStatus::Failed == (int32)UE::StateTree::ETaskCompletionStatus::Failed);
 		static_assert(UE::StateTree::NumberOfTaskStatus == 4, "The number of entries in EStateTreeRunStatus changed. Update CastToRunStatus.");
-		static constexpr EStateTreeRunStatus ConversionMatrix[] =
-		{
-			EStateTreeRunStatus::Running,
-			EStateTreeRunStatus::Stopped,
-			EStateTreeRunStatus::Succeeded,
-			EStateTreeRunStatus::Failed,
-		};
-		return ConversionMatrix[(int32)InStatus];
+
+		return (EStateTreeRunStatus)InStatus;
 	}
 }
 
diff --git a/Engine/Plugins/Runtime/StateTree/Source/StateTreeModule/Public/StateTreeExecutionTypes.h b/Engine/Plugins/Runtime/StateTree/Source/StateTreeModule/Public/StateTreeExecutionTypes.h
--- a/Engine/Plugins/Runtime/StateTree/Source/StateTreeModule/Public/StateTreeExecutionTypes.h
+++ b/Engine/Plugins/Runtime/StateTree/Source/StateTreeModule/Public/StateTreeExecutionTypes.h
@@ -52,14 +52,14 @@ enum class EStateTreeRunStatus : uint8
 	/** Tree is still running. */
 	Running,
 	
-	/** Tree execution has stopped on failure. */
-	Failed,
+	/** The State Tree was requested to stop without a particular success or failure state. */
+	Stopped,
 	
 	/** Tree execution has stopped on success. */
 	Succeeded,
 
-	/** The State Tree was requested to stop without a particular success or failure state. */
-	Stopped,
+	/** Tree execution has stopped on failure. */
+	Failed,
 
 	/** Status not set. */
 	Unset,
@@ -71,10 +71,10 @@ UENUM()
 enum class EStateTreeFinishTaskType : uint8
 {	
 	/** The task execution failed. */
-	Failed = (uint8)EStateTreeRunStatus::Failed,
+	Failed,
 	
 	/** The task execution succeed. */
-	Succeeded = (uint8)EStateTreeRunStatus::Succeeded,
+	Succeeded
 };
 
 
