StateTree: use the correct buffer when accessing the temporary global buffer.

diff --git a/Engine/Plugins/Runtime/StateTree/Source/StateTreeEditorModule/Private/StateTreeState.cpp b/Engine/Plugins/Runtime/StateTree/Source/StateTreeEditorModule/Private/StateTreeState.cpp
--- a/Engine/Plugins/Runtime/StateTree/Source/StateTreeEditorModule/Private/StateTreeState.cpp
+++ b/Engine/Plugins/Runtime/StateTree/Source/StateTreeEditorModule/Private/StateTreeState.cpp
@@ -182,7 +182,6 @@ void UStateTreeState::PostEditChangeChainProperty(FPropertyChangedChainEvent& Pr
 			// Linked parameter layout is fixed, and copied from the linked target state.
 			Parameters.bFixedLayout = true;
 			UpdateParametersFromLinkedSubtree();
-			SelectionBehavior = EStateTreeStateSelectionBehavior::TrySelectChildrenInOrder;
 		}
 		else
 		{
@@ -515,3 +514,27 @@ TSubclassOf<UStateTreeSchema> UStateTreeState::GetSchema() const
 	}
 	return nullptr;
 }
+
+void UStateTreeState::SetLinkedState(FStateTreeStateLink InStateLink)
+{
+	check(Type == EStateTreeStateType::Linked);
+	LinkedSubtree = InStateLink;
+
+	Tasks.Reset();
+	LinkedAsset = nullptr;
+	Parameters.bFixedLayout = true;
+	UpdateParametersFromLinkedSubtree();
+	SelectionBehavior = EStateTreeStateSelectionBehavior::TryEnterState;
+}
+
+void UStateTreeState::SetLinkedStateAsset(UStateTree* InLinkedAsset)
+{
+	check(Type == EStateTreeStateType::LinkedAsset);
+	LinkedAsset = InLinkedAsset;
+
+	Tasks.Reset();
+	LinkedSubtree = FStateTreeStateLink();
+	Parameters.bFixedLayout = true;
+	UpdateParametersFromLinkedSubtree();
+	SelectionBehavior = EStateTreeStateSelectionBehavior::TryEnterState;
+}
\ No newline at end of file
diff --git a/Engine/Plugins/Runtime/StateTree/Source/StateTreeEditorModule/Public/StateTreeState.h b/Engine/Plugins/Runtime/StateTree/Source/StateTreeEditorModule/Public/StateTreeState.h
--- a/Engine/Plugins/Runtime/StateTree/Source/StateTreeEditorModule/Public/StateTreeState.h
+++ b/Engine/Plugins/Runtime/StateTree/Source/StateTreeEditorModule/Public/StateTreeState.h
@@ -242,7 +242,7 @@ public:
 	/** @returns Default parameters from linked state or asset). */
 	const FInstancedPropertyBag* GetDefaultParameters() const;
 	
-	// StateTree Builder API
+	//~ StateTree Builder API
 	/** @return state link to this state. */
 	FStateTreeStateLink GetLinkToState() const;
 	
@@ -294,13 +294,11 @@ public:
 		return static_cast<TStateTreeEditorNode<T>&>(TaskItem);
 	}
 
+	/** Sets linked state and updates parameters to match the linked state. */
+	void SetLinkedState(FStateTreeStateLink InStateLink);
+
 	/** Sets linked asset and updates parameters to match the linked asset. */
-	void SetLinkedStateAsset(UStateTree* InLinkedAsset)
-	{
-		check(Type == EStateTreeStateType::LinkedAsset);
-		LinkedAsset = InLinkedAsset;
-		UpdateParametersFromLinkedSubtree();
-	}
+	void SetLinkedStateAsset(UStateTree* InLinkedAsset);
 	
 	/**
 	 * Adds Transition.
@@ -325,17 +323,13 @@ public:
 		return FGuid::Combine(ID, FGuid::NewDeterministicGuid(TEXT("Event")));
 	}
 
-	// ~StateTree Builder API
-
-	// IStateTreeSchemaProvider API
-	/**
-	* @return Class of schema used by the state tree containing this state.
-	*/
+	//~ IStateTreeSchemaProvider API
+	/**  @return Class of schema used by the state tree containing this state. */
 	virtual TSubclassOf<UStateTreeSchema> GetSchema() const override;
 
-	// ~IStateTreeSchemaProvider API
+	//~IStateTreeSchemaProvider API
 
-	// Note: these properties are customized out in FStateTreeStateDetails, adding a new property might require to adjust the customization.
+	//~ Note: these properties are customized out in FStateTreeStateDetails, adding a new property might require to adjust the customization.
 	
 	/** Display name of the State */
 	UPROPERTY(EditDefaultsOnly, Category = "State")
diff --git a/Engine/Plugins/Runtime/StateTree/Source/StateTreeModule/Private/StateTreeExecutionContext.cpp b/Engine/Plugins/Runtime/StateTree/Source/StateTreeModule/Private/StateTreeExecutionContext.cpp
--- a/Engine/Plugins/Runtime/StateTree/Source/StateTreeModule/Private/StateTreeExecutionContext.cpp
+++ b/Engine/Plugins/Runtime/StateTree/Source/StateTreeModule/Private/StateTreeExecutionContext.cpp
@@ -5470,7 +5470,7 @@ bool FStateTreeExecutionContext::SelectStateInternal(
 				NewFrame.StateParameterDataHandle = RootState.ParameterDataHandle;
 
 				// Start global tasks and evaluators temporarily, so that their data is available already during select.
-				if (StartTemporaryEvaluatorsAndGlobalTasks(nullptr, NewFrame) != EStateTreeRunStatus::Running)
+				if (StartTemporaryEvaluatorsAndGlobalTasks(&CurrentFrame, NewFrame) != EStateTreeRunStatus::Running)
 				{
 					STATETREE_LOG(VeryVerbose, TEXT("%hs: Cannot select state '%s' because cannot start nested tree's '%s' global tasks and evaluators.  '%s' using StateTree '%s'."),
 						__FUNCTION__, *GetSafeStateName(CurrentFrame, NextStateHandle), *GetFullNameSafe(NewFrame.StateTree), *GetNameSafe(&Owner), *GetFullNameSafe(CurrentFrame.StateTree));
@@ -5495,7 +5495,7 @@ bool FStateTreeExecutionContext::SelectStateInternal(
 
 			if (bStartedTemporaryEvaluatorsAndGlobalTasks)
 			{
-				StopTemporaryEvaluatorsAndGlobalTasks(nullptr, NewFrame);
+				StopTemporaryEvaluatorsAndGlobalTasks(&CurrentFrame, NewFrame);
 				GetExecState().DelegateActiveListeners.RemoveAll(NewFrame.FrameID);
 			}
 				
diff --git a/Engine/Plugins/Runtime/StateTree/Source/StateTreeTestSuite/Private/StateTreeLinkedStateTest.cpp b/Engine/Plugins/Runtime/StateTree/Source/StateTreeTestSuite/Private/StateTreeLinkedStateTest.cpp
--- a/Engine/Plugins/Runtime/StateTree/Source/StateTreeTestSuite/Private/StateTreeLinkedStateTest.cpp
+++ b/Engine/Plugins/Runtime/StateTree/Source/StateTreeTestSuite/Private/StateTreeLinkedStateTest.cpp
@@ -380,8 +380,7 @@ struct FStateTreeTest_MultipleSameLinkedAsset : FStateTreeTestBase
 				AITEST_TRUE("4th Tick should enter Tree2StateRootTaskA", LogOrder);
 				LogOrder = LogOrder.Then("Tree2StateATaskA", TEXT("EnterState111"));
 				AITEST_TRUE("4th Tick should enter Tree2StateATaskA", LogOrder);
-				//@todo: bug: Binding are not executed
-				//AITEST_TRUE("4th Tick should tick tasks", Exec.Expect("Tree2GlobalTaskA", TEXT("EnterState111")));
+				AITEST_TRUE("4th Tick should tick tasks", Exec.Expect("Tree2GlobalTaskA", TEXT("EnterState111")));
 				Exec.LogClear();
 			}
 			{
@@ -1001,6 +1000,301 @@ struct FStateTreeTest_LinkedAssetTransitionSameTick : FStateTreeTestBase
 };
 IMPLEMENT_AI_INSTANT_TEST(FStateTreeTest_LinkedAssetTransitionSameTick, "System.StateTree.LinkedAsset.TransitionSameTick");
 
+struct FStateTreeTest_Linked_GlobalParameter : FStateTreeTestBase
+{
+	virtual bool InstantTest() override
+	{
+		//Tree 1
+		//  Global task and parameter
+		//	Root
+		//		StateLinkedTree1 (Tree2) -> Next
+		//		SubTree2
+		//	SubTree
+		//		State3 (with transition OnTick to succeeded)
+		//Tree 2
+		//  Global task and parameter
+		//	Root
+		//		State1 (with transition OnTick to succeeded)
+
+		auto AddInt = [](FInstancedPropertyBag& PropertyBag, FName VarName)
+			{
+				PropertyBag.AddProperty(VarName, EPropertyBagPropertyType::Int32);
+				PropertyBag.SetValueInt32(VarName, -99);
+				return PropertyBag.FindPropertyDescByName(VarName)->ID;
+			};
+		auto AddDouble = [](FInstancedPropertyBag& PropertyBag, FName VarName)
+			{
+				PropertyBag.AddProperty(VarName, EPropertyBagPropertyType::Double);
+				PropertyBag.SetValueDouble(VarName, -99.0);
+				return PropertyBag.FindPropertyDescByName(VarName)->ID;
+			};
+
+		FGuid Tree2GlobalParameter_ValueID_Int;
+
+		// Tree 2
+		UStateTree& StateTree2 = NewStateTree();
+		{
+			UStateTreeEditorData& EditorData2 = *Cast<UStateTreeEditorData>(StateTree2.EditorData);
+			
+			// note double before int
+			AddDouble(GetRootPropertyBag(EditorData2), "Tree2GlobalDouble");
+			Tree2GlobalParameter_ValueID_Int = AddInt(GetRootPropertyBag(EditorData2), "Tree2GlobalInt");
+
+			// Global tasks
+			TStateTreeEditorNode<FTestTask_PrintValue>& GlobalTask1 = EditorData2.AddGlobalTask<FTestTask_PrintValue>("Tree2GlobalTask1");
+			{
+				GlobalTask1.GetInstanceData().Value = -1;
+				EditorData2.AddPropertyBinding(FPropertyBindingPath(EditorData2.GetRootParametersGuid(), "Tree2GlobalInt"), FPropertyBindingPath(GlobalTask1.ID, "Value"));
+			}
+			{
+				TStateTreeEditorNode<FTestTask_PrintValue>& GlobalTask2 = EditorData2.AddGlobalTask<FTestTask_PrintValue>("Tree2GlobalTask2");
+				GlobalTask2.GetInstanceData().Value = -2;
+				EditorData2.AddPropertyBinding(FPropertyBindingPath(GlobalTask1.ID, "Value"), FPropertyBindingPath(GlobalTask2.ID, "Value"));
+			}
+
+			UStateTreeState& Root = EditorData2.AddSubTree("Tree2StateRoot");
+			{
+				AddInt(Root.Parameters.Parameters, "Tree2StateRootParametersInt");
+				AddDouble(Root.Parameters.Parameters, "Tree2StateRootParametersDouble");
+				EditorData2.AddPropertyBinding(FPropertyBindingPath(EditorData2.GetRootParametersGuid(), "Tree2GlobalInt"), FPropertyBindingPath(Root.Parameters.ID, "Tree2StateRootParametersInt"));
+			}
+			{
+				TStateTreeEditorNode<FTestTask_PrintValue>& Task1 = Root.AddTask<FTestTask_PrintValue>("Tree2StateRootTask1");
+				Task1.GetInstanceData().Value = -1;
+				EditorData2.AddPropertyBinding(FPropertyBindingPath(EditorData2.GetRootParametersGuid(), "Tree2GlobalInt"), FPropertyBindingPath(Task1.ID, "Value"));
+
+				TStateTreeEditorNode<FTestTask_PrintValue>& Task2 = Root.AddTask<FTestTask_PrintValue>("Tree2StateRootTask2");
+				Task2.GetInstanceData().Value = -2;
+				EditorData2.AddPropertyBinding(FPropertyBindingPath(Root.Parameters.ID, "Tree2StateRootParametersInt"), FPropertyBindingPath(Task2.ID, "Value"));
+				
+				TStateTreeEditorNode<FTestTask_PrintValue>& Task3 = Root.AddTask<FTestTask_PrintValue>("Tree2StateRootTask3");
+				Task3.GetInstanceData().Value = -3;
+				EditorData2.AddPropertyBinding(FPropertyBindingPath(GlobalTask1.ID, "Value"), FPropertyBindingPath(Task3.ID, "Value"));
+			}
+			{
+				UStateTreeState& State1 = Root.AddChildState("Tree2State1", EStateTreeStateType::State);
+				{
+					AddDouble(State1.Parameters.Parameters, "Tree2State1ParametersDouble");
+					AddInt(State1.Parameters.Parameters, "Tree2State1ParametersInt");
+					EditorData2.AddPropertyBinding(FPropertyBindingPath(Root.Parameters.ID, "Tree2StateRootParametersInt"), FPropertyBindingPath(State1.Parameters.ID, "Tree2State1ParametersInt"));
+				}
+
+				TStateTreeEditorNode<FTestTask_PrintValue>& Task1 = State1.AddTask<FTestTask_PrintValue>("Tree2State1Task1");
+				Task1.GetInstanceData().Value = -1;
+				EditorData2.AddPropertyBinding(FPropertyBindingPath(EditorData2.GetRootParametersGuid(), "Tree2GlobalInt"), FPropertyBindingPath(Task1.ID, TEXT("Value")));
+
+				TStateTreeEditorNode<FTestTask_PrintValue>& Task2 = State1.AddTask<FTestTask_PrintValue>("Tree2State1Task2");
+				Task2.GetInstanceData().Value = -2;
+				EditorData2.AddPropertyBinding(FPropertyBindingPath(Root.Parameters.ID, "Tree2StateRootParametersInt"), FPropertyBindingPath(Task2.ID, TEXT("Value")));
+
+				TStateTreeEditorNode<FTestTask_PrintValue>& Task3 = State1.AddTask<FTestTask_PrintValue>("Tree2State1Task3");
+				Task3.GetInstanceData().Value = -3;
+				EditorData2.AddPropertyBinding(FPropertyBindingPath(GlobalTask1.ID, "Value"), FPropertyBindingPath(Task3.ID, "Value"));
+
+				TStateTreeEditorNode<FTestTask_PrintValue>& Task4 = State1.AddTask<FTestTask_PrintValue>("Tree2State1Task4");
+				Task3.GetInstanceData().Value = -4;
+				EditorData2.AddPropertyBinding(FPropertyBindingPath(State1.Parameters.ID, "Tree2State1ParametersInt"), FPropertyBindingPath(Task4.ID, "Value"));
+
+				State1.AddTransition(EStateTreeTransitionTrigger::OnTick, EStateTreeTransitionType::Succeeded);
+			}
+
+			FStateTreeCompilerLog Log;
+			FStateTreeCompiler Compiler2(Log);
+			const bool bResult2 = Compiler2.Compile(StateTree2);
+			AITEST_TRUE("StateTree2 should get compiled", bResult2);
+		}
+
+		// Tree 1
+		UStateTree& StateTree1 = NewStateTree();
+		FGuid Tree1GlobalParameter_ValueID_Int;
+		{
+			UStateTreeEditorData& EditorData1 = *Cast<UStateTreeEditorData>(StateTree1.EditorData);
+
+			Tree1GlobalParameter_ValueID_Int = AddInt(GetRootPropertyBag(EditorData1), "Tree1GlobalInt");
+			AddDouble(GetRootPropertyBag(EditorData1), "Tree1GlobalDouble");
+
+			// Global tasks
+			TStateTreeEditorNode<FTestTask_PrintValue>& GlobalTask1 = EditorData1.AddGlobalTask<FTestTask_PrintValue>("Tree1GlobalTask1");
+			{
+				GlobalTask1.GetInstanceData().Value = -1;
+				EditorData1.AddPropertyBinding(FPropertyBindingPath(EditorData1.GetRootParametersGuid(), "Tree1GlobalInt"), FPropertyBindingPath(GlobalTask1.ID, "Value"));
+			}
+			{
+				TStateTreeEditorNode<FTestTask_PrintValue>& GlobalTask2 = EditorData1.AddGlobalTask<FTestTask_PrintValue>("Tree1GlobalTask2");
+				GlobalTask2.GetInstanceData().Value = -2;
+				EditorData1.AddPropertyBinding(FPropertyBindingPath(GlobalTask1.ID, "Value"), FPropertyBindingPath(GlobalTask2.ID, "Value"));
+			}
+
+			UStateTreeState& Root = EditorData1.AddSubTree("Tree1StateRoot");
+			{
+				AddDouble(Root.Parameters.Parameters, "Tree1StateRootParametersDouble");
+				AddInt(Root.Parameters.Parameters, "Tree1StateRootParametersInt");
+				EditorData1.AddPropertyBinding(FPropertyBindingPath(EditorData1.GetRootParametersGuid(), "Tree1GlobalInt"), FPropertyBindingPath(Root.Parameters.ID, "Tree1StateRootParametersInt"));
+			}
+			{
+				TStateTreeEditorNode<FTestTask_PrintValue>& Task1 = Root.AddTask<FTestTask_PrintValue>("Tree1StateRootTask1");
+				Task1.GetInstanceData().Value = -1;
+				EditorData1.AddPropertyBinding(FPropertyBindingPath(EditorData1.GetRootParametersGuid(), "Tree1GlobalInt"), FPropertyBindingPath(Task1.ID, "Value"));
+
+				TStateTreeEditorNode<FTestTask_PrintValue>& Task2 = Root.AddTask<FTestTask_PrintValue>("Tree1StateRootTask2");
+				Task2.GetInstanceData().Value = -2;
+				EditorData1.AddPropertyBinding(FPropertyBindingPath(Root.Parameters.ID, "Tree1StateRootParametersInt"), FPropertyBindingPath(Task2.ID, "Value"));
+
+				TStateTreeEditorNode<FTestTask_PrintValue>& Task3 = Root.AddTask<FTestTask_PrintValue>("Tree1StateRootTask3");
+				Task3.GetInstanceData().Value = -3;
+				EditorData1.AddPropertyBinding(FPropertyBindingPath(GlobalTask1.ID, "Value"), FPropertyBindingPath(Task3.ID, "Value"));
+			}
+			{
+				UStateTreeState& State1 = Root.AddChildState("Tree1State1", EStateTreeStateType::LinkedAsset);
+				State1.SetLinkedStateAsset(&StateTree2);
+				State1.AddTransition(EStateTreeTransitionTrigger::OnStateSucceeded, EStateTreeTransitionType::NextState);
+
+				State1.SetParametersPropertyOverridden(Tree2GlobalParameter_ValueID_Int, true);
+				const FInstancedPropertyBag* Parameters = State1.GetDefaultParameters();
+				AITEST_TRUE(TEXT("Parameter is invalid"), Parameters != nullptr);
+				EditorData1.AddPropertyBinding(FPropertyBindingPath(GlobalTask1.ID, "Value"), FPropertyBindingPath(State1.Parameters.ID, "Tree2GlobalInt"));
+			}
+
+			FGuid Tree1Sub1Parameter_ValueID_Int;
+			UStateTreeState& Sub1 = EditorData1.AddSubTree("Tree1StateSub1");
+			{
+				Sub1.Type = EStateTreeStateType::Subtree;
+				Sub1.AddTransition(EStateTreeTransitionTrigger::OnStateSucceeded, EStateTreeTransitionType::Succeeded);
+
+				AddDouble(Sub1.Parameters.Parameters, "Tree1StateSub1ParametersDouble");
+				Tree1Sub1Parameter_ValueID_Int = AddInt(Sub1.Parameters.Parameters, "Tree1StateSub1ParametersInt");
+
+				TStateTreeEditorNode<FTestTask_PrintValue>& Sub1Task1 = Sub1.AddTask<FTestTask_PrintValue>("Tree1StateSub1Task1");
+				Sub1Task1.GetInstanceData().Value = -1;
+				EditorData1.AddPropertyBinding(FPropertyBindingPath(Sub1.Parameters.ID, "Tree1StateSub1ParametersInt"), FPropertyBindingPath(Sub1Task1.ID, "Value"));
+
+				{
+					UStateTreeState& State3 = Sub1.AddChildState("Tree1State3", EStateTreeStateType::State);
+					{
+						AddDouble(State3.Parameters.Parameters, "Tree1State3ParametersDouble1");
+						AddDouble(State3.Parameters.Parameters, "Tree1State3ParametersDouble2");
+						AddInt(State3.Parameters.Parameters, "Tree1State3ParametersInt");
+						EditorData1.AddPropertyBinding(FPropertyBindingPath(Sub1.Parameters.ID, "Tree1StateSub1ParametersInt"), FPropertyBindingPath(State3.Parameters.ID, "Tree1State3ParametersInt"));
+					}
+
+					TStateTreeEditorNode<FTestTask_PrintValue>& Task1 = State3.AddTask<FTestTask_PrintValue>("Tree1State3Task1");
+					Task1.GetInstanceData().Value = -1;
+					EditorData1.AddPropertyBinding(FPropertyBindingPath(EditorData1.GetRootParametersGuid(), "Tree1GlobalInt"), FPropertyBindingPath(Task1.ID, TEXT("Value")));
+
+					TStateTreeEditorNode<FTestTask_PrintValue>& Task2 = State3.AddTask<FTestTask_PrintValue>("Tree1State3Task2");
+					Task2.GetInstanceData().Value = -2;
+					EditorData1.AddPropertyBinding(FPropertyBindingPath(Sub1.Parameters.ID, "Tree1StateSub1ParametersInt"), FPropertyBindingPath(Task2.ID, TEXT("Value")));
+
+					TStateTreeEditorNode<FTestTask_PrintValue>& Task3 = State3.AddTask<FTestTask_PrintValue>("Tree1State3Task3");
+					Task3.GetInstanceData().Value = -3;
+					EditorData1.AddPropertyBinding(FPropertyBindingPath(GlobalTask1.ID, "Value"), FPropertyBindingPath(Task3.ID, "Value"));
+
+					TStateTreeEditorNode<FTestTask_PrintValue>& Task4 = State3.AddTask<FTestTask_PrintValue>("Tree1State3Task4");
+					Task3.GetInstanceData().Value = -4;
+					EditorData1.AddPropertyBinding(FPropertyBindingPath(Task1.ID, "Value"), FPropertyBindingPath(Task4.ID, "Value"));
+
+					State3.AddTransition(EStateTreeTransitionTrigger::OnTick, EStateTreeTransitionType::Succeeded);
+				}
+			}
+			{
+				UStateTreeState& State2 = Root.AddChildState("Tree1State2", EStateTreeStateType::Linked);
+				State2.SetLinkedState(Sub1.GetLinkToState());
+				State2.AddTransition(EStateTreeTransitionTrigger::OnStateSucceeded, EStateTreeTransitionType::Succeeded);
+
+				State2.SetParametersPropertyOverridden(Tree1Sub1Parameter_ValueID_Int, true);
+				const FInstancedPropertyBag* Parameters = State2.GetDefaultParameters();
+				AITEST_TRUE(TEXT("Parameter is invalid"), Parameters != nullptr);
+				EditorData1.AddPropertyBinding(FPropertyBindingPath(GlobalTask1.ID, "Value"), FPropertyBindingPath(State2.Parameters.ID, "Tree1StateSub1ParametersInt"));
+			}
+
+			FStateTreeCompilerLog Log;
+			FStateTreeCompiler Compiler1(Log);
+			const bool bResult1 = Compiler1.Compile(StateTree1);
+			AITEST_TRUE("StateTree1 should get compiled", bResult1);
+		}
+
+		FStateTreeInstanceData InstanceData;
+		{
+			FTestStateTreeExecutionContext Exec(StateTree1, StateTree1, InstanceData);
+			const bool bInitSucceeded = Exec.IsValid();
+			AITEST_TRUE("StateTree should init", bInitSucceeded);
+		}
+		{
+			FStateTreeReference StateTreeRef;
+			StateTreeRef.SetStateTree(&StateTree1);
+			StateTreeRef.SetPropertyOverridden(Tree1GlobalParameter_ValueID_Int, true);
+			StateTreeRef.GetMutableParameters().SetValueInt32("Tree1GlobalInt", 5);
+			FTestStateTreeExecutionContext Exec(StateTree1, StateTree1, InstanceData);
+			const EStateTreeRunStatus Status = Exec.Start(FStateTreeExecutionContext::FStartParameters
+				{
+					.GlobalParameters = &StateTreeRef.GetParameters(),
+				});
+			AITEST_EQUAL("Start should complete with Running", Status, EStateTreeRunStatus::Running);
+			AITEST_TRUE(TEXT("In correct states"), Exec.ExpectInActiveStates("Tree1StateRoot", "Tree1State1", "Tree2StateRoot", "Tree2State1"));
+			AITEST_TRUE("Start should enter Tree1GlobalTask1", Exec.Expect("Tree1GlobalTask1", TEXT("EnterState5")));
+			AITEST_TRUE("Start should enter Tree1GlobalTask2", Exec.Expect("Tree1GlobalTask2", TEXT("EnterState5")));
+			AITEST_TRUE("Start should enter Tree1StateRootTask1", Exec.Expect("Tree1StateRootTask1", TEXT("EnterState5")));
+			AITEST_TRUE("Start should enter Tree1StateRootTask2", Exec.Expect("Tree1StateRootTask2", TEXT("EnterState5")));
+			AITEST_TRUE("Start should enter Tree1StateRootTask3", Exec.Expect("Tree1StateRootTask3", TEXT("EnterState5")));
+			AITEST_TRUE("Start should enter Tree2GlobalTask1", Exec.Expect("Tree2GlobalTask1", TEXT("EnterState5")));
+			AITEST_TRUE("Start should enter Tree2GlobalTask2", Exec.Expect("Tree2GlobalTask2", TEXT("EnterState5")));
+			AITEST_TRUE("Start should enter Tree2StateRootTask1", Exec.Expect("Tree2StateRootTask1", TEXT("EnterState5")));
+			AITEST_TRUE("Start should enter Tree2StateRootTask2", Exec.Expect("Tree2StateRootTask2", TEXT("EnterState5")));
+			AITEST_TRUE("Start should enter Tree2StateRootTask3", Exec.Expect("Tree2StateRootTask3", TEXT("EnterState5")));
+			AITEST_TRUE("Start should enter Tree2State1Task1", Exec.Expect("Tree2State1Task1", TEXT("EnterState5")));
+			AITEST_TRUE("Start should enter Tree2State1Task2", Exec.Expect("Tree2State1Task2", TEXT("EnterState5")));
+			AITEST_TRUE("Start should enter Tree2State1Task3", Exec.Expect("Tree2State1Task3", TEXT("EnterState5")));
+			AITEST_TRUE("Start should enter Tree2State1Task4", Exec.Expect("Tree2State1Task4", TEXT("EnterState5")));
+			Exec.LogClear();
+		}
+		{
+			FStateTreeReference StateTreeRef;
+			StateTreeRef.SetStateTree(&StateTree1);
+			StateTreeRef.SetPropertyOverridden(Tree1GlobalParameter_ValueID_Int, true);
+			StateTreeRef.GetMutableParameters().SetValueInt32("Tree1GlobalInt", 6);
+
+			FTestStateTreeExecutionContext Exec(StateTree1, StateTree1, InstanceData);
+			InstanceData.GetMutableStorage().SetGlobalParameters(StateTreeRef.GetParameters());
+
+			const EStateTreeRunStatus Status = Exec.Tick(1.0f);
+			AITEST_EQUAL("Tick should complete with Running", Status, EStateTreeRunStatus::Running);
+
+			AITEST_TRUE(TEXT("In correct states"), Exec.ExpectInActiveStates("Tree1StateRoot", "Tree1State2", "Tree1StateSub1", "Tree1State3"));
+			AITEST_TRUE("Tick should tick Tree1GlobalTask1", Exec.Expect("Tree1GlobalTask1", TEXT("Tick6")));
+			AITEST_TRUE("Tick should tick Tree1GlobalTask2", Exec.Expect("Tree1GlobalTask2", TEXT("Tick6")));
+			AITEST_TRUE("Tick should tick Tree1StateRootTask1", Exec.Expect("Tree1StateRootTask1", TEXT("Tick6")));
+			//@todo Binding on the state are not updated on tick. Is this a bug? See Tree1StateRootParametersInt
+			//AITEST_TRUE("Tick should tick Tree1StateRootTask2", Exec.Expect("Tree1StateRootTask2", TEXT("Tick6")));
+			AITEST_TRUE("Tick should tick Tree1StateRootTask3", Exec.Expect("Tree1StateRootTask3", TEXT("Tick6")));
+			AITEST_TRUE("Tick should tick Tree2GlobalTask1", Exec.Expect("Tree2GlobalTask1", TEXT("Tick6")));
+			AITEST_TRUE("Tick should tick Tree2GlobalTask2", Exec.Expect("Tree2GlobalTask2", TEXT("Tick6")));
+			AITEST_TRUE("Tick should tick Tree2StateRootTask1", Exec.Expect("Tree2StateRootTask1", TEXT("Tick6")));
+			//@todo Binding on the state are not updated on tick. Is this a bug? See Tree2StateRootParametersInt
+			//AITEST_TRUE("Tick should tick Tree2StateRootTask2", Exec.Expect("Tree2StateRootTask2", TEXT("Tick6")));
+			AITEST_TRUE("Tick should tick Tree2StateRootTask3", Exec.Expect("Tree2StateRootTask3", TEXT("Tick6")));
+			AITEST_TRUE("Tick should tick Tree2State1Task1", Exec.Expect("Tree2State1Task1", TEXT("Tick6")));
+			//@todo Binding on the state are not updated on tick. Is this a bug? See Tree2StateRootParametersInt
+			//AITEST_TRUE("Tick should tick Tree2State1Task2", Exec.Expect("Tree2State1Task2", TEXT("Tick6")));
+			AITEST_TRUE("Tick should tick Tree2State1Task3", Exec.Expect("Tree2State1Task3", TEXT("Tick6")));
+			//@todo Binding on the state are not updated on tick. Is this a bug? See Tree2State1ParametersInt
+			//AITEST_TRUE("Tick should tick Tree2State1Task4", Exec.Expect("Tree2State1Task4", TEXT("Tick6")));
+
+			AITEST_TRUE("Tick should enter Tree1StateSub1Task1", Exec.Expect("Tree1StateSub1Task1", TEXT("EnterState6")));
+			AITEST_TRUE("Tick should enter Tree1State3Task1", Exec.Expect("Tree1State3Task1", TEXT("EnterState6")));
+			AITEST_TRUE("Tick should enter Tree1State3Task2", Exec.Expect("Tree1State3Task2", TEXT("EnterState6")));
+			AITEST_TRUE("Tick should enter Tree1State3Task3", Exec.Expect("Tree1State3Task3", TEXT("EnterState6")));
+			AITEST_TRUE("Tick should enter Tree1State3Task4", Exec.Expect("Tree1State3Task4", TEXT("EnterState6")));
+		}
+		{
+			FTestStateTreeExecutionContext Exec(StateTree1, StateTree1, InstanceData);
+			Exec.Stop();
+		}
+
+		return true;
+	}
+};
+IMPLEMENT_AI_INSTANT_TEST(FStateTreeTest_Linked_GlobalParameter, "System.StateTree.LinkedAsset.GlobalParameter");
+
 } // namespace UE::StateTree::Tests
 
 UE_ENABLE_OPTIMIZATION_SHIP
diff --git a/Engine/Plugins/Runtime/StateTree/Source/StateTreeTestSuite/Private/StateTreeStatePathTest.cpp b/Engine/Plugins/Runtime/StateTree/Source/StateTreeTestSuite/Private/StateTreeStatePathTest.cpp
--- a/Engine/Plugins/Runtime/StateTree/Source/StateTreeTestSuite/Private/StateTreeStatePathTest.cpp
+++ b/Engine/Plugins/Runtime/StateTree/Source/StateTreeTestSuite/Private/StateTreeStatePathTest.cpp
@@ -71,39 +71,39 @@ struct FStateTreeTest_StatePath_LinkStates : FStateTreeTestBase
 				{
 					{
 						UStateTreeState& ChildState = Tree1RootState1.AddChildState("Tree1State1StateSub1A", EStateTreeStateType::Linked);
-						ChildState.LinkedSubtree = Sub1.GetLinkToState();
+						ChildState.SetLinkedState(Sub1.GetLinkToState());
 						ChildState.AddTransition(EStateTreeTransitionTrigger::OnStateSucceeded, EStateTreeTransitionType::NextState);
 					}
 					{
 						UStateTreeState& ChildState = Tree1RootState1.AddChildState("Tree1State1StateSub1B", EStateTreeStateType::Linked);
-						ChildState.LinkedSubtree = Sub1.GetLinkToState();
+						ChildState.SetLinkedState(Sub1.GetLinkToState());
 						FStateTreeTransition& Transition = ChildState.AddTransition(EStateTreeTransitionTrigger::OnStateSucceeded, EStateTreeTransitionType::GotoState);
 						Transition.State = Tree1RootStateSub2A.GetLinkToState();
 					}
 				}
 				//Tree1RootStateSub2A
 				{
-					Tree1RootStateSub2A.LinkedSubtree = Sub2.GetLinkToState();
+					Tree1RootStateSub2A.SetLinkedState(Sub2.GetLinkToState());
 					Tree1RootStateSub2A.AddTransition(EStateTreeTransitionTrigger::OnStateSucceeded, EStateTreeTransitionType::NextState);
 				}
 				{
 					UStateTreeState& State = Root.AddChildState("Tree1RootStateSub2B", EStateTreeStateType::Linked);
-					State.LinkedSubtree = Sub2.GetLinkToState();
+					State.SetLinkedState(Sub2.GetLinkToState());
 					State.AddTransition(EStateTreeTransitionTrigger::OnStateSucceeded, EStateTreeTransitionType::NextState);
 				}
 				{
 					UStateTreeState& State = Root.AddChildState("Tree1RootStateLinkTree2A", EStateTreeStateType::LinkedAsset);
-					State.LinkedAsset = &StateTree2;
+					State.SetLinkedStateAsset(&StateTree2);
 					State.AddTransition(EStateTreeTransitionTrigger::OnStateSucceeded, EStateTreeTransitionType::NextState);
 				}
 				{
 					UStateTreeState& State = Root.AddChildState("Tree1RootStateLinkTree3", EStateTreeStateType::LinkedAsset);
-					State.LinkedAsset = &StateTree3;
+					State.SetLinkedStateAsset(&StateTree3);
 					State.AddTransition(EStateTreeTransitionTrigger::OnStateSucceeded, EStateTreeTransitionType::NextState);
 				}
 				{
 					UStateTreeState& State = Root.AddChildState("Tree1RootStateLinkTree2B", EStateTreeStateType::LinkedAsset);
-					State.LinkedAsset = &StateTree2;
+					State.SetLinkedStateAsset(&StateTree2);
 					State.AddTransition(EStateTreeTransitionTrigger::OnStateSucceeded, EStateTreeTransitionType::NextState);
 				}
 				{
@@ -118,22 +118,22 @@ struct FStateTreeTest_StatePath_LinkStates : FStateTreeTestBase
 			{
 				{
 					UStateTreeState& State = Sub1.AddChildState("Tree1Sub1StateLinkTree2A", EStateTreeStateType::LinkedAsset);
-					State.LinkedAsset = &StateTree2;
+					State.SetLinkedStateAsset(&StateTree2);
 					State.AddTransition(EStateTreeTransitionTrigger::OnStateSucceeded, EStateTreeTransitionType::NextState);
 				}
 				{
 					UStateTreeState& State = Sub1.AddChildState("Tree1Sub1StateLinkTree2B", EStateTreeStateType::LinkedAsset);
-					State.LinkedAsset = &StateTree2;
+					State.SetLinkedStateAsset(&StateTree2);
 					State.AddTransition(EStateTreeTransitionTrigger::OnStateSucceeded, EStateTreeTransitionType::NextState);
 				}
 				{
 					UStateTreeState& State = Sub1.AddChildState("Tree1Sub1StateSub2A", EStateTreeStateType::Linked);
-					State.LinkedSubtree = Sub2.GetLinkToState();
+					State.SetLinkedState(Sub2.GetLinkToState());
 					State.AddTransition(EStateTreeTransitionTrigger::OnStateSucceeded, EStateTreeTransitionType::NextState);
 				}
 				{
 					UStateTreeState& State = Sub1.AddChildState("Tree1Sub1StateSub2B", EStateTreeStateType::Linked);
-					State.LinkedSubtree = Sub2.GetLinkToState();
+					State.SetLinkedState(Sub2.GetLinkToState());
 					State.AddTransition(EStateTreeTransitionTrigger::OnStateSucceeded, EStateTreeTransitionType::NextState);
 				}
 				{
@@ -167,12 +167,12 @@ struct FStateTreeTest_StatePath_LinkStates : FStateTreeTestBase
 			{
 				{
 					UStateTreeState& State = Root.AddChildState("Tree2RootStateSub1A", EStateTreeStateType::Linked);
-					State.LinkedSubtree = Sub1.GetLinkToState();
+					State.SetLinkedState(Sub1.GetLinkToState());
 					State.AddTransition(EStateTreeTransitionTrigger::OnStateSucceeded, EStateTreeTransitionType::NextState);
 				}
 				{
 					UStateTreeState& State = Root.AddChildState("Tree2RootStateSub1B", EStateTreeStateType::Linked);
-					State.LinkedSubtree = Sub1.GetLinkToState();
+					State.SetLinkedState(Sub1.GetLinkToState());
 					State.AddTransition(EStateTreeTransitionTrigger::OnStateSucceeded, EStateTreeTransitionType::NextState);
 				}
 				{
@@ -205,7 +205,7 @@ struct FStateTreeTest_StatePath_LinkStates : FStateTreeTestBase
 			//Root
 			{
 				UStateTreeState& State = Root.AddChildState("Tree2RootStateSub1A", EStateTreeStateType::Linked);
-				State.LinkedSubtree = Sub1.GetLinkToState();
+				State.SetLinkedState(Sub1.GetLinkToState());
 				State.AddTransition(EStateTreeTransitionTrigger::OnStateSucceeded, EStateTreeTransitionType::Succeeded);
 			}
 			//Tree3Sub1
diff --git a/Engine/Plugins/Runtime/StateTree/Source/StateTreeTestSuite/Private/StateTreeTest.cpp b/Engine/Plugins/Runtime/StateTree/Source/StateTreeTestSuite/Private/StateTreeTest.cpp
--- a/Engine/Plugins/Runtime/StateTree/Source/StateTreeTestSuite/Private/StateTreeTest.cpp
+++ b/Engine/Plugins/Runtime/StateTree/Source/StateTreeTestSuite/Private/StateTreeTest.cpp
@@ -380,7 +380,7 @@ struct FStateTreeTest_SubTree_ActiveTasks : FStateTreeTestBase
 		UStateTreeState& State3A = State3.AddChildState(FName(TEXT("State3A")));
 		UStateTreeState& State3B = State3.AddChildState(FName(TEXT("State3B")));
 
-		State1.LinkedSubtree = State3.GetLinkToState();
+		State1.SetLinkedState(State3.GetLinkToState());
 
 		State1.AddTransition(EStateTreeTransitionTrigger::OnStateCompleted, EStateTreeTransitionType::GotoState, &State2);
 
@@ -470,7 +470,7 @@ struct FStateTreeTest_SubTree_NoActiveTasks : FStateTreeTestBase
 
 		SubTree.Type = EStateTreeStateType::Subtree;
 		StateA.Type = EStateTreeStateType::Linked;
-		StateA.LinkedSubtree = SubTree.GetLinkToState();
+		StateA.SetLinkedState(SubTree.GetLinkToState());
 
 		StateA.AddTransition(EStateTreeTransitionTrigger::OnStateCompleted, EStateTreeTransitionType::GotoState, &StateB);
 		RootB.AddTransition(EStateTreeTransitionTrigger::OnStateCompleted, EStateTreeTransitionType::GotoState, &StateB);
@@ -547,7 +547,7 @@ struct FStateTreeTest_SubTree_Condition : FStateTreeTestBase
 		UStateTreeState& State1 = SubTree.AddChildState(FName(TEXT("State1")));
 		UStateTreeState& State2 = SubTree.AddChildState(FName(TEXT("State2")));
 
-		Linked.LinkedSubtree = SubTree.GetLinkToState();
+		Linked.SetLinkedState(SubTree.GetLinkToState());
 
 		Linked.AddTransition(EStateTreeTransitionTrigger::OnStateCompleted, EStateTreeTransitionType::GotoState, &Root);
 
@@ -630,8 +630,8 @@ struct FStateTreeTest_SubTree_CascadedSucceeded : FStateTreeTestBase
 		UStateTreeState& SubSubTreeState = Root.AddChildState(FName(TEXT("SubSubTreeState")), EStateTreeStateType::Subtree);
 		UStateTreeState& SubSubLeaf = SubSubTreeState.AddChildState(FName(TEXT("SubSubLeaf")));
 
-		LinkedState.LinkedSubtree = SubTreeState.GetLinkToState();
-		SubLinkedState.LinkedSubtree = SubSubTreeState.GetLinkToState();
+		LinkedState.SetLinkedState(SubTreeState.GetLinkToState());
+		SubLinkedState.SetLinkedState(SubSubTreeState.GetLinkToState());
 
 		LinkedState.AddTransition(EStateTreeTransitionTrigger::OnStateFailed, EStateTreeTransitionType::Failed);
 		SubLinkedState.AddTransition(EStateTreeTransitionTrigger::OnStateSucceeded, EStateTreeTransitionType::Failed);
@@ -2146,7 +2146,7 @@ struct FStateTreeTest_SubTreeTransition : FStateTreeTestBase
 		UStateTreeState& TimeoutChecker = LastStand.AddChildState(FName(TEXT("TimeoutChecker")), EStateTreeStateType::Subtree);
 		UStateTreeState& RemainingCount = TimeoutChecker.AddChildState(FName(TEXT("RemainingCount")));
 
-		Reinforcements.LinkedSubtree = TimeoutChecker.GetLinkToState();
+		Reinforcements.SetLinkedState(TimeoutChecker.GetLinkToState());
 
 
 		TStateTreeEditorNode<FTestTask_Stand>& Task1 = PreLastStand.AddTask<FTestTask_Stand>(FName(TEXT("Task1")));
