StateTree: Only stop ticking tasks on failure if the task is considered for completion.
EnterState stops if a task is completed, it now has the same behavior as the tick. It was a bug but not an issue before the task completion feature.
StartTemporaryEvaluatorsAndGlobalTasks returns the correct result.
TasksCompletionStatus::ResetStatus was only resetting the tasks that were in the completion mask instead of resetting all the tasks inside the state.
The completion mask was ordered in reverse order.

diff --git a/Engine/Plugins/Runtime/StateTree/Source/StateTreeEditorModule/Private/StateTreeCompiler.cpp b/Engine/Plugins/Runtime/StateTree/Source/StateTreeEditorModule/Private/StateTreeCompiler.cpp
--- a/Engine/Plugins/Runtime/StateTree/Source/StateTreeEditorModule/Private/StateTreeCompiler.cpp
+++ b/Engine/Plugins/Runtime/StateTree/Source/StateTreeEditorModule/Private/StateTreeCompiler.cpp
@@ -240,7 +240,7 @@ namespace UE::StateTree::Compiler
 		}
 		else
 		{
-			for (int32 Index = 0; Index < NumberOfTasks; ++Index)
+			for (int32 Index = NumberOfTasks-1; Index >= 0; --Index)
 			{
 				const int32 TaskIndex = ValidTasks[Index];
 				const FStateTreeTaskBase& Task = AllTasks[TaskIndex].Node.Get<FStateTreeTaskBase>();
@@ -257,17 +257,16 @@ namespace UE::StateTree::Compiler
 
 		// Is the new amount of bits bring up over the next buffer?
 		const int32 CurrentEndBitIndex = FullStartBitIndex + NumberOfBitsNeeded;
-		if (CurrentEndBitIndex / NumberOfBitsPerMask != FullStartBitIndex / NumberOfBitsPerMask)
+		const int32 NewMaskBufferIndex = (CurrentEndBitIndex - 1) / NumberOfBitsPerMask;
+		if (NewMaskBufferIndex != FullStartBitIndex / NumberOfBitsPerMask)
 		{
 			// Do not shift the mask. Use the next int32
-			const int32 NewMaskBufferIndex = CurrentEndBitIndex / NumberOfBitsPerMask;
 			const int32 NewMaskFirstTaskBitOffset = 0;
 			const int32 NewMaskEndTaskBitOffset = (NewMaskBufferIndex * NumberOfBitsPerMask) + NumberOfBitsNeeded;
 			return { .Mask = Mask, .MaskBufferIndex = NewMaskBufferIndex, .MaskFirstTaskBitOffset = NewMaskFirstTaskBitOffset, .FullMaskEndTaskBitOffset = NewMaskEndTaskBitOffset };
 		}
 		else
 		{
-			const int32 NewMaskBufferIndex = FullStartBitIndex / NumberOfBitsPerMask;
 			const int32 NewMaskFirstTaskBitOffset = FullStartBitIndex % NumberOfBitsPerMask;
 			const int32 NewMaskEndTaskBitOffset = CurrentEndBitIndex;
 
diff --git a/Engine/Plugins/Runtime/StateTree/Source/StateTreeModule/Private/StateTreeExecutionContext.cpp b/Engine/Plugins/Runtime/StateTree/Source/StateTreeModule/Private/StateTreeExecutionContext.cpp
--- a/Engine/Plugins/Runtime/StateTree/Source/StateTreeModule/Private/StateTreeExecutionContext.cpp
+++ b/Engine/Plugins/Runtime/StateTree/Source/StateTreeModule/Private/StateTreeExecutionContext.cpp
@@ -2716,7 +2716,7 @@ EStateTreeRunStatus FStateTreeExecutionContext::EnterState(FStateTreeTransitionR
 					__FUNCTION__, *GetStateStatusString(Exec), *GetNameSafe(&Owner), *GetFullNameSafe(&RootStateTree));
 				break;
 			}
-			CurrentFrame.ActiveTasksStatus.Push(CurrentStateTree, State);
+			CurrentFrame.ActiveTasksStatus.Push(State);
 
 			UE::StateTree::FTasksCompletionStatus CurrentStateTasksStatus = CurrentFrame.ActiveTasksStatus.GetStatus(State);
 			CurrentFrame.NumCurrentlyActiveStates = static_cast<uint8>(CurrentFrame.ActiveStates.Num());
@@ -2843,7 +2843,7 @@ EStateTreeRunStatus FStateTreeExecutionContext::EnterState(FStateTreeTransitionR
 					UE_STATETREE_DEBUG_TASK_EVENT(this, AssetTaskIndex, TaskInstanceView, EStateTreeTraceEventType::OnEntered, TaskRunStatus);
 
 					Result = UE::StateTree::ExecutionContext::Private::GetPriorityRunStatus(Result, TaskRunStatus);
-					if (Result == EStateTreeRunStatus::Failed)
+					if (TaskRunStatus == EStateTreeRunStatus::Failed && CurrentStateTasksStatus.IsConsideredForCompletion(StateTaskIndex))
 					{
 						Exec.EnterStateFailedFrameIndex = FStateTreeIndex16(FrameIndex);
 						Exec.EnterStateFailedTaskIndex = FStateTreeIndex16(AssetTaskIndex);
@@ -3438,7 +3438,7 @@ EStateTreeRunStatus FStateTreeExecutionContext::StartEvaluatorsAndGlobalTasks(FS
 				UE_STATETREE_DEBUG_TASK_EVENT(this, AssetTaskIndex, TaskInstanceView, EStateTreeTraceEventType::OnEntered, TaskRunStatus);
 
 				Result = UE::StateTree::ExecutionContext::Private::GetPriorityRunStatus(Result, TaskRunStatus);
-				if (Result != EStateTreeRunStatus::Running)
+				if (TaskRunStatus == EStateTreeRunStatus::Failed && CurrentGlobalTasksStatus.IsConsideredForCompletion(GlobalTaskIndex))
 				{
 					OutLastInitializedTaskIndex = FStateTreeIndex16(AssetTaskIndex);
 				}
@@ -3664,7 +3664,8 @@ EStateTreeRunStatus FStateTreeExecutionContext::StartTemporaryEvaluatorsAndGloba
 		TaskRunStatus = UE::StateTree::ExecutionContext::Private::CastToRunStatus(TaskStatus);
 		UE_STATETREE_DEBUG_TASK_EVENT(this, AssetTaskIndex, TaskDataView, EStateTreeTraceEventType::OnEntered, TaskRunStatus);
 
-		if (TaskRunStatus != EStateTreeRunStatus::Running)
+		Result = UE::StateTree::ExecutionContext::Private::GetPriorityRunStatus(Result, TaskRunStatus);
+		if (TaskRunStatus == EStateTreeRunStatus::Failed && CurrentTasksStatus.IsConsideredForCompletion(GlobalTaskIndex))
 		{
 			break;
 		}
@@ -3675,6 +3676,7 @@ EStateTreeRunStatus FStateTreeExecutionContext::StartTemporaryEvaluatorsAndGloba
 
 void FStateTreeExecutionContext::StopTemporaryEvaluatorsAndGlobalTasks(const FStateTreeExecutionFrame* CurrentParentFrame, const FStateTreeExecutionFrame& CurrentFrame)
 {
+	//@todo only stop the evaluators and tasks that were started in StartTemporaryEvaluatorsAndGlobalTasks
 	STATETREE_LOG(Verbose, TEXT("Stop Temporary Evaluators & Global tasks"));
 
 	FCurrentlyProcessedFrameScope FrameScope(*this, CurrentParentFrame, CurrentFrame);
@@ -3984,7 +3986,7 @@ FStateTreeExecutionContext::FTickTaskResult FStateTreeExecutionContext::TickTask
 			TaskRunStatus != EStateTreeRunStatus::Running ? EStateTreeTraceEventType::OnTaskCompleted : EStateTreeTraceEventType::OnTicked,
 			TaskRunStatus);
 
-		if (TaskRunStatus == EStateTreeRunStatus::Failed)
+		if (TaskRunStatus == EStateTreeRunStatus::Failed && Args.TasksCompletionStatus->IsConsideredForCompletion(OwnerTaskIndex))
 		{
 			bShouldTickTasks = false;
 		}
@@ -5270,7 +5272,7 @@ bool FStateTreeExecutionContext::SelectStateInternal(
 			__FUNCTION__, *GetSafeStateName(CurrentFrame, NextStateHandle), *GetStateStatusString(Exec), *GetNameSafe(&Owner), *GetFullNameSafe(CurrentFrame.StateTree));
 		return false;
 	}
-	CurrentFrame.ActiveTasksStatus.Push(CurrentStateTree, NextState);
+	CurrentFrame.ActiveTasksStatus.Push(NextState);
 
 	// Check if we're still tracking on the current active frame and state.
 	// If we are, update the NumCurrentlyActiveStates to indicate that this state's instance data can be accessed. 
diff --git a/Engine/Plugins/Runtime/StateTree/Source/StateTreeModule/Private/StateTreeTasksStatus.cpp b/Engine/Plugins/Runtime/StateTree/Source/StateTreeModule/Private/StateTreeTasksStatus.cpp
--- a/Engine/Plugins/Runtime/StateTree/Source/StateTreeModule/Private/StateTreeTasksStatus.cpp
+++ b/Engine/Plugins/Runtime/StateTree/Source/StateTreeModule/Private/StateTreeTasksStatus.cpp
@@ -14,6 +14,32 @@ namespace UE::StateTree::Private
 	}
 }
 
+namespace UE::StateTree
+{
+	template<typename T>
+	void TTasksCompletionStatus<T>::ResetStatus(int32 NumberOfTasksInTheCompletionMask)
+	{
+		T ClearBuffer = 0;
+		if (NumberOfTasksInTheCompletionMask == 0)
+		{
+			// All mask as at least one bit to mark the state.
+			ClearBuffer = 1 << BitIndex;
+		}
+		else if (NumberOfTasksInTheCompletionMask == MaxNumTasks)
+		{
+			// Prevent buffer overflow.
+			ClearBuffer = (T)(-1);
+		}
+		else
+		{
+			ClearBuffer = ((1 << NumberOfTasksInTheCompletionMask) - 1) << BitIndex;
+		}
+		const T ClearMask = ~ClearBuffer;
+		(*FirstCompletionBits) &= ClearMask;
+		(*SecondCompletionBits) &= ClearMask;
+	}
+}
+
 FStateTreeTasksCompletionStatus::FStateTreeTasksCompletionStatus(const FCompactStateTreeFrame& Frame)
 {
 	BufferNum = Frame.NumberOfTasksStatusMasks;
@@ -174,10 +200,10 @@ UE::StateTree::FConstTasksCompletionStatus FStateTreeTasksCompletionStatus::GetS
 		;
 }
 
-void FStateTreeTasksCompletionStatus::Push(TNotNull<const UStateTree*> StateTree, const FCompactStateTreeState& State)
+void FStateTreeTasksCompletionStatus::Push(const FCompactStateTreeState& State)
 {
 	check(BufferNum > State.CompletionTasksMaskBufferIndex);
-	GetStatus(State).ResetStatus();
+	GetStatus(State).ResetStatus(State.TasksNum);
 }
 
 bool FStateTreeTasksCompletionStatus::Serialize(FArchive& Ar)
diff --git a/Engine/Plugins/Runtime/StateTree/Source/StateTreeModule/Public/StateTreeTasksStatus.h b/Engine/Plugins/Runtime/StateTree/Source/StateTreeModule/Public/StateTreeTasksStatus.h
--- a/Engine/Plugins/Runtime/StateTree/Source/StateTreeModule/Public/StateTreeTasksStatus.h
+++ b/Engine/Plugins/Runtime/StateTree/Source/StateTreeModule/Public/StateTreeTasksStatus.h
@@ -104,6 +104,18 @@ namespace UE::StateTree
 			}
 		}
 
+		/**
+		 * @param StateTaskIndex The index of the task in the state.
+		 * @return true when the task is considered for GetCompletionStatus, HasAllCompleted, HasAnyCompleted, HasAnyFailed, IsCompleted
+		 */
+		[[nodiscard]] bool IsConsideredForCompletion(int32 StateTaskIndex) const
+		{
+			StateTaskIndex += BitIndex;
+			check(StateTaskIndex < MaxNumTasks);
+			const T TaskMask = T(1) << StateTaskIndex;
+			return (TaskMask & CompletionMask) != 0;
+		}
+
 		/**
 		 * @param StateTaskIndex The index of the task in the state.
 		 * @return true when the task status is running.
@@ -200,12 +212,7 @@ namespace UE::StateTree
 		}
 
 		/** Set the status of all tasks to running. */
-		void ResetStatus()
-		{
-			const T ClearMask = ~CompletionMask;
-			(*FirstCompletionBits) &= ClearMask;
-			(*SecondCompletionBits) &= ClearMask;
-		}
+		STATETREEMODULE_API void ResetStatus(int32 NumberOfTasksInTheCompletionMask);
 
 	private:
 		uint8 GetStatusInternal(int32 Index) const
@@ -281,7 +288,7 @@ struct FStateTreeTasksCompletionStatus
 
 public:
 	STATETREEMODULE_API FStateTreeTasksCompletionStatus() = default;
-	STATETREEMODULE_API FStateTreeTasksCompletionStatus(const FCompactStateTreeFrame& Frame);
+	STATETREEMODULE_API explicit FStateTreeTasksCompletionStatus(const FCompactStateTreeFrame& Frame);
 	STATETREEMODULE_API ~FStateTreeTasksCompletionStatus();
 
 	STATETREEMODULE_API FStateTreeTasksCompletionStatus(const FStateTreeTasksCompletionStatus&);
@@ -296,7 +303,7 @@ public:
 	[[nodiscard]] STATETREEMODULE_API UE::StateTree::FTasksCompletionStatus GetStatus(TNotNull<const UStateTree*> StateTree);
 	[[nodiscard]] STATETREEMODULE_API UE::StateTree::FConstTasksCompletionStatus GetStatus(TNotNull<const UStateTree*> StateTree) const;
 
-	STATETREEMODULE_API void Push(TNotNull<const UStateTree*> StateTree, const FCompactStateTreeState& State);
+	STATETREEMODULE_API void Push(const FCompactStateTreeState& State);
 
 	STATETREEMODULE_API bool Serialize(FArchive& Ar);
 	STATETREEMODULE_API bool NetSerialize(FArchive& Ar, UPackageMap* Map, bool& bOutSuccess);
diff --git a/Engine/Plugins/Runtime/StateTree/Source/StateTreeTestSuite/Private/StateTreeTaskStateTest.cpp b/Engine/Plugins/Runtime/StateTree/Source/StateTreeTestSuite/Private/StateTreeTaskStateTest.cpp
--- a/Engine/Plugins/Runtime/StateTree/Source/StateTreeTestSuite/Private/StateTreeTaskStateTest.cpp
+++ b/Engine/Plugins/Runtime/StateTree/Source/StateTreeTestSuite/Private/StateTreeTaskStateTest.cpp
@@ -37,6 +37,11 @@ struct FStateTreeTest_TasksCompletion_AllAny : FStateTreeTestBase
 			{
 				TStateTreeEditorNode<FTestTask_Stand>& Task = Root.AddTask<FTestTask_Stand>(FName(*FString::Printf(TEXT("Tree1StateRootTask_%d"), TaskIndex)));
 				Task.GetNode().TicksToCompletion = TaskIndex + 1;
+				if (TaskIndex == 10)
+				{
+					Task.GetNode().bConsideredForCompletion = false;
+					Task.GetNode().TickCompletionResult = EStateTreeRunStatus::Failed;
+				}
 			}
 			UStateTreeState& Tree1StateA = Root.AddChildState("Tree1StateA");
 			Tree1StateA.TasksCompletion = EStateTreeTaskCompletionType::All;
@@ -44,6 +49,12 @@ struct FStateTreeTest_TasksCompletion_AllAny : FStateTreeTestBase
 			{
 				TStateTreeEditorNode<FTestTask_Stand>& Task = Tree1StateA.AddTask<FTestTask_Stand>(FName(*FString::Printf(TEXT("Tree1StateATask_%d"), TaskIndex)));
 				Task.GetNode().TicksToCompletion = TaskIndex + 1;
+
+				if (TaskIndex == 22)
+				{
+					Task.GetNode().bConsideredForCompletion = false;
+					Task.GetNode().EnterStateResult = EStateTreeRunStatus::Failed;
+				}
 			}
 
 			{
@@ -86,7 +97,7 @@ struct FStateTreeTest_TasksCompletion_AllAny : FStateTreeTestBase
 						AITEST_TRUE(*FString::Printf(TEXT("Should tick Task %d, %d"), TickIndex, TaskIndex), bTickedRoot);
 					}
 					bool bTickedA = Exec.Expect(*FString::Printf(TEXT("Tree1StateATask_%d"), TaskIndex), TEXT("Tick"));
-					if (TaskIndex < TickIndex)
+					if (TaskIndex < TickIndex || TaskIndex == 22) // task 22 fails on enter
 					{
 						AITEST_FALSE(*FString::Printf(TEXT("Should not tick Task A %d, %d"), TickIndex, TaskIndex), bTickedA);
 					}
@@ -249,7 +260,7 @@ struct FStateTreeTest_TasksCompletion_Status : FStateTreeTestBase
 	{
 		FCompactStateTreeFrame Frame;
 		Frame.NumberOfTasksStatusMasks = 1;
-		FStateTreeTasksCompletionStatus Status = Frame;
+		FStateTreeTasksCompletionStatus Status = FStateTreeTasksCompletionStatus(Frame);
 
 		constexpr int32 NumberOfTasks = 4;
 		FCompactStateTreeState State;
