StateTree: Add a widget to select the completion type for state and global tasks. Add an option to disable the completion per schema. Add an icon to select tasks that are part of the completion for the state and the global task.

diff --git a/Engine/Plugins/Runtime/StateTree/Source/StateTreeEditorModule/Private/Customizations/StateTreeEditorDataDetails.cpp b/Engine/Plugins/Runtime/StateTree/Source/StateTreeEditorModule/Private/Customizations/StateTreeEditorDataDetails.cpp
--- a/Engine/Plugins/Runtime/StateTree/Source/StateTreeEditorModule/Private/Customizations/StateTreeEditorDataDetails.cpp
+++ b/Engine/Plugins/Runtime/StateTree/Source/StateTreeEditorModule/Private/Customizations/StateTreeEditorDataDetails.cpp
@@ -220,16 +220,22 @@ void FStateTreeEditorDataDetails::CustomizeDetails(IDetailLayoutBuilder& DetailB
 	check(GlobalTasksProperty.IsValid());
 	const FName GlobalTasksCategoryName(TEXT("Global Tasks"));
 
-	UE::StateTreeEditor::EditorNodeUtils::MakeArrayCategory(
+	const bool bAllowTasksCompletion = Schema && Schema->AllowTasksCompletion();
+	TSharedPtr<IPropertyHandle> GlobalTasksCompletionProperty = DetailBuilder.GetProperty(GET_MEMBER_NAME_CHECKED(UStateTreeEditorData, GlobalTasksCompletion));
+	GlobalTasksCompletionProperty->MarkHiddenByCustomization();
+
+	IDetailCategoryBuilder& GlobalTasksCategory = UE::StateTreeEditor::EditorNodeUtils::MakeArrayCategoryHeader(
 		DetailBuilder,
 		GlobalTasksProperty,
 		GlobalTasksCategoryName,
 		LOCTEXT("EditorDataDetailsGlobalTasks", "Global Tasks"),
 		FName("StateTreeEditor.Tasks"),
 		UE::StateTree::Colors::Cyan,
+		bAllowTasksCompletion ? GlobalTasksCompletionProperty->CreatePropertyValueWidget(/*bDisplayDefaultPropertyButtons*/false) : TSharedPtr<SWidget>(),
 		UE::StateTree::Colors::Cyan.WithAlpha(192),
 		LOCTEXT("EditorDataDetailsGlobalTasksAddTooltip", "Add new Global Task"),
 		/*SortOrder*/4);
+	UE::StateTreeEditor::EditorNodeUtils::MakeArrayItems(GlobalTasksCategory, GlobalTasksProperty);
 
 	// Refresh the UI when the Schema changes.	
 	TSharedPtr<IPropertyHandle> SchemaProperty = DetailBuilder.GetProperty(GET_MEMBER_NAME_CHECKED(UStateTreeEditorData, Schema));
diff --git a/Engine/Plugins/Runtime/StateTree/Source/StateTreeEditorModule/Private/Customizations/StateTreeEditorNodeDetails.cpp b/Engine/Plugins/Runtime/StateTree/Source/StateTreeEditorModule/Private/Customizations/StateTreeEditorNodeDetails.cpp
--- a/Engine/Plugins/Runtime/StateTree/Source/StateTreeEditorModule/Private/Customizations/StateTreeEditorNodeDetails.cpp
+++ b/Engine/Plugins/Runtime/StateTree/Source/StateTreeEditorModule/Private/Customizations/StateTreeEditorNodeDetails.cpp
@@ -406,6 +406,7 @@ void FStateTreeEditorNodeDetails::CustomizeHeader(TSharedRef<class IPropertyHand
 	};
 
 	TSharedPtr<SBorder> FlagBorder;
+	TSharedPtr<SHorizontalBox> DescriptionBox;
 
 	HeaderRow
 		.WholeRowContent()
@@ -514,13 +515,13 @@ void FStateTreeEditorNodeDetails::CustomizeHeader(TSharedRef<class IPropertyHand
 				.VAlign(VAlign_Center)
 				.Padding(FMargin(0, 4))
 				[
-					SNew(SHorizontalBox)
+					SAssignNew(DescriptionBox, SHorizontalBox)
 
 					// Icon
 					+ SHorizontalBox::Slot()
 					.AutoWidth()
 					.VAlign(VAlign_Center)
-					.Padding(0, 0, 4, 0)
+					.Padding(0, 0, 0, 0)
 					[
 						SNew(SImage)
 						.Image(this, &FStateTreeEditorNodeDetails::GetIcon)
@@ -672,6 +673,31 @@ void FStateTreeEditorNodeDetails::CustomizeHeader(TSharedRef<class IPropertyHand
 		.CopyAction(FUIAction(FExecuteAction::CreateSP(this, &FStateTreeEditorNodeDetails::OnCopyNode)))
 		.PasteAction(FUIAction(FExecuteAction::CreateSP(this, &FStateTreeEditorNodeDetails::OnPasteNode)));
 
+	// Task completion
+	if (const FStateTreeEditorNode * Node = UE::StateTreeEditor::EditorNodeUtils::GetCommonNode(StructProperty))
+	{
+		if (const FStateTreeTaskBase* TaskBase = Node->Node.GetPtr<FStateTreeTaskBase>())
+		{
+			DescriptionBox->InsertSlot(0)
+				.AutoWidth()
+				.VAlign(VAlign_Center)
+				.Padding(0.f, 0.f, 0.f, 0.f)
+				[
+					// Create the toggle favorites button
+					SNew(SButton)
+					.ButtonStyle(FAppStyle::Get(), "SimpleButton")
+					.OnClicked(this, &FStateTreeEditorNodeDetails::HandleToggleCompletionTaskClicked)
+					.ToolTipText(LOCTEXT("ToggleTaskCompletion", "Toggle Completion.\nA task can be considered for state completion or run in the background without affecting the state completion."))
+					[
+						SNew(SImage)
+						.ColorAndOpacity(this, &FStateTreeEditorNodeDetails::GetToggleCompletionTaskColor)
+						.Image(FAppStyle::Get().GetBrush("Icons.Star"))
+					]
+					.IsEnabled(TaskBase->bCanEditConsideredForCompletion)
+				];
+		}
+	}
+
 	MakeFlagsWidget();
 }
 
@@ -1592,6 +1618,33 @@ void FStateTreeEditorNodeDetails::HandleNameCommitted(const FText& NewText, ETex
 	NameSwitcher->SetActiveWidgetIndex(0);
 }
 
+FReply FStateTreeEditorNodeDetails::HandleToggleCompletionTaskClicked()
+{
+	if (FStateTreeEditorNode* Node = UE::StateTreeEditor::EditorNodeUtils::GetMutableCommonNode(StructProperty))
+	{
+		if (FStateTreeTaskBase* TaskBase = Node->Node.GetMutablePtr<FStateTreeTaskBase>())
+		{
+			TaskBase->bConsideredForCompletion = !TaskBase->bConsideredForCompletion;
+		}
+	}
+	return FReply::Handled();
+}
+
+FSlateColor FStateTreeEditorNodeDetails::GetToggleCompletionTaskColor() const
+{
+	if (const FStateTreeEditorNode* Node = UE::StateTreeEditor::EditorNodeUtils::GetCommonNode(StructProperty))
+	{
+		if (const FStateTreeTaskBase* TaskBase = Node->Node.GetPtr<FStateTreeTaskBase>())
+		{
+			if (TaskBase->bConsideredForCompletion)
+			{
+				return FSlateColor(EStyleColor::Foreground);
+			}
+		}
+	}
+	return FSlateColor(EStyleColor::Recessed);
+}
+
 FText FStateTreeEditorNodeDetails::GetNodePickerTooltip() const
 {
 	check(StructProperty);
diff --git a/Engine/Plugins/Runtime/StateTree/Source/StateTreeEditorModule/Private/Customizations/StateTreeEditorNodeDetails.h b/Engine/Plugins/Runtime/StateTree/Source/StateTreeEditorModule/Private/Customizations/StateTreeEditorNodeDetails.h
--- a/Engine/Plugins/Runtime/StateTree/Source/StateTreeEditorModule/Private/Customizations/StateTreeEditorNodeDetails.h
+++ b/Engine/Plugins/Runtime/StateTree/Source/StateTreeEditorModule/Private/Customizations/StateTreeEditorNodeDetails.h
@@ -100,6 +100,9 @@ private:
 	bool HandleVerifyNameChanged(const FText& InText, FText& OutErrorMessage) const;
 	void HandleNameCommitted(const FText& NewLabel, ETextCommit::Type CommitType) const;
 
+	FReply HandleToggleCompletionTaskClicked();
+	FSlateColor GetToggleCompletionTaskColor() const;
+
 	FText GetNodePickerTooltip() const;
 	void OnNodePicked(const UStruct* InStruct) const;
 
diff --git a/Engine/Plugins/Runtime/StateTree/Source/StateTreeEditorModule/Private/Customizations/StateTreeEditorNodeUtils.cpp b/Engine/Plugins/Runtime/StateTree/Source/StateTreeEditorModule/Private/Customizations/StateTreeEditorNodeUtils.cpp
--- a/Engine/Plugins/Runtime/StateTree/Source/StateTreeEditorModule/Private/Customizations/StateTreeEditorNodeUtils.cpp
+++ b/Engine/Plugins/Runtime/StateTree/Source/StateTreeEditorModule/Private/Customizations/StateTreeEditorNodeUtils.cpp
@@ -794,6 +794,34 @@ IDetailCategoryBuilder& MakeArrayCategory(
 	const FLinearColor AddIconColor,
 	const FText& AddButtonTooltipText,
 	const int32 SortOrder)
+{
+	IDetailCategoryBuilder& Category = MakeArrayCategoryHeader(
+		DetailBuilder,
+		ArrayPropertyHandle,
+		CategoryName,
+		CategoryDisplayName,
+		IconName,
+		IconColor,
+		TSharedPtr<SWidget>(),
+		AddIconColor,
+		AddButtonTooltipText,
+		SortOrder
+	);
+	MakeArrayItems(Category, ArrayPropertyHandle);
+	return Category;
+}
+
+IDetailCategoryBuilder& MakeArrayCategoryHeader(
+	IDetailLayoutBuilder& DetailBuilder,
+	const TSharedPtr<IPropertyHandle>& ArrayPropertyHandle,
+	const FName CategoryName,
+	const FText& CategoryDisplayName,
+	const FName IconName,
+	const FLinearColor IconColor,
+	const TSharedPtr<SWidget> Extension,
+	const FLinearColor AddIconColor,
+	const FText& AddButtonTooltipText,
+	const int32 SortOrder)
 {
 	IDetailCategoryBuilder& Category = DetailBuilder.EditCategory(CategoryName, CategoryDisplayName);
 	Category.SetSortOrder(SortOrder);
@@ -835,17 +863,30 @@ IDetailCategoryBuilder& MakeArrayCategory(
 	}
 	
 	HeaderContent->AddSlot()
+		.FillWidth(1.f)
+		.HAlign(HAlign_Left)
 		.VAlign(VAlign_Center)
-		.AutoWidth()
 		.Padding(FMargin(4, 0, 0, 0))
 		[
 			SNew(STextBlock)
 			.TextStyle(FStateTreeEditorStyle::Get(), "StateTree.Category")
 			.Text(CategoryDisplayName)
 		];
-			
+
+	if (Extension)
+	{
+		HeaderContent->AddSlot()
+			.AutoWidth()
+			.HAlign(HAlign_Right)
+			.VAlign(VAlign_Center)
+			[
+				Extension.ToSharedRef()
+			];
+	}
+
+
 	HeaderContent->AddSlot()
-		.FillWidth(1.f)
+		.AutoWidth()
 		.HAlign(HAlign_Right)
 		.VAlign(VAlign_Center)
 		[
@@ -858,7 +899,11 @@ IDetailCategoryBuilder& MakeArrayCategory(
 			HeaderContent
 		],
 		/*bWholeRowContent*/true);
+	return Category;
+}
 
+void MakeArrayItems(IDetailCategoryBuilder& Category, const TSharedPtr<IPropertyHandle>& ArrayPropertyHandle)
+{
 	// Add items inline
 	const TSharedRef<FDetailArrayBuilder> Builder = MakeShareable(new FDetailArrayBuilder(ArrayPropertyHandle.ToSharedRef(), /*InGenerateHeader*/ false, /*InDisplayResetToDefault*/ true, /*InDisplayElementNum*/ false));
 	Builder->OnGenerateArrayElementWidget(FOnGenerateArrayElementWidget::CreateLambda([](TSharedRef<IPropertyHandle> PropertyHandle, int32 ArrayIndex, IDetailChildrenBuilder& ChildrenBuilder)
@@ -866,8 +911,6 @@ IDetailCategoryBuilder& MakeArrayCategory(
 		ChildrenBuilder.AddProperty(PropertyHandle);
 	}));
 	Category.AddCustomBuilder(Builder, /*bForAdvanced*/ false);
-
-	return Category;
 }
 
 bool ImportTextAsNode(const UScriptStruct* BaseScriptStruct, const UStateTreeEditorData* EditorData, FStateTreeEditorNode& OutEditorNode)
diff --git a/Engine/Plugins/Runtime/StateTree/Source/StateTreeEditorModule/Private/Customizations/StateTreeEditorNodeUtils.h b/Engine/Plugins/Runtime/StateTree/Source/StateTreeEditorModule/Private/Customizations/StateTreeEditorNodeUtils.h
--- a/Engine/Plugins/Runtime/StateTree/Source/StateTreeEditorModule/Private/Customizations/StateTreeEditorNodeUtils.h
+++ b/Engine/Plugins/Runtime/StateTree/Source/StateTreeEditorModule/Private/Customizations/StateTreeEditorNodeUtils.h
@@ -18,6 +18,7 @@ class IDetailCategoryBuilder;
 class IDetailLayoutBuilder;
 template <typename FuncType> class TFunctionRef;
 class SComboButton;
+class SWidget;
 
 namespace UE::StateTreeEditor::EditorNodeUtils
 {
@@ -160,6 +161,38 @@ IDetailCategoryBuilder& MakeArrayCategory(
 	const FText& AddButtonTooltipText,
 	int32 SortOrder);
 
+/**
+ * Creates a category and sets the contents of the row to: [Icon] [DisplayName]  [+].
+ * @param DetailBuilder Detail builder where the category is added.
+ * @param ArrayPropertyHandle Property handle pointing to a TArray<FStateTreeEditorNode>.
+ * @param CategoryName Name of the category to create.
+ * @param CategoryDisplayName Display name of the category.
+ * @param IconName Name to the icon resource to show in front of the category display name.
+ * @param IconColor Color of icon resource.
+ * @param AddIconColor Color of the add icon.
+ * @param AddButtonTooltipText Tooltip text to show on the add buttons.
+ * @param SortOrder Category sort order (categories with smaller sort order appear earlier in details panel).
+ * @return Reference to the category builder of the category just created.
+ */
+IDetailCategoryBuilder& MakeArrayCategoryHeader(
+	IDetailLayoutBuilder& DetailBuilder,
+	const TSharedPtr<IPropertyHandle>& ArrayPropertyHandle,
+	FName CategoryName,
+	const FText& CategoryDisplayName,
+	FName IconName,
+	FLinearColor IconColor,
+	const TSharedPtr<SWidget> Extension,
+	FLinearColor AddIconColor,
+	const FText& AddButtonTooltipText,
+	int32 SortOrder);
+
+/**
+ * Creates the items of a category.
+ * @param Category Builder created by the header.
+ * @param ArrayPropertyHandle Property handle pointing to a TArray<FStateTreeEditorNode>.
+ */
+void MakeArrayItems(IDetailCategoryBuilder& Category, const TSharedPtr<IPropertyHandle>& ArrayPropertyHandle);
+
 /**
  * Import the text from clipboard as an Editor Node
  * @param BaseScriptStruct the Base Struct of the Editor Node's Node
diff --git a/Engine/Plugins/Runtime/StateTree/Source/StateTreeEditorModule/Private/Customizations/StateTreeStateDetails.cpp b/Engine/Plugins/Runtime/StateTree/Source/StateTreeEditorModule/Private/Customizations/StateTreeStateDetails.cpp
--- a/Engine/Plugins/Runtime/StateTree/Source/StateTreeEditorModule/Private/Customizations/StateTreeStateDetails.cpp
+++ b/Engine/Plugins/Runtime/StateTree/Source/StateTreeEditorModule/Private/Customizations/StateTreeStateDetails.cpp
@@ -83,6 +83,7 @@ void FStateTreeStateDetails::CustomizeDetails(IDetailLayoutBuilder& DetailBuilde
 	const TSharedPtr<IPropertyHandle> CustomTickRateProperty = DetailBuilder.GetProperty(GET_MEMBER_NAME_CHECKED(UStateTreeState, CustomTickRate));
 	const TSharedPtr<IPropertyHandle> ParametersProperty = DetailBuilder.GetProperty(GET_MEMBER_NAME_CHECKED(UStateTreeState, Parameters));
 	const TSharedPtr<IPropertyHandle> SelectionBehaviorProperty = DetailBuilder.GetProperty(GET_MEMBER_NAME_CHECKED(UStateTreeState, SelectionBehavior));
+	const TSharedPtr<IPropertyHandle> TasksCompletionProperty = DetailBuilder.GetProperty(GET_MEMBER_NAME_CHECKED(UStateTreeState, TasksCompletion));
 	const TSharedPtr<IPropertyHandle> RequiredEventToEnterProperty = DetailBuilder.GetProperty(GET_MEMBER_NAME_CHECKED(UStateTreeState, RequiredEventToEnter));
 	const TSharedPtr<IPropertyHandle> CheckPrerequisitesProperty = DetailBuilder.GetProperty(GET_MEMBER_NAME_CHECKED(UStateTreeState, bCheckPrerequisitesWhenActivatingChildDirectly));
 	const TSharedPtr<IPropertyHandle> WeightProperty = DetailBuilder.GetProperty(GET_MEMBER_NAME_CHECKED(UStateTreeState, Weight));
@@ -174,6 +175,7 @@ void FStateTreeStateDetails::CustomizeDetails(IDetailLayoutBuilder& DetailBuilde
 
 	// Per state type properties
 	SelectionBehaviorProperty->MarkHiddenByCustomization();
+	TasksCompletionProperty->MarkHiddenByCustomization();
 	LinkedSubtreeProperty->MarkHiddenByCustomization();
 	LinkedAssetProperty->MarkHiddenByCustomization();
 
@@ -330,18 +332,21 @@ void FStateTreeStateDetails::CustomizeDetails(IDetailLayoutBuilder& DetailBuilde
 	{
 		if (Schema && Schema->AllowMultipleTasks())
 		{
+			const bool bAllowTasksCompletion = Schema->AllowTasksCompletion();
 			const FName TasksCategoryName(TEXT("Tasks"));
-			UE::StateTreeEditor::EditorNodeUtils::MakeArrayCategory(
+			IDetailCategoryBuilder& TasksCategory = UE::StateTreeEditor::EditorNodeUtils::MakeArrayCategoryHeader(
 				DetailBuilder,
 				TasksProperty,
 				TasksCategoryName,
 				LOCTEXT("StateDetailsTasks", "Tasks"),
 				FName("StateTreeEditor.Tasks"),
 				UE::StateTree::Colors::Cyan,
+				bAllowTasksCompletion ? TasksCompletionProperty->CreatePropertyValueWidget(/*bDisplayDefaultPropertyButtons*/false) : TSharedPtr<SWidget>(),
 				UE::StateTree::Colors::Cyan.WithAlpha(192),
 				LOCTEXT("StateDetailsTasksAddTooltip", "Add new Task"),
 				/*SortOrder*/4);
 			SingleTaskProperty->MarkHiddenByCustomization();
+			UE::StateTreeEditor::EditorNodeUtils::MakeArrayItems(TasksCategory, TasksProperty);
 		}
 		else
 		{
diff --git a/Engine/Plugins/Runtime/StateTree/Source/StateTreeModule/Public/StateTreeSchema.h b/Engine/Plugins/Runtime/StateTree/Source/StateTreeModule/Public/StateTreeSchema.h
--- a/Engine/Plugins/Runtime/StateTree/Source/StateTreeModule/Public/StateTreeSchema.h
+++ b/Engine/Plugins/Runtime/StateTree/Source/StateTreeModule/Public/StateTreeSchema.h
@@ -105,6 +105,12 @@ public:
 		return true;
 	}
 
+	/** @return True if modifying the tasks completion is allowed. */
+	virtual bool AllowTasksCompletion() const
+	{
+		return true;
+	}
+
 	virtual const FStateTreeDetailsCustomizationSettings& GetDetailsCustomizationSettings() const
 	{
 		static FStateTreeDetailsCustomizationSettings Settings;
