StateTree: Mark debug task not considered for completion by default

diff --git a/Engine/Plugins/Runtime/StateTree/Source/StateTreeModule/Private/Tasks/StateTreeDebugTextTask.cpp b/Engine/Plugins/Runtime/StateTree/Source/StateTreeModule/Private/Tasks/StateTreeDebugTextTask.cpp
--- a/Engine/Plugins/Runtime/StateTree/Source/StateTreeModule/Private/Tasks/StateTreeDebugTextTask.cpp
+++ b/Engine/Plugins/Runtime/StateTree/Source/StateTreeModule/Private/Tasks/StateTreeDebugTextTask.cpp
@@ -15,6 +15,11 @@ FStateTreeDebugTextTask::FStateTreeDebugTextTask()
 	// We do not want to change the ReferenceActor if it's bound.
 	bShouldCopyBoundPropertiesOnTick = false;
 	bShouldCopyBoundPropertiesOnExitState = false;
+
+#if WITH_EDITORONLY_DATA
+	bConsideredForCompletion = false;
+	bCanEditConsideredForCompletion = false;
+#endif
 }
 
 EStateTreeRunStatus FStateTreeDebugTextTask::EnterState(FStateTreeExecutionContext& Context, const FStateTreeTransitionResult& Transition) const
