StateTree: Reset the compiled data when the editor data is copied. It prevents old data in the compiled data.

diff --git a/Engine/Plugins/Runtime/StateTree/Source/StateTreeEditorModule/Private/StateTreeEditorData.cpp b/Engine/Plugins/Runtime/StateTree/Source/StateTreeEditorModule/Private/StateTreeEditorData.cpp
--- a/Engine/Plugins/Runtime/StateTree/Source/StateTreeEditorModule/Private/StateTreeEditorData.cpp
+++ b/Engine/Plugins/Runtime/StateTree/Source/StateTreeEditorModule/Private/StateTreeEditorData.cpp
@@ -342,6 +342,7 @@ void UStateTreeEditorData::PostDuplicate(EDuplicateMode::Type DuplicateMode)
 
 	EditorBindings.SetBindingsOwner(this);
 	DuplicateIDs();
+	GetOuterUStateTree()->ResetCompiled();
 }
 #endif // WITH_EDITOR
 
