StateTree: Fix active tasks on ForceTransition.
#rnx

diff --git a/Engine/Plugins/Runtime/StateTree/Source/StateTreeModule/Private/StateTreeExecutionContext.cpp b/Engine/Plugins/Runtime/StateTree/Source/StateTreeModule/Private/StateTreeExecutionContext.cpp
--- a/Engine/Plugins/Runtime/StateTree/Source/StateTreeModule/Private/StateTreeExecutionContext.cpp
+++ b/Engine/Plugins/Runtime/StateTree/Source/StateTreeModule/Private/StateTreeExecutionContext.cpp
@@ -2686,11 +2686,17 @@ EStateTreeRunStatus FStateTreeExecutionContext::EnterState(FStateTreeTransitionR
 		
 		FStateTreeExecutionFrame* CurrentParentFrame = !Exec.ActiveFrames.IsEmpty() ? &Exec.ActiveFrames.Last() : nullptr;
 		FStateTreeExecutionFrame& CurrentFrame = Exec.ActiveFrames.Add_GetRef(NextFrame);
+		const UStateTree* CurrentStateTree = NextFrame.StateTree;
 		const UE::StateTree::FActiveFrameID CurrentFrameID = NextFrame.FrameID;
 		
 		// We'll add new states one by one, so that active states contain only the states which have EnterState called.
 		CurrentFrame.ActiveStates.Reset();
 
+		// Create new status
+		const FCompactStateTreeFrame* FrameInfo = NextFrame.StateTree->GetFrameFromHandle(NextFrame.RootState);
+		ensureAlwaysMsgf(FrameInfo, TEXT("The compiled data is invalid. It should contains the information for the root frame."));
+		CurrentFrame.ActiveTasksStatus = FrameInfo ? FStateTreeTasksCompletionStatus(*FrameInfo) : FStateTreeTasksCompletionStatus();
+
 		// Get previous active states, they are used to calculate transition type.
 		FStateTreeActiveStates PreviousActiveStates;
 		if (PreviousActiveFrames.IsValidIndex(FrameIndex)
@@ -2700,7 +2706,6 @@ EStateTreeRunStatus FStateTreeExecutionContext::EnterState(FStateTreeTransitionR
 		}
 
 		FCurrentlyProcessedFrameScope FrameScope(*this, CurrentParentFrame, CurrentFrame);
-		const UStateTree* CurrentStateTree = NextFrame.StateTree;
 
 		for (int32 StateIndex = 0; StateIndex < NextFrame.ActiveStates.Num() && Result != EStateTreeRunStatus::Failed; StateIndex++)
 		{
@@ -2716,6 +2721,7 @@ EStateTreeRunStatus FStateTreeExecutionContext::EnterState(FStateTreeTransitionR
 					__FUNCTION__, *GetStateStatusString(Exec), *GetNameSafe(&Owner), *GetFullNameSafe(&RootStateTree));
 				break;
 			}
+			//@todo push the same state as previously handle
 			CurrentFrame.ActiveTasksStatus.Push(State);
 
 			UE::StateTree::FTasksCompletionStatus CurrentStateTasksStatus = CurrentFrame.ActiveTasksStatus.GetStatus(State);
