StateTree: Global tasks complete the tree they are part of. Previously, the behavior was that any global task completes the root execution. You can revert to the previous behavior with the cvar StateTree.GlobalTasksCompleteOwningFrame.

diff --git a/Engine/Plugins/Runtime/StateTree/Source/StateTreeModule/Private/StateTreeExecutionContext.cpp b/Engine/Plugins/Runtime/StateTree/Source/StateTreeModule/Private/StateTreeExecutionContext.cpp
--- a/Engine/Plugins/Runtime/StateTree/Source/StateTreeModule/Private/StateTreeExecutionContext.cpp
+++ b/Engine/Plugins/Runtime/StateTree/Source/StateTreeModule/Private/StateTreeExecutionContext.cpp
@@ -46,6 +46,15 @@ namespace UE::StateTree::ExecutionContext::Private
 		TEXT("If false (5.5 behavior), then all the global nodes, from all linked assets, are evaluated, then the state tasks are ticked.\n")
 		TEXT("You should upgrade your asset. This option is to help migrate to the new behavior.")
 	);
+	
+	bool bGlobalTasksCompleteOwningFrame = true;
+	FAutoConsoleVariableRef CVarGlobalTasksCompleteFrame(
+		TEXT("StateTree.GlobalTasksCompleteOwningFrame"),
+		bGlobalTasksCompleteOwningFrame,
+		TEXT("If true, the global tasks complete the tree they are part of. The root or the linked state. 5.6 behavior.\n")
+		TEXT("If false, the global tasks complete the root tree execution (even for linked assets). 5.5 behavior.")
+		TEXT("You should upgrade your asset. This option is to help migrate to the new behavior.")
+	);
 
 	bool MarkDelegateAsBroadcasted(FStateTreeDelegateDispatcher Dispatcher, const FStateTreeExecutionFrame& CurrentFrame, FStateTreeInstanceStorage& Storage)
 	{
@@ -1131,7 +1140,8 @@ EStateTreeRunStatus FStateTreeExecutionContext::Start(FStartParameters Parameter
 	if (GlobalTasksRunStatus == EStateTreeRunStatus::Running)
 	{
 		// First tick.
-		// Tasks are not ticked here, since their behavior is that EnterState() (called above) is treated as a tick.  
+		// Tasks are not ticked here, since their behavior is that EnterState() (called above) is treated as a tick.
+		//@todo Check the result of TickEvaluatorsAndGlobalTasks and early exit if it is not running
 		constexpr bool bTickGlobalTasks = false;
 		TickEvaluatorsAndGlobalTasks(0.0f, bTickGlobalTasks);
 
@@ -1453,9 +1463,29 @@ void FStateTreeExecutionContext::TickUpdateTasksInternal(float DeltaTime)
 		}
 		else
 		{
-			//@TODO check what we want to do with global task from linked asset state. Should they request a stop or close their own tree.
-			Exec.RequestedStop = UE::StateTree::ExecutionContext::Private::GetPriorityRunStatus(Exec.RequestedStop, EvalAndGlobalTaskStatus);
-			LogRequestStop();
+			using namespace UE::StateTree;
+			using namespace UE::StateTree::ExecutionContext::Private;
+			if (bGlobalTasksCompleteOwningFrame)
+			{
+				// Only set RequestStop if it's the first frame (root)
+				check(Exec.ActiveFrames.Num() > 0);
+				const UStateTree* StateTree = Exec.ActiveFrames[0].StateTree;
+				check(StateTree == &RootStateTree);
+				const ETaskCompletionStatus GlobalTaskStatus = Exec.ActiveFrames[0].ActiveTasksStatus.GetStatus(StateTree).GetCompletionStatus();
+				const EStateTreeRunStatus GlobalRunStatus = CastToRunStatus(GlobalTaskStatus);
+				if (GlobalRunStatus != EStateTreeRunStatus::Running)
+				{
+					// Note. Exec.RequestedStop default value is Unset.
+					Exec.RequestedStop = GetPriorityRunStatus(Exec.RequestedStop, GlobalRunStatus);
+					LogRequestStop();
+				}
+			}
+			else
+			{
+				// Any completion stops the tree execution.
+				Exec.RequestedStop = GetPriorityRunStatus(Exec.RequestedStop, EvalAndGlobalTaskStatus);
+				LogRequestStop();
+			}
 		}
 	}
 }
@@ -1702,7 +1732,7 @@ void FStateTreeExecutionContext::FinishTask(const FStateTreeTaskBase& Task, ESta
 
 		UE::StateTree::FTasksCompletionStatus StateTasksStatus = const_cast<FStateTreeExecutionFrame*>(CurrentlyProcessedFrame)->ActiveTasksStatus.GetStatus(State);
 		StateTasksStatus.SetStatusWithPriority(StateTaskIndex, TaskStatus);
-		Exec.bHasPendingCompletedState = StateTasksStatus.IsCompleted();
+		Exec.bHasPendingCompletedState = Exec.bHasPendingCompletedState || StateTasksStatus.IsCompleted();
 	}
 	else
 	{
@@ -1711,7 +1741,7 @@ void FStateTreeExecutionContext::FinishTask(const FStateTreeTaskBase& Task, ESta
 		check(FrameTaskIndex >= 0);
 		UE::StateTree::FTasksCompletionStatus GlobalTasksStatus = const_cast<FStateTreeExecutionFrame*>(CurrentlyProcessedFrame)->ActiveTasksStatus.GetStatus(CurrentStateTree);
 		GlobalTasksStatus.SetStatusWithPriority(FrameTaskIndex, TaskStatus);
-		Exec.bHasPendingCompletedState = GlobalTasksStatus.IsCompleted();
+		Exec.bHasPendingCompletedState = Exec.bHasPendingCompletedState || GlobalTasksStatus.IsCompleted();
 	}
 }
 
@@ -3290,6 +3320,7 @@ EStateTreeRunStatus FStateTreeExecutionContext::TickEvaluatorsAndGlobalTasks(con
 		}
 	}
 
+	Exec.bHasPendingCompletedState = Exec.bHasPendingCompletedState || Result != EStateTreeRunStatus::Running;
 	return Result;
 }
 
@@ -3785,7 +3816,7 @@ EStateTreeRunStatus FStateTreeExecutionContext::TickTasks(const float DeltaTime)
 
 		FCurrentlyProcessedFrameScope FrameScope(*this, TickArgs.ParentFrame, *TickArgs.Frame);
 
-		if (UE::StateTree::ExecutionContext::Private::bTickGlobalNodesFollowingTreeHierarchy)
+		if (ExecutionContext::Private::bTickGlobalNodesFollowingTreeHierarchy)
 		{
 			if (TickArgs.Frame->bIsGlobalFrame)
 			{
@@ -3793,7 +3824,11 @@ EStateTreeRunStatus FStateTreeExecutionContext::TickTasks(const float DeltaTime)
 				const EStateTreeRunStatus FrameResult = TickEvaluatorsAndGlobalTasksForFrame(DeltaTime, bTickGlobalTasks, FrameIndex, TickArgs.ParentFrame, TickArgs.Frame);
 				if (FrameResult != EStateTreeRunStatus::Running)
 				{
-					Exec.RequestedStop = ExecutionContext::Private::GetPriorityRunStatus(Exec.RequestedStop, FrameResult);
+					if (ExecutionContext::Private::bGlobalTasksCompleteOwningFrame == false || FrameIndex == 0)
+					{
+						// Stop the tree execution when it's the root frame or if the previous behavior is desired.
+						Exec.RequestedStop = ExecutionContext::Private::GetPriorityRunStatus(Exec.RequestedStop, FrameResult);
+					}
 					TickArgs.bShouldTickTasks = false;
 					break;
 				}
@@ -3804,7 +3839,7 @@ EStateTreeRunStatus FStateTreeExecutionContext::TickTasks(const float DeltaTime)
 		{
 			const FStateTreeStateHandle CurrentHandle = TickArgs.Frame->ActiveStates[StateIndex];
 			const FCompactStateTreeState& CurrentState = CurrentStateTree->States[CurrentHandle.Index];
-			UE::StateTree::FTasksCompletionStatus CurrentCompletionStatus = TickArgs.Frame->ActiveTasksStatus.GetStatus(CurrentState);
+			FTasksCompletionStatus CurrentCompletionStatus = TickArgs.Frame->ActiveTasksStatus.GetStatus(CurrentState);
 
 			TickArgs.StateID = TickArgs.Frame->ActiveStates.StateIDs[StateIndex];
 			TickArgs.TasksCompletionStatus = &CurrentCompletionStatus;
@@ -3813,7 +3848,7 @@ EStateTreeRunStatus FStateTreeExecutionContext::TickTasks(const float DeltaTime)
 			UE_STATETREE_DEBUG_SCOPED_STATE(this, CurrentHandle);
 
 			STATETREE_CLOG(CurrentState.TasksNum > 0, VeryVerbose, TEXT("%*sState '%s'")
-				, (FrameIndex + StateIndex + 1) * UE::StateTree::Debug::IndentSize, TEXT("")
+				, (FrameIndex + StateIndex + 1) * Debug::IndentSize, TEXT("")
 				, *DebugGetStatePath(Exec.ActiveFrames, TickArgs.Frame, StateIndex));
 
 			if (CurrentState.Type == EStateTreeStateType::Linked || CurrentState.Type == EStateTreeStateType::LinkedAsset)
@@ -3853,6 +3888,7 @@ EStateTreeRunStatus FStateTreeExecutionContext::TickTasks(const float DeltaTime)
 
 	// Collect the result after every tasks has the chance to tick.
 	//An async or delegate might complete a global or "previous" task (in a different order).
+	EStateTreeRunStatus FirstFrameResult = EStateTreeRunStatus::Running;
 	EStateTreeRunStatus FrameResult = EStateTreeRunStatus::Running;
 	EStateTreeRunStatus StateResult = EStateTreeRunStatus::Running;
 	for (int32 FrameIndex = 0; FrameIndex < Exec.ActiveFrames.Num(); ++FrameIndex)
@@ -3864,8 +3900,13 @@ EStateTreeRunStatus FStateTreeExecutionContext::TickTasks(const float DeltaTime)
 		if (CurrentFrame.bIsGlobalFrame)
 		{
 			const ETaskCompletionStatus GlobalTasksStatus = CurrentFrame.ActiveTasksStatus.GetStatus(CurrentStateTree).GetCompletionStatus();
+			if (FrameIndex == 0)
+			{
+				FirstFrameResult = CastToRunStatus(GlobalTasksStatus);
+			}
 			FrameResult = GetPriorityRunStatus(FrameResult, CastToRunStatus(GlobalTasksStatus));
 		}
+
 		for (int32 StateIndex = 0; StateIndex < CurrentFrame.ActiveStates.Num() && StateResult != EStateTreeRunStatus::Failed; ++StateIndex)
 		{
 			const FStateTreeStateHandle CurrentHandle = CurrentFrame.ActiveStates[StateIndex];
@@ -3874,9 +3915,12 @@ EStateTreeRunStatus FStateTreeExecutionContext::TickTasks(const float DeltaTime)
 			StateResult = GetPriorityRunStatus(StateResult, CastToRunStatus(StateTasksStatus));
 		}
 	}
-
-	//@todo frame should only close their own
-	if (FrameResult != EStateTreeRunStatus::Running)
+	
+	if (ExecutionContext::Private::bGlobalTasksCompleteOwningFrame && FirstFrameResult != EStateTreeRunStatus::Running)
+	{
+		Exec.RequestedStop = ExecutionContext::Private::GetPriorityRunStatus(Exec.RequestedStop, FrameResult);
+	}
+	else if (ExecutionContext::Private::bGlobalTasksCompleteOwningFrame == false && FrameResult != EStateTreeRunStatus::Running)
 	{
 		Exec.RequestedStop = ExecutionContext::Private::GetPriorityRunStatus(Exec.RequestedStop, FrameResult);
 	}
@@ -3888,16 +3932,16 @@ EStateTreeRunStatus FStateTreeExecutionContext::TickTasks(const float DeltaTime)
 		{
 			FStateTreeExecutionFrame& LastFrame = Exec.ActiveFrames.Last();
 			const int32 NumberOfActiveState = LastFrame.ActiveStates.Num();
-			if (NumberOfActiveState == 0)
-			{
-				LastFrame.ActiveTasksStatus.GetStatus(LastFrame.StateTree).SetCompletionStatus(ETaskCompletionStatus::Succeeded);
-			}
-			else
+			if (ensureMsgf(NumberOfActiveState != 0, TEXT("No task is allowed to clear/stop/transition. Those action should be delayed inside the execution context.")))
 			{
 				const FStateTreeStateHandle CurrentHandle = LastFrame.ActiveStates[NumberOfActiveState - 1];
 				const FCompactStateTreeState& State = LastFrame.StateTree->States[CurrentHandle.Index];
 				LastFrame.ActiveTasksStatus.GetStatus(State).SetCompletionStatus(ETaskCompletionStatus::Succeeded);
 			}
+			else
+			{
+				LastFrame.ActiveTasksStatus.GetStatus(LastFrame.StateTree).SetCompletionStatus(ETaskCompletionStatus::Succeeded);
+			}
 		}
 		else
 		{
@@ -3906,7 +3950,7 @@ EStateTreeRunStatus FStateTreeExecutionContext::TickTasks(const float DeltaTime)
 		StateResult = EStateTreeRunStatus::Succeeded;
 	}
 
-	Exec.bHasPendingCompletedState = StateResult != EStateTreeRunStatus::Running;
+	Exec.bHasPendingCompletedState = StateResult != EStateTreeRunStatus::Running || FrameResult != EStateTreeRunStatus::Running;
 	return StateResult;
 }
 
@@ -4784,8 +4828,47 @@ PRAGMA_ENABLE_DEPRECATION_WARNINGS
 		EStateTreeRunStatus CurrentStatus = EStateTreeRunStatus::Unset;
 		for (int32 FrameIndex = 0; FrameIndex < Exec.ActiveFrames.Num(); ++FrameIndex)
 		{
+			using namespace UE::StateTree;
+			using namespace UE::StateTree::ExecutionContext::Private;
+
 			const FStateTreeExecutionFrame& CurrentFrame = Exec.ActiveFrames[FrameIndex];
 			const UStateTree* CurrentStateTree = CurrentFrame.StateTree;
+			const ETaskCompletionStatus FrameTasksStatus = CurrentFrame.ActiveTasksStatus.GetStatus(CurrentStateTree).GetCompletionStatus();
+			if (FrameTasksStatus != ETaskCompletionStatus::Running)
+			{
+				if (FrameIndex == 0)
+				{
+					// If first frame, then complete the tree execution.
+					Exec.RequestedStop = GetPriorityRunStatus(Exec.RequestedStop, CastToRunStatus(FrameTasksStatus));
+					break;
+				}
+				else if (UE::StateTree::ExecutionContext::Private::bGlobalTasksCompleteOwningFrame)
+				{
+					const int32 ParentFrameIndex = FrameIndex - 1;
+					FStateTreeExecutionFrame& ParentFrame = Exec.ActiveFrames[ParentFrameIndex];
+					const FStateTreeStateHandle ParentLinkedState = ParentFrame.ActiveStates.Last();
+					if (ensure(ParentLinkedState.IsValid()))
+					{
+						// Set the parent linked state as last completed state, and update tick status to the status from the transition.
+						const EStateTreeRunStatus RunStatus = NextTransition.TargetState.ToCompletionStatus();
+						STATETREE_LOG(Verbose, TEXT("Completed subtree '%s' from state '%s': %s"),
+							*GetSafeStateName(ParentFrame, ParentLinkedState),
+							*GetSafeStateName(CurrentFrame, NextTransition.SourceState),
+							*UEnum::GetDisplayValueAsText(RunStatus).ToString()
+							);
+
+						const FCompactStateTreeState& State = ParentFrame.StateTree->States[ParentLinkedState.Index];
+						ParentFrame.ActiveTasksStatus.GetStatus(State).SetCompletionStatus(FrameTasksStatus);
+						Exec.bHasPendingCompletedState = true;
+
+						CurrentStatus = UE::StateTree::ExecutionContext::Private::CastToRunStatus(FrameTasksStatus);
+						FrameIndexToStart = ParentFrameIndex;
+						StateIndexToStart = ParentLinkedState.Index;
+						break;
+					}
+				}
+			}
+
 			for (int32 StateIndex = 0; StateIndex < CurrentFrame.ActiveStates.Num(); ++StateIndex)
 			{
 				const FStateTreeStateHandle CurrentHandle = CurrentFrame.ActiveStates[StateIndex];
@@ -4806,8 +4889,6 @@ PRAGMA_ENABLE_DEPRECATION_WARNINGS
 			}
 		}
 
-		// Possible if the completed task was a global task
-		//@todo remove this once the frame global task only completes it's own frame.
 		if (CurrentStatus != EStateTreeRunStatus::Unset)
 		{
 			const bool bIsCurrentStatusSucceeded = CurrentStatus == EStateTreeRunStatus::Succeeded;
diff --git a/Engine/Plugins/Runtime/StateTree/Source/StateTreeTestSuite/Private/StateTreeLinkedStateTest.cpp b/Engine/Plugins/Runtime/StateTree/Source/StateTreeTestSuite/Private/StateTreeLinkedStateTest.cpp
--- a/Engine/Plugins/Runtime/StateTree/Source/StateTreeTestSuite/Private/StateTreeLinkedStateTest.cpp
+++ b/Engine/Plugins/Runtime/StateTree/Source/StateTreeTestSuite/Private/StateTreeLinkedStateTest.cpp
@@ -1295,6 +1295,189 @@ struct FStateTreeTest_Linked_GlobalParameter : FStateTreeTestBase
 };
 IMPLEMENT_AI_INSTANT_TEST(FStateTreeTest_Linked_GlobalParameter, "System.StateTree.LinkedAsset.GlobalParameter");
 
+struct FStateTreeTest_Linked_FinishGlobalTasks : FStateTreeTestBase
+{
+	virtual bool InstantTest() override
+	{
+		//Tree 1
+		//  Global task and parameter
+		//	Root
+		//		StateLinkedTree1 (Tree2) -> Next
+		//		SubTree2
+		//	SubTree
+		//		State3
+		//Tree 2
+		//  Global task and parameter
+		//	Root
+		//		State1
+
+
+		// Tree 2
+		UStateTree& StateTree2 = NewStateTree();
+		{
+			UStateTreeEditorData& EditorData2 = *Cast<UStateTreeEditorData>(StateTree2.EditorData);
+
+			// Global tasks
+			{
+				TStateTreeEditorNode<FTestTask_PrintValue>& GlobalTask1 = EditorData2.AddGlobalTask<FTestTask_PrintValue>("Tree2GlobalTask1");
+				GlobalTask1.GetInstanceData().Value = 1;
+
+				TStateTreeEditorNode<FTestTask_Stand>& GlobalTask2 = EditorData2.AddGlobalTask<FTestTask_Stand>("Tree2GlobalTask2");
+				GlobalTask2.GetNode().TicksToCompletion = 2;
+				GlobalTask2.GetNode().TickCompletionResult = EStateTreeRunStatus::Succeeded;
+			}
+
+			UStateTreeState& Root = EditorData2.AddSubTree("Tree2StateRoot");
+			{
+				TStateTreeEditorNode<FTestTask_PrintValue>& Task1 = Root.AddTask<FTestTask_PrintValue>("Tree2StateRootTask1");
+				Task1.GetInstanceData().Value = 1;
+			}
+
+			FStateTreeCompilerLog Log;
+			FStateTreeCompiler Compiler2(Log);
+			const bool bResult2 = Compiler2.Compile(StateTree2);
+			AITEST_TRUE(TEXT("StateTree2 should get compiled"), bResult2);
+		}
+
+		// Tree 1
+		UStateTree& StateTree1 = NewStateTree();
+		{
+			UStateTreeEditorData& EditorData1 = *Cast<UStateTreeEditorData>(StateTree1.EditorData);
+
+			// Global tasks
+			{
+				TStateTreeEditorNode<FTestTask_PrintValue>& GlobalTask1 = EditorData1.AddGlobalTask<FTestTask_PrintValue>("Tree1GlobalTask1");
+				GlobalTask1.GetInstanceData().Value = 1;
+
+				TStateTreeEditorNode<FTestTask_Stand>& GlobalTask2 = EditorData1.AddGlobalTask<FTestTask_Stand>("Tree1GlobalTask2");
+				GlobalTask2.GetNode().TicksToCompletion = 4;
+				GlobalTask2.GetNode().TickCompletionResult = EStateTreeRunStatus::Succeeded;
+			}
+
+			UStateTreeState& Root = EditorData1.AddSubTree("Tree1StateRoot");
+			{
+				TStateTreeEditorNode<FTestTask_PrintValue>& Task1 = Root.AddTask<FTestTask_PrintValue>("Tree1StateRootTask1");
+				Task1.GetInstanceData().Value = 1;
+			}
+			{
+				UStateTreeState& State1 = Root.AddChildState("Tree1State1", EStateTreeStateType::LinkedAsset);
+				State1.SetLinkedStateAsset(&StateTree2);
+				State1.AddTransition(EStateTreeTransitionTrigger::OnStateSucceeded, EStateTreeTransitionType::NextState);
+			}
+			{
+				UStateTreeState& State2 = Root.AddChildState("Tree1State2", EStateTreeStateType::State);
+				State2.AddTransition(EStateTreeTransitionTrigger::OnStateSucceeded, EStateTreeTransitionType::Succeeded);
+
+				TStateTreeEditorNode<FTestTask_Stand>& Task1 = State2.AddTask<FTestTask_Stand>("Tree1State2Task1");
+				Task1.GetNode().TicksToCompletion = 10;
+				Task1.GetNode().TickCompletionResult = EStateTreeRunStatus::Succeeded;
+			}
+
+			FStateTreeCompilerLog Log;
+			FStateTreeCompiler Compiler1(Log);
+			const bool bResult1 = Compiler1.Compile(StateTree1);
+			AITEST_TRUE(TEXT("StateTree1 should get compiled"), bResult1);
+		}
+
+		for (int32 Index = 0; Index < 4; ++Index)
+		{
+			Private::FScopedCVarBool CVarGGlobalTasksCompleteOwningFrame = TEXT("StateTree.GlobalTasksCompleteOwningFrame");
+			const bool bGlobalTasksCompleteOwningFrame = (Index % 2) == 0;
+			CVarGGlobalTasksCompleteOwningFrame.Set(bGlobalTasksCompleteOwningFrame);
+
+			Private::FScopedCVarBool CVarTickGlobalNodesWithHierarchy = TEXT("StateTree.TickGlobalNodesFollowingTreeHierarchy");
+			const bool bTickGlobalNodesWithHierarchy = Index >= 2;
+			CVarTickGlobalNodesWithHierarchy.Set(bTickGlobalNodesWithHierarchy);
+
+			FStateTreeInstanceData InstanceData;
+			{
+				FTestStateTreeExecutionContext Exec(StateTree1, StateTree1, InstanceData);
+				const bool bInitSucceeded = Exec.IsValid();
+				AITEST_TRUE(TEXT("StateTree should init"), bInitSucceeded);
+			}
+			{
+				FTestStateTreeExecutionContext Exec(StateTree1, StateTree1, InstanceData);
+				const EStateTreeRunStatus Status = Exec.Start();
+				AITEST_EQUAL(TEXT("Start should complete with Running"), Status, EStateTreeRunStatus::Running);
+				AITEST_TRUE(TEXT("In correct states"), Exec.ExpectInActiveStates("Tree1StateRoot", "Tree1State1", "Tree2StateRoot"));
+				AITEST_TRUE(TEXT("Start should enter Tree1GlobalTask1"), Exec.Expect("Tree1GlobalTask1", TEXT("EnterState1")));
+				AITEST_TRUE(TEXT("Start should enter Tree1GlobalTask2"), Exec.Expect("Tree1GlobalTask2", TEXT("EnterState")));
+				AITEST_TRUE(TEXT("Start should enter Tree1StateRootTask1"), Exec.Expect("Tree1StateRootTask1", TEXT("EnterState1")));
+				AITEST_TRUE(TEXT("Start should enter Tree2GlobalTask1"), Exec.Expect("Tree2GlobalTask1", TEXT("EnterState1")));
+				AITEST_TRUE(TEXT("Start should enter Tree2GlobalTask2"), Exec.Expect("Tree2GlobalTask2", TEXT("EnterState")));
+				AITEST_TRUE(TEXT("Start should enter Tree2StateRootTask1"), Exec.Expect("Tree2StateRootTask1", TEXT("EnterState1")));
+				Exec.LogClear();
+			}
+			{
+				FTestStateTreeExecutionContext Exec(StateTree1, StateTree1, InstanceData);
+				const EStateTreeRunStatus Status = Exec.Tick(1.0f);
+				AITEST_EQUAL(TEXT("Tick should complete with Running"), Status, EStateTreeRunStatus::Running);
+				AITEST_TRUE(TEXT("In correct states"), Exec.ExpectInActiveStates("Tree1StateRoot", "Tree1State1", "Tree2StateRoot"));
+				AITEST_TRUE(TEXT("Tick should tick Tree1GlobalTask1"), Exec.Expect("Tree1GlobalTask1", TEXT("Tick1")));
+				AITEST_TRUE(TEXT("Tick should tick Tree1GlobalTask2"), Exec.Expect("Tree1GlobalTask2", TEXT("Tick")));
+				AITEST_TRUE(TEXT("Tick should tick Tree1StateRootTask1"), Exec.Expect("Tree1StateRootTask1", TEXT("Tick1")));
+				AITEST_TRUE(TEXT("Tick should tick Tree2GlobalTask1"), Exec.Expect("Tree2GlobalTask1", TEXT("Tick1")));
+				AITEST_TRUE(TEXT("Tick should tick Tree2GlobalTask2"), Exec.Expect("Tree2GlobalTask2", TEXT("Tick")));
+				AITEST_TRUE(TEXT("Tick should tick Tree2StateRootTask1"), Exec.Expect("Tree2StateRootTask1", TEXT("Tick1")));
+				Exec.LogClear();
+			}
+			if (bGlobalTasksCompleteOwningFrame)
+			{
+				{
+					FTestStateTreeExecutionContext Exec(StateTree1, StateTree1, InstanceData);
+					const EStateTreeRunStatus Status = Exec.Tick(1.0f);
+					AITEST_EQUAL(TEXT("Tick should complete with Running"), Status, EStateTreeRunStatus::Running);
+					AITEST_TRUE(TEXT("In correct states"), Exec.ExpectInActiveStates("Tree1StateRoot", "Tree1State2"));
+					AITEST_TRUE(TEXT("Tick should tick Tree1GlobalTask1"), Exec.Expect("Tree1GlobalTask1", TEXT("Tick1")));
+					AITEST_TRUE(TEXT("Tick should tick Tree1GlobalTask2"), Exec.Expect("Tree1GlobalTask2", TEXT("Tick")));
+					if (bTickGlobalNodesWithHierarchy)
+					{
+						AITEST_TRUE(TEXT("Tick should tick Tree1StateRootTask1"), Exec.Expect("Tree1StateRootTask1", TEXT("Tick1")));
+					}
+					AITEST_TRUE(TEXT("Tick should tick Tree2GlobalTask1"), Exec.Expect("Tree2GlobalTask1", TEXT("Tick1")));
+					AITEST_TRUE(TEXT("Tick should tick Tree2GlobalTask2"), Exec.Expect("Tree2GlobalTask2", TEXT("Tick")));
+					if (bTickGlobalNodesWithHierarchy)
+					{
+						AITEST_FALSE(TEXT("Tick not should tick Tree2StateRootTask1"), Exec.Expect("Tree2StateRootTask1", TEXT("Tick1")));
+					}
+
+					AITEST_TRUE(TEXT("Tick should ExitState Tree2StateRootTask1"), Exec.Expect("Tree2StateRootTask1", TEXT("ExitState1")));
+					AITEST_TRUE(TEXT("Tick should ExitState Tree2GlobalTask2"), Exec.Expect("Tree2GlobalTask2", TEXT("ExitState")));
+					AITEST_TRUE(TEXT("Tick should ExitState Tree2GlobalTask1"), Exec.Expect("Tree2GlobalTask1", TEXT("ExitState1")));
+					AITEST_TRUE(TEXT("Start should enter Tree2State2Task1"), Exec.Expect("Tree1State2Task1", TEXT("EnterState")));
+					Exec.LogClear();
+				}
+				{
+					FTestStateTreeExecutionContext Exec(StateTree1, StateTree1, InstanceData);
+					const EStateTreeRunStatus Status = Exec.Tick(1.0f);
+					AITEST_EQUAL(TEXT("Tick should complete with Running"), Status, EStateTreeRunStatus::Running);
+					Exec.LogClear();
+				}
+				{
+					FTestStateTreeExecutionContext Exec(StateTree1, StateTree1, InstanceData);
+					const EStateTreeRunStatus Status = Exec.Tick(1.0f);
+					AITEST_EQUAL(TEXT("Tick should complete with Running"), Status, EStateTreeRunStatus::Succeeded);
+					Exec.LogClear();
+				}
+			}
+			else
+			{
+				FTestStateTreeExecutionContext Exec(StateTree1, StateTree1, InstanceData);
+				const EStateTreeRunStatus Status = Exec.Tick(1.0f);
+				AITEST_EQUAL(TEXT("Tick should complete with Running"), Status, EStateTreeRunStatus::Succeeded);
+				Exec.LogClear();
+			}
+			{
+				FTestStateTreeExecutionContext Exec(StateTree1, StateTree1, InstanceData);
+				Exec.Stop();
+			}
+		}
+
+		return true;
+	}
+};
+IMPLEMENT_AI_INSTANT_TEST(FStateTreeTest_Linked_FinishGlobalTasks, "System.StateTree.LinkedAsset.FinishGlobalTasks");
+
 } // namespace UE::StateTree::Tests
 
 UE_ENABLE_OPTIMIZATION_SHIP
