StateTree: Create task status when the transition comes from a RecordedTransitionResult.

diff --git a/Engine/Plugins/Runtime/StateTree/Source/StateTreeModule/Private/StateTreeExecutionContext.cpp b/Engine/Plugins/Runtime/StateTree/Source/StateTreeModule/Private/StateTreeExecutionContext.cpp
--- a/Engine/Plugins/Runtime/StateTree/Source/StateTreeModule/Private/StateTreeExecutionContext.cpp
+++ b/Engine/Plugins/Runtime/StateTree/Source/StateTreeModule/Private/StateTreeExecutionContext.cpp
@@ -1539,7 +1539,7 @@ void FStateTreeExecutionContext::TickTriggerTransitionsInternal()
 			// Enter state tasks can fail/succeed, treat it same as tick.
 			const EStateTreeRunStatus LastTickStatus = EnterState(NextTransition);
 
-			NextTransition.Reset();
+			NextTransition = FStateTreeTransitionResult();
 
 			Exec.LastTickStatus = LastTickStatus;
 
@@ -2138,11 +2138,21 @@ EStateTreeRunStatus FStateTreeExecutionContext::ForceTransition(const FRecordedS
 		return EStateTreeRunStatus::Failed;
 	}
 
-	FStateTreeTransitionResult TransitionResult = FStateTreeTransitionResult(Transition);
+	// A reentrant call to ForceTransition or a call from Start, Tick or Stop must be deferred.
+	if (GetExecState().CurrentPhase != EStateTreeUpdatePhase::Unset)
+	{
+		UE_STATETREE_DEBUG_LOG_EVENT(this, Warning, TEXT("Can't force a transition while %s"), *UEnum::GetDisplayValueAsText(GetExecState().CurrentPhase).ToString());
+		return EStateTreeRunStatus::Unset;
+	}
 
-	ExitState(TransitionResult);
+	TOptional<FStateTreeTransitionResult> TransitionResult = MakeTransitionResult(Transition);
+	if (!TransitionResult.IsSet())
+	{
+		return EStateTreeRunStatus::Unset;
+	}
 
-	return EnterState(TransitionResult);
+	ExitState(TransitionResult.GetValue());
+	return EnterState(TransitionResult.GetValue());
 }
 
 FStateTreeDataView FStateTreeExecutionContext::GetDataViewFromInstanceStorage(FStateTreeInstanceStorage& InstanceDataStorage, FStateTreeInstanceStorage* CurrentlyProcessedSharedInstanceStorage, const FStateTreeExecutionFrame* ParentFrame, const FStateTreeExecutionFrame& CurrentFrame, const FStateTreeDataHandle Handle)
@@ -2677,7 +2687,7 @@ EStateTreeRunStatus FStateTreeExecutionContext::EnterState(FStateTreeTransitionR
 
 	if (bRecordTransitions)
 	{
-		RecordedTransitions.Add(FRecordedStateTreeTransitionResult(Transition));
+		RecordedTransitions.Add(MakeRecordedTransitionResult(Transition));
 	}
 
 	// Allocate new tasks.
@@ -2722,10 +2732,13 @@ EStateTreeRunStatus FStateTreeExecutionContext::EnterState(FStateTreeTransitionR
 		// We'll add new states one by one, so that active states contain only the states which have EnterState called.
 		CurrentFrame.ActiveStates.Reset();
 
-		// Create new status
-		const FCompactStateTreeFrame* FrameInfo = NextFrame.StateTree->GetFrameFromHandle(NextFrame.RootState);
-		ensureAlwaysMsgf(FrameInfo, TEXT("The compiled data is invalid. It should contains the information for the root frame."));
-		CurrentFrame.ActiveTasksStatus = FrameInfo ? FStateTreeTasksCompletionStatus(*FrameInfo) : FStateTreeTasksCompletionStatus();
+		if (!ensureMsgf(CurrentFrame.ActiveTasksStatus.IsValid(), TEXT("Frame is not formed correct.")))
+		{
+			// Create the status. We lost the previous tasks' completion status.
+			const FCompactStateTreeFrame* FrameInfo = NextFrame.StateTree->GetFrameFromHandle(NextFrame.RootState);
+			ensureAlwaysMsgf(FrameInfo, TEXT("The compiled data is invalid. It should contains the information for the root frame."));
+			CurrentFrame.ActiveTasksStatus = FrameInfo ? FStateTreeTasksCompletionStatus(*FrameInfo) : FStateTreeTasksCompletionStatus();
+		}
 
 		// Get previous active states, they are used to calculate transition type.
 		FStateTreeActiveStates PreviousActiveStates;
@@ -4428,7 +4441,7 @@ bool FStateTreeExecutionContext::TriggerTransitions()
 		UE_STATETREE_DEBUG_LOG_EVENT(this, Log, TEXT("Trigger transitions with events: %s"), *DebugGetEventsAsString());
 	}
 
-	NextTransition.Reset();
+	NextTransition = FStateTreeTransitionResult();
 
 	//
 	// Process transition requests
@@ -5036,7 +5049,7 @@ PRAGMA_ENABLE_DEPRECATION_WARNINGS
 				// now starting from the linked parent state. If we run out of retires in the selection loop (e.g. very deep hierarchy)
 				// we will continue on next tick.
 				TriggerTransitionsFromFrameIndex = ParentFrameIndex;
-				NextTransition.Reset();
+				NextTransition = FStateTreeTransitionResult();
 				return false;
 			}
 		}
@@ -5045,6 +5058,112 @@ PRAGMA_ENABLE_DEPRECATION_WARNINGS
 	return NextTransition.TargetState.IsValid();
 }
 
+TOptional<FStateTreeTransitionResult> FStateTreeExecutionContext::MakeTransitionResult(const FRecordedStateTreeTransitionResult& RecordedTransition) const
+{
+	FStateTreeTransitionResult Result;
+
+	for (int32 RecordedFrameIndex = 0; RecordedFrameIndex < RecordedTransition.NextActiveFrames.Num(); RecordedFrameIndex++)
+	{
+		const FRecordedStateTreeExecutionFrame& RecordedExecutionFrame = RecordedTransition.NextActiveFrames[RecordedFrameIndex];
+
+		if (RecordedExecutionFrame.StateTree == nullptr)
+		{
+			return {};
+		}
+
+		if (RecordedExecutionFrame.StateTree->GetStateFromHandle(RecordedExecutionFrame.RootState) == nullptr)
+		{
+			return {};
+		}
+
+		const FCompactStateTreeFrame* CompactFrame = RecordedExecutionFrame.StateTree->GetFrameFromHandle(RecordedExecutionFrame.RootState);
+		if (CompactFrame == nullptr)
+		{
+			return {};
+		}
+
+		FStateTreeExecutionFrame& ExecutionFrame = Result.NextActiveFrames.AddDefaulted_GetRef();
+		ExecutionFrame.StateTree = RecordedExecutionFrame.StateTree;
+		ExecutionFrame.RootState = RecordedExecutionFrame.RootState;
+		ExecutionFrame.ActiveStates = RecordedExecutionFrame.ActiveStates;
+		ExecutionFrame.ActiveTasksStatus = FStateTreeTasksCompletionStatus(*CompactFrame);
+		ExecutionFrame.bIsGlobalFrame = RecordedExecutionFrame.bIsGlobalFrame;
+
+		FStateTreeFrameStateSelectionEvents& StateTreeFrameStateSelectionEvents = Result.NextActiveFrameEvents.AddDefaulted_GetRef();
+		for (int32 EventIdx = 0; EventIdx < RecordedExecutionFrame.EventIndices.Num(); EventIdx++)
+		{
+			if (RecordedTransition.NextActiveFrameEvents.IsValidIndex(EventIdx))
+			{
+				const FStateTreeEvent& RecordedStateTreeEvent = RecordedTransition.NextActiveFrameEvents[EventIdx];
+				StateTreeFrameStateSelectionEvents.Events[EventIdx] = FStateTreeSharedEvent(RecordedStateTreeEvent);
+			}
+		}
+	}
+
+
+	if (Result.NextActiveFrames.Num() != Result.NextActiveFrameEvents.Num())
+	{
+		return {};
+	}
+
+	if (RecordedTransition.SourceStateTree == nullptr)
+	{
+		return {};
+	}
+
+	if (RecordedTransition.SourceStateTree->GetFrameFromHandle(RecordedTransition.SourceRootState) == nullptr)
+	{
+		return {};
+	}
+
+	Result.SourceState = RecordedTransition.SourceState;
+	Result.TargetState = RecordedTransition.TargetState;
+	Result.Priority = RecordedTransition.Priority;
+	Result.SourceStateTree = RecordedTransition.SourceStateTree;
+	Result.SourceRootState = RecordedTransition.SourceRootState;
+
+	return Result;
+}
+
+FRecordedStateTreeTransitionResult FStateTreeExecutionContext::MakeRecordedTransitionResult(const FStateTreeTransitionResult& Transition) const
+{
+	check(Transition.NextActiveFrames.Num() == Transition.NextActiveFrameEvents.Num());
+
+	const FStateTreeExecutionState& Exec = GetExecState();
+
+	FRecordedStateTreeTransitionResult Result;
+
+	for (int32 FrameIndex = 0; FrameIndex < Transition.NextActiveFrames.Num(); FrameIndex++)
+	{
+		const FStateTreeExecutionFrame& ExecutionFrame = Transition.NextActiveFrames[FrameIndex];
+		const FStateTreeFrameStateSelectionEvents& StateSelectionEvents = Transition.NextActiveFrameEvents[FrameIndex];
+
+		FRecordedStateTreeExecutionFrame& RecordedFrame = Result.NextActiveFrames.AddDefaulted_GetRef();
+		RecordedFrame.StateTree = ExecutionFrame.StateTree;
+		RecordedFrame.RootState = ExecutionFrame.RootState;
+		RecordedFrame.ActiveStates = ExecutionFrame.ActiveStates;
+		RecordedFrame.bIsGlobalFrame = ExecutionFrame.bIsGlobalFrame;
+
+		for (int32 StateIndex = 0; StateIndex < ExecutionFrame.ActiveStates.Num(); StateIndex++)
+		{
+			const FStateTreeEvent* Event = StateSelectionEvents.Events[StateIndex].Get();
+			if (Event)
+			{
+				const int32 EventIndex = Result.NextActiveFrameEvents.Add(*Event);
+				RecordedFrame.EventIndices[StateIndex] = static_cast<uint8>(EventIndex);
+			}
+		}
+	}
+
+	Result.SourceState = Transition.SourceState;
+	Result.TargetState = Transition.TargetState;
+	Result.Priority = Transition.Priority;
+	Result.SourceStateTree = Transition.SourceStateTree;
+	Result.SourceRootState = Transition.SourceRootState;
+
+	return Result;
+}
+
 bool FStateTreeExecutionContext::SelectState(const FStateTreeExecutionFrame& CurrentFrame,
 											const FStateTreeStateHandle NextState,
 											FStateSelectionResult& OutSelectionResult,
diff --git a/Engine/Plugins/Runtime/StateTree/Source/StateTreeModule/Private/StateTreeExecutionTypes.cpp b/Engine/Plugins/Runtime/StateTree/Source/StateTreeModule/Private/StateTreeExecutionTypes.cpp
--- a/Engine/Plugins/Runtime/StateTree/Source/StateTreeModule/Private/StateTreeExecutionTypes.cpp
+++ b/Engine/Plugins/Runtime/StateTree/Source/StateTreeModule/Private/StateTreeExecutionTypes.cpp
@@ -31,64 +31,17 @@ FStateTreeTransitionSource::FStateTreeTransitionSource(
 //----------------------------------------------------------------------//
 // FStateTreeTransitionResult
 //----------------------------------------------------------------------//
+//Deprecated
 FStateTreeTransitionResult::FStateTreeTransitionResult(const FRecordedStateTreeTransitionResult& RecordedTransition)
 {
-	for (int32 RecordedFrameIndex = 0; RecordedFrameIndex < RecordedTransition.NextActiveFrames.Num(); RecordedFrameIndex++)
-	{
-		const FRecordedStateTreeExecutionFrame& RecordedExecutionFrame = RecordedTransition.NextActiveFrames[RecordedFrameIndex];
-		NextActiveFrames.Add(RecordedExecutionFrame);
-
-		FStateTreeFrameStateSelectionEvents StateTreeFrameStateSelectionEvents;
-
-		for (int32 EventIdx = 0; EventIdx < RecordedExecutionFrame.EventIndices.Num(); EventIdx++)
-		{
-			if (RecordedTransition.NextActiveFrameEvents.IsValidIndex(EventIdx))
-			{
-				const FStateTreeEvent& RecordedStateTreeEvent = RecordedTransition.NextActiveFrameEvents[EventIdx];
-				StateTreeFrameStateSelectionEvents.Events[EventIdx] = FStateTreeSharedEvent(RecordedStateTreeEvent);
-			}
-		}
-
-		NextActiveFrameEvents.Add(MoveTemp(StateTreeFrameStateSelectionEvents));
-	}
-
-	SourceState = RecordedTransition.SourceState;
-	TargetState = RecordedTransition.SourceState;
-	Priority = RecordedTransition.Priority;
-	SourceStateTree = RecordedTransition.SourceStateTree;
-	SourceRootState = RecordedTransition.SourceRootState;
 }
 
 //----------------------------------------------------------------------//
 // FRecordedStateTreeTransitionResult
 //----------------------------------------------------------------------//
+//Deprecated
 FRecordedStateTreeTransitionResult::FRecordedStateTreeTransitionResult(const FStateTreeTransitionResult& Transition)
 {
-	check(Transition.NextActiveFrames.Num() == Transition.NextActiveFrameEvents.Num());
-
-	for (int32 FrameIndex = 0; FrameIndex < Transition.NextActiveFrames.Num(); FrameIndex++)
-	{
-		const FStateTreeExecutionFrame& ExecutionFrame = Transition.NextActiveFrames[FrameIndex];
-		const FStateTreeFrameStateSelectionEvents& StateSelectionEvents = Transition.NextActiveFrameEvents[FrameIndex];
-
-		FRecordedStateTreeExecutionFrame& RecordedFrame = NextActiveFrames.Emplace_GetRef(ExecutionFrame);
-
-		for (int32 StateIndex = 0; StateIndex < ExecutionFrame.ActiveStates.Num(); StateIndex++)
-		{
-			const FStateTreeEvent* Event = StateSelectionEvents.Events[StateIndex].Get();
-			if (Event)
-			{
-				const int32 EventIndex = NextActiveFrameEvents.Add(*Event);
-				RecordedFrame.EventIndices[StateIndex] = static_cast<uint8>(EventIndex);
-			}
-		}
-	}
-
-	SourceState = Transition.SourceState;
-	TargetState = Transition.SourceState;
-	Priority = Transition.Priority;
-	SourceStateTree = Transition.SourceStateTree;
-	SourceRootState = Transition.SourceRootState;
 }
 
 
@@ -222,23 +175,17 @@ void FStateTreeExecutionState::CacheScheduledTickRequest()
 //----------------------------------------------------------------------//
 // FStateTreeExecutionFrame
 //----------------------------------------------------------------------//
+//Deprecated
 FStateTreeExecutionFrame::FStateTreeExecutionFrame(const FRecordedStateTreeExecutionFrame& RecordedExecutionFrame)
 {
-	StateTree = RecordedExecutionFrame.StateTree;
-	RootState = RecordedExecutionFrame.RootState;
-	ActiveStates = RecordedExecutionFrame.ActiveStates;
-	bIsGlobalFrame = RecordedExecutionFrame.bIsGlobalFrame;
 }
 
 //----------------------------------------------------------------------//
 // FRecordedStateTreeExecutionFrame
 //----------------------------------------------------------------------//
+//Deprecated
 FRecordedStateTreeExecutionFrame::FRecordedStateTreeExecutionFrame(const FStateTreeExecutionFrame& ExecutionFrame)
 {
-	StateTree = ExecutionFrame.StateTree;
-	RootState = ExecutionFrame.RootState;
-	ActiveStates = ExecutionFrame.ActiveStates;
-	bIsGlobalFrame = ExecutionFrame.bIsGlobalFrame;
 }
 
 //Deprecated
diff --git a/Engine/Plugins/Runtime/StateTree/Source/StateTreeModule/Public/StateTreeExecutionContext.h b/Engine/Plugins/Runtime/StateTree/Source/StateTreeModule/Public/StateTreeExecutionContext.h
--- a/Engine/Plugins/Runtime/StateTree/Source/StateTreeModule/Public/StateTreeExecutionContext.h
+++ b/Engine/Plugins/Runtime/StateTree/Source/StateTreeModule/Public/StateTreeExecutionContext.h
@@ -893,6 +893,11 @@ protected:
 	 */
 	bool TriggerTransitions();
 
+	/** Create a new transition result from a recorded transition result. It will fail if the recorded transition is malformed. */
+	TOptional<FStateTreeTransitionResult> MakeTransitionResult(const FRecordedStateTreeTransitionResult& Transition) const;
+	/** Create a new recorded transition from a transition result. */
+	FRecordedStateTreeTransitionResult MakeRecordedTransitionResult(const FStateTreeTransitionResult& Transition) const;
+
 PRAGMA_DISABLE_DEPRECATION_WARNINGS
 	UE_DEPRECATED(5.6, "Replaced with FStateTreeTasksCompletionStatus")
 	/** Confirms that the frame and state ID are valid and the task index is correct. */
diff --git a/Engine/Plugins/Runtime/StateTree/Source/StateTreeModule/Public/StateTreeExecutionTypes.h b/Engine/Plugins/Runtime/StateTree/Source/StateTreeModule/Public/StateTreeExecutionTypes.h
--- a/Engine/Plugins/Runtime/StateTree/Source/StateTreeModule/Public/StateTreeExecutionTypes.h
+++ b/Engine/Plugins/Runtime/StateTree/Source/StateTreeModule/Public/StateTreeExecutionTypes.h
@@ -870,8 +870,10 @@ struct STATETREEMODULE_API FStateTreeExecutionFrame
 	GENERATED_BODY()
 
 	FStateTreeExecutionFrame() = default;
+	UE_DEPRECATED(5.6, "The recorded frame doesn't have all the needed information to properly form an FStateTreeExecutionFrame.")
 	FStateTreeExecutionFrame(const FRecordedStateTreeExecutionFrame& RecordedExecutionFrame);
 
+	//@todo Refactor IsSameFrame to use FrameID
 	bool IsSameFrame(const FStateTreeExecutionFrame& OtherFrame) const
 	{
 		return StateTree == OtherFrame.StateTree && RootState == OtherFrame.RootState;
@@ -1122,16 +1124,12 @@ PRAGMA_DISABLE_DEPRECATION_WARNINGS
 	FStateTreeTransitionResult& operator=(FStateTreeTransitionResult&&) = default;
 PRAGMA_ENABLE_DEPRECATION_WARNINGS
 
+	UE_DEPRECATED(5.6, "Use FStateTreeExecutionContext::MakeTransitionResult to crate a new transition.")
 	FStateTreeTransitionResult(const FRecordedStateTreeTransitionResult& RecordedTransition);
-	
+
 	void Reset()
 	{
-		NextActiveFrames.Reset();
-		CurrentRunStatus = EStateTreeRunStatus::Unset;
-		TargetState = FStateTreeStateHandle::Invalid;
-		CurrentState = FStateTreeStateHandle::Invalid;
-		ChangeType = EStateTreeStateChangeType::Changed;
-		Priority = EStateTreeTransitionPriority::None;
+		*this = FStateTreeTransitionResult();
 	}
 
 	/** States selected as result of the transition. */
@@ -1195,6 +1193,7 @@ struct STATETREEMODULE_API FRecordedStateTreeExecutionFrame
 	GENERATED_BODY()
 
 	FRecordedStateTreeExecutionFrame() = default;
+	UE_DEPRECATED(5.6, "Use FStateTreeExecutionContext::MakeRecordedTransitionResult to crate a new recorded transition.")
 	FRecordedStateTreeExecutionFrame(const FStateTreeExecutionFrame& ExecutionFrame);
 
 	/** The State Tree used for ticking this frame. */
@@ -1218,10 +1217,10 @@ struct STATETREEMODULE_API FRecordedStateTreeExecutionFrame
 };
 
 /*
-* Captured state tree transition result that can be cached for recording purposes.
-* Primarily, when transitions are recorded through this structure, we can replicate them down
-* to clients to keep our state tree in sync.
-*/
+ * Captured state tree transition result that can be cached for recording purposes.
+ * Primarily, when transitions are recorded through this structure, we can replicate them down
+ * to clients to keep our state tree in sync.
+ */
 USTRUCT()
 struct STATETREEMODULE_API FRecordedStateTreeTransitionResult
 {
@@ -1231,7 +1230,7 @@ struct STATETREEMODULE_API FRecordedStateTreeTransitionResult
 	FRecordedStateTreeTransitionResult(const FStateTreeTransitionResult& Transition);
 
    	/** States selected as result of the transition. */
-    UPROPERTY()
+	UPROPERTY()
 	TArray<FRecordedStateTreeExecutionFrame> NextActiveFrames;
 
 	/** Captured events from the transition that we've recorded */
diff --git a/Engine/Plugins/Runtime/StateTree/Source/StateTreeModule/Public/StateTreeTasksStatus.h b/Engine/Plugins/Runtime/StateTree/Source/StateTreeModule/Public/StateTreeTasksStatus.h
--- a/Engine/Plugins/Runtime/StateTree/Source/StateTreeModule/Public/StateTreeTasksStatus.h
+++ b/Engine/Plugins/Runtime/StateTree/Source/StateTreeModule/Public/StateTreeTasksStatus.h
@@ -303,6 +303,12 @@ public:
 	[[nodiscard]] STATETREEMODULE_API UE::StateTree::FTasksCompletionStatus GetStatus(TNotNull<const UStateTree*> StateTree);
 	[[nodiscard]] STATETREEMODULE_API UE::StateTree::FConstTasksCompletionStatus GetStatus(TNotNull<const UStateTree*> StateTree) const;
 
+	/** @return true when the status is initialized correctly. */
+	[[nodiscard]] bool IsValid() const
+	{
+		return BufferNum > 0;
+	}
+
 	STATETREEMODULE_API void Push(const FCompactStateTreeState& State);
 
 	STATETREEMODULE_API bool Serialize(FArchive& Ar);
