StateTree: Fix bad index usage in transition. The index used was the state index inside the asset, and we needed the state index in the active list.

diff --git a/Engine/Plugins/Runtime/StateTree/Source/StateTreeModule/Private/StateTreeExecutionContext.cpp b/Engine/Plugins/Runtime/StateTree/Source/StateTreeModule/Private/StateTreeExecutionContext.cpp
--- a/Engine/Plugins/Runtime/StateTree/Source/StateTreeModule/Private/StateTreeExecutionContext.cpp
+++ b/Engine/Plugins/Runtime/StateTree/Source/StateTreeModule/Private/StateTreeExecutionContext.cpp
@@ -4876,7 +4876,7 @@ PRAGMA_ENABLE_DEPRECATION_WARNINGS
 
 						CurrentStatus = UE::StateTree::ExecutionContext::Private::CastToRunStatus(FrameTasksStatus);
 						FrameIndexToStart = ParentFrameIndex;
-						StateIndexToStart = ParentLinkedState.Index;
+						StateIndexToStart = ParentFrame.ActiveStates.Num() - 1;
 						break;
 					}
 				}
diff --git a/Engine/Plugins/Runtime/StateTree/Source/StateTreeTestSuite/Private/StateTreeLinkedStateTest.cpp b/Engine/Plugins/Runtime/StateTree/Source/StateTreeTestSuite/Private/StateTreeLinkedStateTest.cpp
--- a/Engine/Plugins/Runtime/StateTree/Source/StateTreeTestSuite/Private/StateTreeLinkedStateTest.cpp
+++ b/Engine/Plugins/Runtime/StateTree/Source/StateTreeTestSuite/Private/StateTreeLinkedStateTest.cpp
@@ -1319,10 +1319,18 @@ struct FStateTreeTest_Linked_FinishGlobalTasks : FStateTreeTestBase
 
 			// Global tasks
 			{
+				EditorData2.GlobalTasksCompletion = EStateTreeTaskCompletionType::Any;
+
 				TStateTreeEditorNode<FTestTask_PrintValue>& GlobalTask1 = EditorData2.AddGlobalTask<FTestTask_PrintValue>("Tree2GlobalTask1");
 				GlobalTask1.GetInstanceData().Value = 1;
 
-				TStateTreeEditorNode<FTestTask_Stand>& GlobalTask2 = EditorData2.AddGlobalTask<FTestTask_Stand>("Tree2GlobalTask2");
+				EditorData2.AddGlobalTask<FTestTask_Stand>("Tree2GlobalTask2").GetNode().TicksToCompletion = 99;
+				EditorData2.AddGlobalTask<FTestTask_Stand>("Tree2GlobalTask3").GetNode().TicksToCompletion = 99;
+				EditorData2.AddGlobalTask<FTestTask_Stand>("Tree2GlobalTask4").GetNode().TicksToCompletion = 99;
+				EditorData2.AddGlobalTask<FTestTask_Stand>("Tree2GlobalTask5").GetNode().TicksToCompletion = 99;
+				EditorData2.AddGlobalTask<FTestTask_Stand>("Tree2GlobalTask6").GetNode().TicksToCompletion = 99;
+
+				TStateTreeEditorNode<FTestTask_Stand>& GlobalTask2 = EditorData2.AddGlobalTask<FTestTask_Stand>("Tree2GlobalTask7");
 				GlobalTask2.GetNode().TicksToCompletion = 2;
 				GlobalTask2.GetNode().TickCompletionResult = EStateTreeRunStatus::Succeeded;
 			}
@@ -1404,7 +1412,7 @@ struct FStateTreeTest_Linked_FinishGlobalTasks : FStateTreeTestBase
 				AITEST_TRUE(TEXT("Start should enter Tree1GlobalTask2"), Exec.Expect("Tree1GlobalTask2", TEXT("EnterState")));
 				AITEST_TRUE(TEXT("Start should enter Tree1StateRootTask1"), Exec.Expect("Tree1StateRootTask1", TEXT("EnterState1")));
 				AITEST_TRUE(TEXT("Start should enter Tree2GlobalTask1"), Exec.Expect("Tree2GlobalTask1", TEXT("EnterState1")));
-				AITEST_TRUE(TEXT("Start should enter Tree2GlobalTask2"), Exec.Expect("Tree2GlobalTask2", TEXT("EnterState")));
+				AITEST_TRUE(TEXT("Start should enter Tree2GlobalTask7"), Exec.Expect("Tree2GlobalTask7", TEXT("EnterState")));
 				AITEST_TRUE(TEXT("Start should enter Tree2StateRootTask1"), Exec.Expect("Tree2StateRootTask1", TEXT("EnterState1")));
 				Exec.LogClear();
 			}
@@ -1417,7 +1425,7 @@ struct FStateTreeTest_Linked_FinishGlobalTasks : FStateTreeTestBase
 				AITEST_TRUE(TEXT("Tick should tick Tree1GlobalTask2"), Exec.Expect("Tree1GlobalTask2", TEXT("Tick")));
 				AITEST_TRUE(TEXT("Tick should tick Tree1StateRootTask1"), Exec.Expect("Tree1StateRootTask1", TEXT("Tick1")));
 				AITEST_TRUE(TEXT("Tick should tick Tree2GlobalTask1"), Exec.Expect("Tree2GlobalTask1", TEXT("Tick1")));
-				AITEST_TRUE(TEXT("Tick should tick Tree2GlobalTask2"), Exec.Expect("Tree2GlobalTask2", TEXT("Tick")));
+				AITEST_TRUE(TEXT("Tick should tick Tree2GlobalTask7"), Exec.Expect("Tree2GlobalTask7", TEXT("Tick")));
 				AITEST_TRUE(TEXT("Tick should tick Tree2StateRootTask1"), Exec.Expect("Tree2StateRootTask1", TEXT("Tick1")));
 				Exec.LogClear();
 			}
@@ -1435,14 +1443,14 @@ struct FStateTreeTest_Linked_FinishGlobalTasks : FStateTreeTestBase
 						AITEST_TRUE(TEXT("Tick should tick Tree1StateRootTask1"), Exec.Expect("Tree1StateRootTask1", TEXT("Tick1")));
 					}
 					AITEST_TRUE(TEXT("Tick should tick Tree2GlobalTask1"), Exec.Expect("Tree2GlobalTask1", TEXT("Tick1")));
-					AITEST_TRUE(TEXT("Tick should tick Tree2GlobalTask2"), Exec.Expect("Tree2GlobalTask2", TEXT("Tick")));
+					AITEST_TRUE(TEXT("Tick should tick Tree2GlobalTask7"), Exec.Expect("Tree2GlobalTask7", TEXT("Tick")));
 					if (bTickGlobalNodesWithHierarchy)
 					{
 						AITEST_FALSE(TEXT("Tick not should tick Tree2StateRootTask1"), Exec.Expect("Tree2StateRootTask1", TEXT("Tick1")));
 					}
 
 					AITEST_TRUE(TEXT("Tick should ExitState Tree2StateRootTask1"), Exec.Expect("Tree2StateRootTask1", TEXT("ExitState1")));
-					AITEST_TRUE(TEXT("Tick should ExitState Tree2GlobalTask2"), Exec.Expect("Tree2GlobalTask2", TEXT("ExitState")));
+					AITEST_TRUE(TEXT("Tick should ExitState Tree2GlobalTask7"), Exec.Expect("Tree2GlobalTask7", TEXT("ExitState")));
 					AITEST_TRUE(TEXT("Tick should ExitState Tree2GlobalTask1"), Exec.Expect("Tree2GlobalTask1", TEXT("ExitState1")));
 					AITEST_TRUE(TEXT("Start should enter Tree2State2Task1"), Exec.Expect("Tree1State2Task1", TEXT("EnterState")));
 					Exec.LogClear();
