StateTree: Use FrameID and StateID in FStateTreeTransitionResult. Kept the deprecated properties valid for Blueprint. The execution context should only use the IDs.

diff --git a/Engine/Plugins/Runtime/StateTree/Source/StateTreeModule/Private/StateTreeExecutionContext.cpp b/Engine/Plugins/Runtime/StateTree/Source/StateTreeModule/Private/StateTreeExecutionContext.cpp
--- a/Engine/Plugins/Runtime/StateTree/Source/StateTreeModule/Private/StateTreeExecutionContext.cpp
+++ b/Engine/Plugins/Runtime/StateTree/Source/StateTreeModule/Private/StateTreeExecutionContext.cpp
@@ -1681,7 +1681,7 @@ void FStateTreeExecutionContext::RequestTransition(const FStateTreeTransitionReq
 
 		FStateTreeTransitionRequest RequestWithSource = Request;
 		RequestWithSource.SourceFrameID = RootFrame->FrameID;
-		int32 ActiveStateIndex = RootFrame->ActiveStates.IndexOfReverse(CurrentlyProcessedState);
+		const int32 ActiveStateIndex = RootFrame->ActiveStates.IndexOfReverse(CurrentlyProcessedState);
 		RequestWithSource.SourceStateID = ActiveStateIndex != INDEX_NONE ? RootFrame->ActiveStates.StateIDs[ActiveStateIndex] : UE::StateTree::FActiveStateID::Invalid;
 
 		InstanceData.AddTransitionRequest(&Owner, RequestWithSource);
@@ -2783,8 +2783,7 @@ EStateTreeRunStatus FStateTreeExecutionContext::EnterState(FStateTreeTransitionR
 				}
 			}
 
-			if (CurrentFrame.StateTree == Transition.SourceStateTree
-				&& CurrentFrame.RootState == Transition.SourceRootState
+			if (CurrentFrame.FrameID == Transition.SourceFrameID
 				&& CurrentHandle == Transition.TargetState)
 			{
 				bOnTargetBranch = true;
@@ -3020,8 +3019,7 @@ void FStateTreeExecutionContext::ExitState(const FStateTreeTransitionResult& Tra
 				}
 			}
 
-			if (CurrentFrame.StateTree == Transition.SourceStateTree
-				&& CurrentFrame.RootState == Transition.SourceRootState
+			if (CurrentFrame.FrameID == Transition.SourceFrameID
 				&& CurrentHandle == Transition.TargetState)
 			{
 				bOnTargetBranch = true;
@@ -4397,11 +4395,26 @@ void FStateTreeExecutionContext::SetupNextTransition(const FStateTreeExecutionFr
 {
 	const FStateTreeExecutionState& Exec = GetExecState();
 
-	NextTransition.CurrentRunStatus = Exec.LastTickStatus;
+	NextTransition.SourceFrameID = CurrentFrame.FrameID;
+	NextTransition.SourceStateID = UE::StateTree::FActiveStateID();
+	if (CurrentlyProcessedState.IsValid())
+	{
+		const int32 CurrentStateIndex = CurrentFrame.ActiveStates.IndexOfReverse(CurrentlyProcessedState);
+		if (CurrentStateIndex != INDEX_NONE)
+		{
+			NextTransition.SourceStateID = CurrentFrame.ActiveStates.StateIDs[CurrentStateIndex];
+		}
+	}
+PRAGMA_DISABLE_DEPRECATION_WARNINGS
 	NextTransition.SourceState = CurrentlyProcessedState;
 	NextTransition.SourceStateTree = CurrentFrame.StateTree;
 	NextTransition.SourceRootState = CurrentFrame.ActiveStates.GetStateSafe(0);
+PRAGMA_ENABLE_DEPRECATION_WARNINGS
+
 	NextTransition.TargetState = NextState;
+	NextTransition.CurrentState = FStateTreeStateHandle::Invalid;
+	NextTransition.CurrentRunStatus = Exec.LastTickStatus;
+	NextTransition.ChangeType = EStateTreeStateChangeType::Changed; 
 	NextTransition.Priority = Priority;
 
 	FStateTreeExecutionFrame& NewFrame = NextTransition.NextActiveFrames.AddDefaulted_GetRef();
@@ -4863,11 +4876,9 @@ PRAGMA_ENABLE_DEPRECATION_WARNINGS
 					if (ensure(ParentLinkedState.IsValid()))
 					{
 						// Set the parent linked state as last completed state, and update tick status to the status from the transition.
-						const EStateTreeRunStatus RunStatus = NextTransition.TargetState.ToCompletionStatus();
-						STATETREE_LOG(Verbose, TEXT("Completed subtree '%s' from state '%s': %s"),
+						STATETREE_LOG(Verbose, TEXT("Completed subtree '%s' from global: %s"),
 							*GetSafeStateName(ParentFrame, ParentLinkedState),
-							*GetSafeStateName(CurrentFrame, NextTransition.SourceState),
-							*UEnum::GetDisplayValueAsText(RunStatus).ToString()
+							*UEnum::GetDisplayValueAsText(CastToRunStatus(FrameTasksStatus)).ToString()
 							);
 
 						const FCompactStateTreeState& State = ParentFrame.StateTree->States[ParentLinkedState.Index];
@@ -5019,11 +5030,8 @@ PRAGMA_ENABLE_DEPRECATION_WARNINGS
 	// Check if the transition was succeed/failed, if we're on a sub-tree, complete the subtree instead of transition.
 	if (NextTransition.TargetState.IsCompletionState() && bProcessSubTreeCompletion)
 	{
-		const int32 SourceFrameIndex = Exec.ActiveFrames.IndexOfByPredicate([&NextTransition = NextTransition](const FStateTreeExecutionFrame& Frame)
-		{
-			return Frame.StateTree == NextTransition.SourceStateTree && Frame.RootState == NextTransition.SourceRootState;
-		});
 		// Check that the transition source frame is a sub-tree, the first frame (0 index) is not a subtree. 
+		const int32 SourceFrameIndex = Exec.IndexOfActiveFrame(NextTransition.SourceFrameID);
 		if (SourceFrameIndex > 0)
 		{
 			const FStateTreeExecutionFrame& SourceFrame = Exec.ActiveFrames[SourceFrameIndex];
@@ -5033,9 +5041,19 @@ PRAGMA_ENABLE_DEPRECATION_WARNINGS
 
 			if (ParentLinkedState.IsValid())
 			{
-				const EStateTreeRunStatus RunStatus = NextTransition.TargetState.ToCompletionStatus(); 
+				const EStateTreeRunStatus RunStatus = NextTransition.TargetState.ToCompletionStatus();
+
+#if ENABLE_VISUAL_LOG
+				const int32 NextTransitionSourceIndex = SourceFrame.ActiveStates.IndexOfReverse(NextTransition.SourceStateID);
+				const FStateTreeStateHandle NextTransitionSourceState = NextTransitionSourceIndex != INDEX_NONE
+					? SourceFrame.ActiveStates[NextTransitionSourceIndex]
+					: FStateTreeStateHandle::Invalid;
 				STATETREE_LOG(Verbose, TEXT("Completed subtree '%s' from state '%s': %s"),
-					*GetSafeStateName(ParentFrame, ParentLinkedState), *GetSafeStateName(SourceFrame, NextTransition.SourceState), *UEnum::GetDisplayValueAsText(RunStatus).ToString());
+					*GetSafeStateName(ParentFrame, ParentLinkedState),
+					*GetSafeStateName(SourceFrame, NextTransitionSourceState),
+					*UEnum::GetDisplayValueAsText(RunStatus).ToString()
+					);
+#endif
 
 				// Set the parent linked state as last completed state, and update tick status to the status from the transition.
 				const UE::StateTree::ETaskCompletionStatus TaskStatus = UE::StateTree::ExecutionContext::Private::CastToTaskStatus(RunStatus);
@@ -5116,11 +5134,30 @@ TOptional<FStateTreeTransitionResult> FStateTreeExecutionContext::MakeTransition
 		return {};
 	}
 
-	Result.SourceState = RecordedTransition.SourceState;
+	// Try to find the same frame and the same state in the currently active frames.
+	// Recorded transitions can be saved and replayed out of context.
+	const FStateTreeExecutionState& Exec = GetExecState();
+	const FStateTreeExecutionFrame* ExecFrame = Exec.ActiveFrames.FindByPredicate([StateTree = RecordedTransition.SourceStateTree, RootState = RecordedTransition.SourceRootState](const FStateTreeExecutionFrame& Frame)
+		{
+			return Frame.StateTree == StateTree && Frame.RootState == RootState;
+		});
+	if (ExecFrame)
+	{
+		Result.SourceFrameID = ExecFrame->FrameID;
+		const int32 SourceStateIndex = ExecFrame->ActiveStates.IndexOfReverse(RecordedTransition.SourceState);
+		if (SourceStateIndex != INDEX_NONE)
+		{
+			Result.SourceStateID = ExecFrame->ActiveStates.StateIDs[SourceStateIndex];
+		}
+	}
 	Result.TargetState = RecordedTransition.TargetState;
 	Result.Priority = RecordedTransition.Priority;
+
+PRAGMA_DISABLE_DEPRECATION_WARNINGS
+	Result.SourceState = RecordedTransition.SourceState;
 	Result.SourceStateTree = RecordedTransition.SourceStateTree;
 	Result.SourceRootState = RecordedTransition.SourceRootState;
+PRAGMA_ENABLE_DEPRECATION_WARNINGS
 
 	return Result;
 }
@@ -5129,8 +5166,6 @@ FRecordedStateTreeTransitionResult FStateTreeExecutionContext::MakeRecordedTrans
 {
 	check(Transition.NextActiveFrames.Num() == Transition.NextActiveFrameEvents.Num());
 
-	const FStateTreeExecutionState& Exec = GetExecState();
-
 	FRecordedStateTreeTransitionResult Result;
 
 	for (int32 FrameIndex = 0; FrameIndex < Transition.NextActiveFrames.Num(); FrameIndex++)
@@ -5155,11 +5190,19 @@ FRecordedStateTreeTransitionResult FStateTreeExecutionContext::MakeRecordedTrans
 		}
 	}
 
-	Result.SourceState = Transition.SourceState;
+	const FStateTreeExecutionState& Exec = GetExecState();
+	if (const FStateTreeExecutionFrame* FoundSourceFrame = Exec.FindActiveFrame(Transition.SourceFrameID))
+	{
+		Result.SourceStateTree = FoundSourceFrame->StateTree;
+		Result.SourceRootState = FoundSourceFrame->RootState;
+		const int32 ActiveStateIndex = FoundSourceFrame->ActiveStates.IndexOfReverse(Transition.SourceStateID);
+		if (ActiveStateIndex != INDEX_NONE)
+		{
+			Result.SourceState = FoundSourceFrame->ActiveStates[ActiveStateIndex];
+		}
+	}
 	Result.TargetState = Transition.TargetState;
 	Result.Priority = Transition.Priority;
-	Result.SourceStateTree = Transition.SourceStateTree;
-	Result.SourceRootState = Transition.SourceRootState;
 
 	return Result;
 }
diff --git a/Engine/Plugins/Runtime/StateTree/Source/StateTreeModule/Public/StateTreeExecutionTypes.h b/Engine/Plugins/Runtime/StateTree/Source/StateTreeModule/Public/StateTreeExecutionTypes.h
--- a/Engine/Plugins/Runtime/StateTree/Source/StateTreeModule/Public/StateTreeExecutionTypes.h
+++ b/Engine/Plugins/Runtime/StateTree/Source/StateTreeModule/Public/StateTreeExecutionTypes.h
@@ -1139,15 +1139,21 @@ PRAGMA_ENABLE_DEPRECATION_WARNINGS
 	/** Events used in state selection. */
 	TArray<FStateTreeFrameStateSelectionEvents> NextActiveFrameEvents;
 
-	/** Current Run status. */
-	UPROPERTY(EditDefaultsOnly, Category = "Default", BlueprintReadOnly)
-	EStateTreeRunStatus CurrentRunStatus = EStateTreeRunStatus::Unset;
+	/** Frame that was active when the transition was requested. */
+	UE::StateTree::FActiveFrameID SourceFrameID;
 
+	/**
+	 * The state the transition was requested from.
+	 * It can be invalid if the transition is requested outside the Tick.
+	 */
+	UE::StateTree::FActiveStateID SourceStateID;
+
+	UE_DEPRECATED(5.6, "Use SourceStateID instead. The value is valid for Blueprint backward compatibility.")
 	/** Transition source state. */
 	UPROPERTY(EditDefaultsOnly, Category = "Default", BlueprintReadOnly)
 	FStateTreeStateHandle SourceState = FStateTreeStateHandle::Invalid;
 
-	/** Transition target state. */
+	/** Transition target state. It can be a completion state. */
 	UPROPERTY(EditDefaultsOnly, Category = "Default", BlueprintReadOnly)
 	FStateTreeStateHandle TargetState = FStateTreeStateHandle::Invalid;
 
@@ -1155,6 +1161,10 @@ PRAGMA_ENABLE_DEPRECATION_WARNINGS
 	UPROPERTY(EditDefaultsOnly, Category = "Default", BlueprintReadOnly)
 	FStateTreeStateHandle CurrentState = FStateTreeStateHandle::Invalid;
 	
+	/** Current Run status. */
+	UPROPERTY(EditDefaultsOnly, Category = "Default", BlueprintReadOnly)
+	EStateTreeRunStatus CurrentRunStatus = EStateTreeRunStatus::Unset;
+
 	/** If the change type is Sustained, then the CurrentState was reselected, or if Changed then the state was just activated. */
 	UPROPERTY(EditDefaultsOnly, Category = "Default", BlueprintReadOnly)
 	EStateTreeStateChangeType ChangeType = EStateTreeStateChangeType::Changed; 
@@ -1163,10 +1173,12 @@ PRAGMA_ENABLE_DEPRECATION_WARNINGS
 	UPROPERTY(EditDefaultsOnly, Category = "Default", BlueprintReadOnly)
 	EStateTreeTransitionPriority Priority = EStateTreeTransitionPriority::None;
 
+	UE_DEPRECATED(5.6, "Use SourceFrameID instead. The value is valid for Blueprint backward compatibility.")
 	/** StateTree asset that was active when the transition was requested. */
 	UPROPERTY(EditDefaultsOnly, Category = "Default", BlueprintReadOnly)
 	TObjectPtr<const UStateTree> SourceStateTree = nullptr;
 
+	UE_DEPRECATED(5.6, "Use SourceFrameID instead. The value is valid for Blueprint backward compatibility.")
 	/** Root state the execution frame where the transition was requested. */
 	UPROPERTY(EditDefaultsOnly, Category = "Default", BlueprintReadOnly)
 	FStateTreeStateHandle SourceRootState = FStateTreeStateHandle::Invalid;
@@ -1193,7 +1205,7 @@ struct STATETREEMODULE_API FRecordedStateTreeExecutionFrame
 	GENERATED_BODY()
 
 	FRecordedStateTreeExecutionFrame() = default;
-	UE_DEPRECATED(5.6, "Use FStateTreeExecutionContext::MakeRecordedTransitionResult to crate a new recorded transition.")
+	UE_DEPRECATED(5.6, "Use FStateTreeExecutionContext::MakeRecordedTransitionResult to create a new recorded transition.")
 	FRecordedStateTreeExecutionFrame(const FStateTreeExecutionFrame& ExecutionFrame);
 
 	/** The State Tree used for ticking this frame. */
@@ -1227,6 +1239,7 @@ struct STATETREEMODULE_API FRecordedStateTreeTransitionResult
 	GENERATED_BODY()
 
 	FRecordedStateTreeTransitionResult() = default;
+	UE_DEPRECATED(5.6, "Use FStateTreeExecutionContext::MakeRecordedTransitionResult to create a new recorded transition.")
 	FRecordedStateTreeTransitionResult(const FStateTreeTransitionResult& Transition);
 
    	/** States selected as result of the transition. */
