[State Tree] Deprecate WeakTaskRef and use NodeIndex in StateTreeWeakExecContext if it is created within the ExecContext tick loop
- Updated unit test for WeakContext FinishTask

diff --git a/Engine/Plugins/Runtime/StateTree/Source/StateTreeModule/Private/Blueprint/StateTreeTaskBlueprintBase.cpp b/Engine/Plugins/Runtime/StateTree/Source/StateTreeModule/Private/Blueprint/StateTreeTaskBlueprintBase.cpp
--- a/Engine/Plugins/Runtime/StateTree/Source/StateTreeModule/Private/Blueprint/StateTreeTaskBlueprintBase.cpp
+++ b/Engine/Plugins/Runtime/StateTree/Source/StateTreeModule/Private/Blueprint/StateTreeTaskBlueprintBase.cpp
@@ -36,7 +36,6 @@ EStateTreeRunStatus UStateTreeTaskBlueprintBase::EnterState(FStateTreeExecutionC
 {
 	// Task became active, cache event queue and owner.
 	SetCachedInstanceDataFromContext(Context);
-	WeakTaskRef = Context.MakeWeakTaskRefFromInstanceDataPtr(this);
 
 	FGuardValue_Bitfield(bIsProcessingEnterStateOrTick, true);
 
@@ -72,7 +71,6 @@ void UStateTreeTaskBlueprintBase::ExitState(FStateTreeExecutionContext& Context,
 	}
 
 	// Task became inactive, clear cached event queue and owner.
-	WeakTaskRef.Release();
 	ClearCachedInstanceData();
 }
 
@@ -109,7 +107,7 @@ void UStateTreeTaskBlueprintBase::FinishTask(const bool bSucceeded)
 	if (!bIsProcessingEnterStateOrTick)
 	{
 		EStateTreeFinishTaskType CompletionType = bSucceeded ? EStateTreeFinishTaskType::Succeeded : EStateTreeFinishTaskType::Failed;
-		GetWeakExecutionContext().FinishTask(WeakTaskRef, CompletionType);
+		GetWeakExecutionContext().FinishTask(CompletionType);
 	}
 }
 
@@ -124,7 +122,7 @@ void UStateTreeTaskBlueprintBase::BroadcastDelegate(FStateTreeDelegateDispatcher
 void UStateTreeTaskBlueprintBase::BindDelegate(const FStateTreeDelegateListener& Listener, const FStateTreeDynamicDelegate& Delegate)
 {
 	FSimpleDelegate SimpleDelegate = FSimpleDelegate::CreateUFunction(const_cast<UObject*>(Delegate.GetUObject()), Delegate.GetFunctionName());
-	if (!GetWeakExecutionContext().BindDelegate(WeakTaskRef, Listener, MoveTemp(SimpleDelegate)))
+	if (!GetWeakExecutionContext().BindDelegate(Listener, MoveTemp(SimpleDelegate)))
 	{
 		UE_VLOG_UELOG(this, LogStateTree, Error, TEXT("Failed to bind the delegate. The instance probably stopped."));
 	}
diff --git a/Engine/Plugins/Runtime/StateTree/Source/StateTreeModule/Private/StateTreeExecutionContext.cpp b/Engine/Plugins/Runtime/StateTree/Source/StateTreeModule/Private/StateTreeExecutionContext.cpp
--- a/Engine/Plugins/Runtime/StateTree/Source/StateTreeModule/Private/StateTreeExecutionContext.cpp
+++ b/Engine/Plugins/Runtime/StateTree/Source/StateTreeModule/Private/StateTreeExecutionContext.cpp
@@ -1010,14 +1010,17 @@ FStateTreeWeakExecutionContext FStateTreeExecutionContext::MakeWeakExecutionCont
 	return FStateTreeWeakExecutionContext(*this);
 }
 
+PRAGMA_DISABLE_DEPRECATION_WARNINGS
 FStateTreeWeakTaskRef FStateTreeExecutionContext::MakeWeakTaskRef(const FStateTreeTaskBase& Node) const
 {
+	// This function has been deprecated
 	check(CurrentNode == &Node);
 	return MakeWeakTaskRefInternal();
 }
 
 FStateTreeWeakTaskRef FStateTreeExecutionContext::MakeWeakTaskRefInternal() const
 {
+	// This function has been deprecated
 	FStateTreeWeakTaskRef Result;
 	if (const FStateTreeExecutionFrame* Frame = GetCurrentlyProcessedFrame())
 	{
@@ -1029,6 +1032,7 @@ FStateTreeWeakTaskRef FStateTreeExecutionContext::MakeWeakTaskRefInternal() cons
 	}
 	return Result;
 }
+PRAGMA_ENABLE_DEPRECATION_WARNINGS
 
 EStateTreeRunStatus FStateTreeExecutionContext::Start(const FInstancedPropertyBag* InitialParameters, int32 RandomSeed)
 {
@@ -6172,6 +6176,9 @@ FStateTreeWeakExecutionContext::FStateTreeWeakExecutionContext(const FStateTreeE
 		{
 			StateID = Frame->ActiveStates.StateIDs[ActiveStateIndex];
 		}
+
+		// Only writes down the NodeIndex if it is created synchronously from the ExecContext tick frame
+		NodeIndex = Context.GetCurrentlyProcessedNodeIndex();
 	}
 }
 
@@ -6186,7 +6193,7 @@ bool FStateTreeWeakExecutionContext::SendEvent(const FGameplayTag Tag, const FCo
 		//The context may already exist and we want to confirm that we are on the same thread.
 		StoragePinned->AcquireWriteAccess();
 
-		if (IsFrameOrStateActive(StoragePinned->GetExecutionState()))
+		if (IsFrameOrStateActive(StoragePinned->GetMutableExecutionState()))
 		{
 			FStateTreeMinimalExecutionContext Context(*OwnerPinned.Get(), *StateTreePinned.Get(), *StoragePinned.Get());
 			Context.SendEvent(Tag, Payload, Origin);
@@ -6214,19 +6221,23 @@ bool FStateTreeWeakExecutionContext::RequestTransition(FStateTreeStateHandle Tar
 		FStateTreeExecutionState& Exec = StoragePinned->GetMutableExecutionState();
 		if (IsFrameOrStateActive(Exec) && Exec.ActiveFrames.Num() > 0)
 		{
-			FStateTreeTransitionRequest Request;
+			// don't request if the frame is no longer active
+			if (const FStateTreeExecutionFrame* CurrentFrame = Exec.FindActiveFrame(FrameID))
+			{
+				FStateTreeTransitionRequest Request;
 
-			Request.SourceFrameID = FrameID;
-			Request.SourceStateID = StateID;
+				Request.SourceFrameID = FrameID;
+				Request.SourceStateID = StateID;
 
-			Request.TargetState = TargetState;
-			Request.Priority = Priority;
-			Request.Fallback = Fallback;
+				Request.TargetState = TargetState;
+				Request.Priority = Priority;
+				Request.Fallback = Fallback;
 
-			StoragePinned->AddTransitionRequest(OwnerPinned.Get(), Request);
+				StoragePinned->AddTransitionRequest(OwnerPinned.Get(), Request);
 
-			ScheduleNextTick(OwnerPinned.Get(), StateTreePinned.Get(), *StoragePinned);
-			bRequested = true;
+				ScheduleNextTick(OwnerPinned.Get(), StateTreePinned.Get(), *StoragePinned);
+				bRequested = true;
+			}
 		}
 
 		StoragePinned->ReleaseWriteAccess();
@@ -6274,7 +6285,7 @@ bool FStateTreeWeakExecutionContext::BroadcastDelegate(const FStateTreeDelegateD
 	return bBroadcasted;
 }
 
-bool FStateTreeWeakExecutionContext::BindDelegate(const FStateTreeWeakTaskRef& Task, const FStateTreeDelegateListener& Listener, FSimpleDelegate Delegate) const
+bool FStateTreeWeakExecutionContext::BindDelegate(const FStateTreeDelegateListener& Listener, FSimpleDelegate Delegate) const
 {
 	if (!Listener.IsValid())
 	{
@@ -6298,11 +6309,10 @@ bool FStateTreeWeakExecutionContext::BindDelegate(const FStateTreeWeakTaskRef& T
 		{
 			if (const FStateTreeExecutionFrame* CurrentFrame = Exec.FindActiveFrame(FrameID))
 			{
-				const FStateTreeStrongTaskRef TaskPin = Task.Pin();
 				// The user uses a task that is not part of this context (ex: use task from state A inside context constructed from state B)
-				if (ensure(TaskPin.IsValid() && TaskPin.GetStateTree() == CurrentFrame->StateTree))
+				if (ensure(CurrentFrame->StateTree))
 				{
-					Exec.DelegateActiveListeners.Add(Listener, MoveTemp(Delegate), FrameID, StateID, TaskPin.GetTaskIndex());
+					Exec.DelegateActiveListeners.Add(Listener, MoveTemp(Delegate), FrameID, StateID, NodeIndex);
 					bResult = true;
 				}
 			}
@@ -6314,6 +6324,13 @@ bool FStateTreeWeakExecutionContext::BindDelegate(const FStateTreeWeakTaskRef& T
 	return bResult;
 }
 
+PRAGMA_DISABLE_DEPRECATION_WARNINGS
+bool FStateTreeWeakExecutionContext::BindDelegate(const FStateTreeWeakTaskRef& Task, const FStateTreeDelegateListener& Listener, FSimpleDelegate Delegate) const
+{
+	return BindDelegate(Listener, Delegate);
+}
+PRAGMA_ENABLE_DEPRECATION_WARNINGS
+
 bool FStateTreeWeakExecutionContext::RemoveDelegateListener(const FStateTreeDelegateListener& Listener) const
 {
 	return UnbindDelegate(Listener);
@@ -6351,7 +6368,7 @@ bool FStateTreeWeakExecutionContext::UnbindDelegate(const FStateTreeDelegateList
 	return bRemoved;
 }
 
-bool FStateTreeWeakExecutionContext::FinishTask(const FStateTreeWeakTaskRef& Task, EStateTreeFinishTaskType FinishType) const
+bool FStateTreeWeakExecutionContext::FinishTask(EStateTreeFinishTaskType FinishType) const
 {
 	bool bSucceed = false;
 
@@ -6367,25 +6384,24 @@ bool FStateTreeWeakExecutionContext::FinishTask(const FStateTreeWeakTaskRef& Tas
 		FStateTreeExecutionState& Exec = StoragePinned->GetMutableExecutionState();
 		if (IsFrameOrStateActive(Exec))
 		{
-			FStateTreeStrongTaskRef TaskPin = Task.Pin();
-			FStateTreeExecutionFrame* ExecutionFrame = Exec.FindActiveFrame(FrameID);
-			if (TaskPin.IsValid() && ExecutionFrame)
+			if (FStateTreeExecutionFrame* ExecutionFrame = Exec.FindActiveFrame(FrameID))
 			{
-				const bool bIsSameTree = ExecutionFrame->StateTree == TaskPin.GetStateTree();
-				if (ensureMsgf(bIsSameTree, TEXT("User error. The task is not from the same execution tree")))
+				if (const UStateTree* FrameStateTree = ExecutionFrame->StateTree)
 				{
 					using namespace UE::StateTree;
 
+					const int32 AssetNodeIndex = NodeIndex.AsInt32();
 					const int32 GlobalTaskBeginIndex = ExecutionFrame->StateTree->GlobalTasksBegin;
 					const int32 GlobalTaskEndIndex = ExecutionFrame->StateTree->GlobalTasksBegin + ExecutionFrame->StateTree->GlobalTasksNum;
-					const int32 AssetTaskIndex = TaskPin.GetTaskIndex().AsInt32();
-					const bool bIsGlobalTask = AssetTaskIndex >= GlobalTaskBeginIndex && AssetTaskIndex < GlobalTaskEndIndex;
-					const ETaskCompletionStatus TaskStatus = UE::StateTree::ExecutionContext::Private::CastToTaskStatus(FinishType);
+					const bool bIsGlobalTask = AssetNodeIndex >= GlobalTaskBeginIndex && AssetNodeIndex < GlobalTaskEndIndex;
+					const ETaskCompletionStatus TaskStatus = ExecutionContext::Private::CastToTaskStatus(FinishType);
 
 					bool bCompleted = false;
-					if (bIsGlobalTask)
+					if (ExecutionFrame->bIsGlobalFrame && bIsGlobalTask)
 					{
-						const int32 FrameTaskIndex = AssetTaskIndex - GlobalTaskBeginIndex;
+						check(FrameStateTree->GetNode(AssetNodeIndex).GetPtr<const FStateTreeTaskBase>() != nullptr);
+
+						const int32 FrameTaskIndex = AssetNodeIndex - GlobalTaskBeginIndex;
 						FTasksCompletionStatus GlobalTasksStatus = ExecutionFrame->ActiveTasksStatus.GetStatus(ExecutionFrame->StateTree);
 						GlobalTasksStatus.SetStatusWithPriority(FrameTaskIndex, TaskStatus);
 						bCompleted = GlobalTasksStatus.IsCompleted();
@@ -6393,22 +6409,22 @@ bool FStateTreeWeakExecutionContext::FinishTask(const FStateTreeWeakTaskRef& Tas
 					}
 					else
 					{
-						// Find the task's state.
-						for (int32 StateIndex = 0; StateIndex < ExecutionFrame->ActiveStates.Num(); ++StateIndex)
+						const FStateTreeStateHandle StateHandle = ExecutionFrame->ActiveStates.FindStateHandle(StateID);
+						if (StateHandle.IsValid())
 						{
-							const FStateTreeStateHandle CurrentHandle = ExecutionFrame->ActiveStates[StateIndex];
-							const FCompactStateTreeState& State = ExecutionFrame->StateTree->States[CurrentHandle.Index];
+							const FCompactStateTreeState& State = ExecutionFrame->StateTree->States[StateHandle.Index];
 							const int32 StateTaskBeginIndex = State.TasksBegin;
 							const int32 StateTaskEndIndex = State.TasksBegin + State.TasksNum;
-							const bool bIsStateTask = AssetTaskIndex >= StateTaskBeginIndex && AssetTaskIndex < StateTaskEndIndex;
+							const bool bIsStateTask = AssetNodeIndex >= StateTaskBeginIndex && AssetNodeIndex < StateTaskEndIndex;
 							if (bIsStateTask)
 							{
-								const int32 StateTaskIndex = AssetTaskIndex - State.TasksBegin;
+								check(FrameStateTree->GetNode(AssetNodeIndex).GetPtr<const FStateTreeTaskBase>() != nullptr);
+
+								const int32 StateTaskIndex = AssetNodeIndex - State.TasksBegin;
 								FTasksCompletionStatus StateTasksStatus = ExecutionFrame->ActiveTasksStatus.GetStatus(State);
 								StateTasksStatus.SetStatusWithPriority(StateTaskIndex, TaskStatus);
 								bCompleted = StateTasksStatus.IsCompleted();
 								bSucceed = true;
-								break;
 							}
 						}
 					}
@@ -6428,6 +6444,12 @@ bool FStateTreeWeakExecutionContext::FinishTask(const FStateTreeWeakTaskRef& Tas
 	return bSucceed;
 }
 
+PRAGMA_DISABLE_DEPRECATION_WARNINGS
+bool FStateTreeWeakExecutionContext::FinishTask(const FStateTreeWeakTaskRef& Task, EStateTreeFinishTaskType FinishType) const
+{
+	return FinishTask(FinishType);
+}
+PRAGMA_ENABLE_DEPRECATION_WARNINGS
 
 bool FStateTreeWeakExecutionContext::UpdateScheduledTickRequest(UE::StateTree::FScheduledTickHandle Handle, FStateTreeScheduledTick ScheduledTick) const
 {
@@ -6443,7 +6465,6 @@ bool FStateTreeWeakExecutionContext::UpdateScheduledTickRequest(UE::StateTree::F
 		StoragePinned->AcquireWriteAccess();
 
 		bCompleted = true;
-
 		FStateTreeExecutionState& Exec = StoragePinned->GetMutableExecutionState();
 		if (IsFrameOrStateActive(Exec))
 		{
@@ -6481,6 +6502,7 @@ bool FStateTreeWeakExecutionContext::IsFrameOrStateActive(const FStateTreeExecut
 		return false;
 	}
 
+	// state or global node
 	return StateID.IsValid() ? CurrentFrame->ActiveStates.Contains(StateID) : true;
 }
 
diff --git a/Engine/Plugins/Runtime/StateTree/Source/StateTreeModule/Private/StateTreeNodeRef.cpp b/Engine/Plugins/Runtime/StateTree/Source/StateTreeModule/Private/StateTreeNodeRef.cpp
--- a/Engine/Plugins/Runtime/StateTree/Source/StateTreeModule/Private/StateTreeNodeRef.cpp
+++ b/Engine/Plugins/Runtime/StateTree/Source/StateTreeModule/Private/StateTreeNodeRef.cpp
@@ -6,7 +6,7 @@
 #include "StateTreeInstanceData.h"
 #include "StateTreeTaskBase.h"
 
-
+PRAGMA_DISABLE_DEPRECATION_WARNINGS
 #if WITH_STATETREE_DEBUG
 FStateTreeStrongTaskRef::FStateTreeStrongTaskRef(TStrongObjectPtr<const UStateTree> InStateTree, const FStateTreeTaskBase* InTask, FStateTreeIndex16 InNodeIndex, FGuid InNodeId)
 	: StateTree(InStateTree)
@@ -74,3 +74,4 @@ FStateTreeStrongTaskRef FStateTreeWeakTaskRef::Pin() const
 	return Task ? FStateTreeStrongTaskRef(StateTreePinned, Task, NodeIndex) : FStateTreeStrongTaskRef();
 #endif
 }
+PRAGMA_ENABLE_DEPRECATION_WARNINGS
\ No newline at end of file
diff --git a/Engine/Plugins/Runtime/StateTree/Source/StateTreeModule/Public/Blueprint/StateTreeTaskBlueprintBase.h b/Engine/Plugins/Runtime/StateTree/Source/StateTreeModule/Public/Blueprint/StateTreeTaskBlueprintBase.h
--- a/Engine/Plugins/Runtime/StateTree/Source/StateTreeModule/Public/Blueprint/StateTreeTaskBlueprintBase.h
+++ b/Engine/Plugins/Runtime/StateTree/Source/StateTreeModule/Public/Blueprint/StateTreeTaskBlueprintBase.h
@@ -102,8 +102,13 @@ protected:
 	UFUNCTION(BlueprintCallable, Category = "StateTree")
 	UE_API void UnbindDelegate(const FStateTreeDelegateListener& Listener);
 
+#if WITH_EDITORONLY_DATA
+	UE_DEPRECATED(5.6, "FStateTreeWeakTaskRef is deprecated.")
 	/** Cached task while the node is active for async execution. */
+	PRAGMA_DISABLE_DEPRECATION_WARNINGS
 	mutable FStateTreeWeakTaskRef WeakTaskRef;
+	PRAGMA_ENABLE_DEPRECATION_WARNINGS
+#endif
 
 	/** Run status when using latent EnterState and Tick */
 	EStateTreeRunStatus RunStatus = EStateTreeRunStatus::Running;
diff --git a/Engine/Plugins/Runtime/StateTree/Source/StateTreeModule/Public/StateTreeExecutionContext.h b/Engine/Plugins/Runtime/StateTree/Source/StateTreeModule/Public/StateTreeExecutionContext.h
--- a/Engine/Plugins/Runtime/StateTree/Source/StateTreeModule/Public/StateTreeExecutionContext.h
+++ b/Engine/Plugins/Runtime/StateTree/Source/StateTreeModule/Public/StateTreeExecutionContext.h
@@ -333,29 +333,38 @@ public:
 	 * @return a weak reference for a task that can be stored for later uses.
 	 * @note similar to GetInstanceData, the node needs to be the current processing node.
 	 */
+PRAGMA_DISABLE_DEPRECATION_WARNINGS
+	UE_DEPRECATED(5.6, "FStateTreeWeakTaskRef is no longer used.")
 	UE_API FStateTreeWeakTaskRef MakeWeakTaskRef(const FStateTreeTaskBase& Node) const;
+PRAGMA_ENABLE_DEPRECATION_WARNINGS
 
 	/**
 	 * @return a weak reference for a task that can be stored for later uses.
 	 * @note similar to GetInstanceData, the instance data needs to be the current processing node.
 	 */
+PRAGMA_DISABLE_DEPRECATION_WARNINGS
 	template <typename T>
+	UE_DEPRECATED(5.6, "FStateTreeWeakTaskRef is no longer used.")
 	FStateTreeWeakTaskRef MakeWeakTaskRefFromInstanceData(const T& InInstanceData) const
 	{
 		check(&CurrentNodeInstanceData.template GetMutable<T>() == &InInstanceData);
 		return MakeWeakTaskRefInternal();
 	}
+PRAGMA_ENABLE_DEPRECATION_WARNINGS
 
 	/**
 	 * @return a weak reference for a task that can be stored for later uses.
 	 * @note similar to GetInstanceData, the instance data needs to be the current processing node.
 	 */
+PRAGMA_DISABLE_DEPRECATION_WARNINGS
 	template <typename T>
+	UE_DEPRECATED(5.6, "FStateTreeWeakTaskRef is no longer used.")
 	FStateTreeWeakTaskRef MakeWeakTaskRefFromInstanceDataPtr(const T* InInstanceData) const
 	{
 		check(CurrentNodeInstanceData.template GetMutablePtr<T>() == InInstanceData);
 		return MakeWeakTaskRefInternal();
 	}
+PRAGMA_ENABLE_DEPRECATION_WARNINGS
 
 	struct FStartParameters
 	{
@@ -495,6 +504,17 @@ public:
 	UE_DEPRECATED(5.5, "Use GetEventsToProcessView() instead.")
 	TConstArrayView<FStateTreeEvent> GetEventsToProcess() const { return {}; }
 
+	FStateTreeIndex16 GetCurrentlyProcessedNodeIndex() const
+	{
+		return FStateTreeIndex16(CurrentNodeIndex);
+	}
+
+	FStateTreeDataHandle GetCurrentlyProcessedNodeInstanceData() const
+	{
+		return CurrentNodeDataHandle;
+	}
+
+	UE_DEPRECATED(5.6, "Use GetCurrentlyProcessedNodeInstanceData() instead.")
 	/** @return the currently processed node if applicable. */
 	FStateTreeDataHandle GetCurrentlyProcessedNode() const
 	{
@@ -1032,7 +1052,10 @@ PRAGMA_ENABLE_DEPRECATION_WARNINGS
 	UE_API void CaptureNewStateEvents(TConstArrayView<FStateTreeExecutionFrame> PrevFrames, TConstArrayView<FStateTreeExecutionFrame> NewFrames, TArrayView<FStateTreeFrameStateSelectionEvents> FramesStateSelectionEvents);
 
 	/** @return a weak reference for a task that can be stored for later uses. */
+PRAGMA_DISABLE_DEPRECATION_WARNINGS
+	UE_DEPRECATED(5.6, "FStateTreeWeakTaskRef is no longer used.")
 	UE_API FStateTreeWeakTaskRef MakeWeakTaskRefInternal() const;
+PRAGMA_ENABLE_DEPRECATION_WARNINGS
 
 	/** Instance data used during current tick. */
 	FStateTreeInstanceData& InstanceData;
@@ -1294,8 +1317,8 @@ public:
 
 	/**
 	 * Requests transition to a state.
-	 * If called during during transition processing (e.g. from FStateTreeTaskBase::TriggerTransitions()) the transition
-	 * is attempted to be activate immediately (it can fail e.g. because of preconditions on a target state).
+	 * If called during transition processing (e.g. from FStateTreeTaskBase::TriggerTransitions()) the transition
+	 * is attempted to be activated immediately (it can fail e.g. because of preconditions on a target state).
 	 * If called outside the transition handling, the request is buffered and handled at the beginning of next transition processing.
 	 * @param TargetState The state to transition to.
 	 * @param Priority The priority of the transition.
@@ -1317,7 +1340,12 @@ public:
 	 * The listener is bound to a dispatcher in the editor.
 	 * @return false if the context is not valid and the bind failed.
 	 */
+	UE_API bool BindDelegate(const FStateTreeDelegateListener& Listener, FSimpleDelegate Delegate) const;
+
+PRAGMA_DISABLE_DEPRECATION_WARNINGS
+	UE_DEPRECATED(5.6, "Use the version without FStateTreeWeakTaskRef")
 	UE_API bool BindDelegate(const FStateTreeWeakTaskRef& Task, const FStateTreeDelegateListener& Listener, FSimpleDelegate Delegate) const;
+PRAGMA_ENABLE_DEPRECATION_WARNINGS
 
 	UE_DEPRECATED(5.6, "Use UnbindDelegate")
 	/** Removes Delegate Listener. */
@@ -1334,7 +1362,12 @@ public:
 	 * If called during tick processing, then the state completes immediately.
 	 * If called outside of the tick processing, then the request is buffered and handled on the next tick.
 	 */
+	UE_API bool FinishTask(EStateTreeFinishTaskType FinishType) const;
+
+PRAGMA_DISABLE_DEPRECATION_WARNINGS
+	UE_DEPRECATED(5.6, "Use the version without FStateTreeWeakTaskRef.")
 	UE_API bool FinishTask(const FStateTreeWeakTaskRef& Task, EStateTreeFinishTaskType FinishType) const;
+PRAGMA_ENABLE_DEPRECATION_WARNINGS
 	
 	/**
 	 * Updates the scheduled tick of a previous request.
@@ -1344,6 +1377,11 @@ public:
 
 private:
 	static UE_API void ScheduleNextTick(TNotNull<UObject*> Owner, TNotNull<const UStateTree*> RootStateTree, FStateTreeInstanceStorage& Storage);
+
+	/**
+	 * check if the recorded Frame ID and/or State ID is still active
+	 * @return if the recorded Frame ID and/or(global or state) State ID is still active, or if the FrameID is invalid(created outside the ExecContext loop)
+	 */
 	UE_API bool IsFrameOrStateActive(const FStateTreeExecutionState& Exec) const;
 
 	TWeakObjectPtr<UObject> Owner;
@@ -1352,6 +1390,7 @@ private:
 
 	UE::StateTree::FActiveFrameID FrameID;
 	UE::StateTree::FActiveStateID StateID;
+	FStateTreeIndex16 NodeIndex;
 };
 
 #undef UE_API
diff --git a/Engine/Plugins/Runtime/StateTree/Source/StateTreeModule/Public/StateTreeExecutionTypes.h b/Engine/Plugins/Runtime/StateTree/Source/StateTreeModule/Public/StateTreeExecutionTypes.h
--- a/Engine/Plugins/Runtime/StateTree/Source/StateTreeModule/Public/StateTreeExecutionTypes.h
+++ b/Engine/Plugins/Runtime/StateTree/Source/StateTreeModule/Public/StateTreeExecutionTypes.h
@@ -413,6 +413,19 @@ struct FStateTreeActiveStates
 		NumStates = static_cast<uint8>(NewNum);
 	}
 
+	FStateTreeStateHandle FindStateHandle(UE::StateTree::FActiveStateID StateId) const
+	{
+		for (int32 Index = (int32)NumStates - 1; Index >= 0; Index--)
+		{
+			if (StateIDs[Index] == StateId)
+			{
+				return States[Index];
+			}
+		}
+
+		return FStateTreeStateHandle::Invalid;
+	}
+
 	/** Returns true of the array contains specified state. */
 	bool Contains(const FStateTreeStateHandle StateHandle) const
 	{
diff --git a/Engine/Plugins/Runtime/StateTree/Source/StateTreeModule/Public/StateTreeNodeRef.h b/Engine/Plugins/Runtime/StateTree/Source/StateTreeModule/Public/StateTreeNodeRef.h
--- a/Engine/Plugins/Runtime/StateTree/Source/StateTreeModule/Public/StateTreeNodeRef.h
+++ b/Engine/Plugins/Runtime/StateTree/Source/StateTreeModule/Public/StateTreeNodeRef.h
@@ -12,7 +12,7 @@ struct FStateTreeTaskBase;
 /**
  * A reference to a task that can utilized in a async callback. Use FStateTreeWeakTaskRef to store the reference and Pin it to get the strong version. Similar to StrongPtr and WeakPtr.
  */
-struct FStateTreeStrongTaskRef
+struct UE_DEPRECATED(5.6, "FStateTreeStrongTaskRef is deprecated. We now use TaskIndex in WeakExecutionContext.") FStateTreeStrongTaskRef
 {
 	explicit FStateTreeStrongTaskRef() = default;
 #if WITH_STATETREE_DEBUG
@@ -48,12 +48,15 @@ private:
 /**
  * A reference to a task that can be retreive. Similar to StrongPtr and WeakPtr.
  */
-struct FStateTreeWeakTaskRef
+struct UE_DEPRECATED(5.6, "FStateTreeWeakTaskRef is deprecated. We now use TaskIndex in WeakExecutionContext.") FStateTreeWeakTaskRef
 {
 	explicit FStateTreeWeakTaskRef() = default;
 	STATETREEMODULE_API explicit FStateTreeWeakTaskRef(TNotNull<const UStateTree*> StateTree, FStateTreeIndex16 TaskIndex);
 
+	PRAGMA_DISABLE_DEPRECATION_WARNINGS
 	STATETREEMODULE_API FStateTreeStrongTaskRef Pin() const;
+	PRAGMA_ENABLE_DEPRECATION_WARNINGS
+
 	void Release()
 	{
 		*this = FStateTreeWeakTaskRef();
diff --git a/Engine/Plugins/Runtime/StateTree/Source/StateTreeTestSuite/Private/StateTreeWeakContextTest.cpp b/Engine/Plugins/Runtime/StateTree/Source/StateTreeTestSuite/Private/StateTreeWeakContextTest.cpp
--- a/Engine/Plugins/Runtime/StateTree/Source/StateTreeTestSuite/Private/StateTreeWeakContextTest.cpp
+++ b/Engine/Plugins/Runtime/StateTree/Source/StateTreeTestSuite/Private/StateTreeWeakContextTest.cpp
@@ -16,131 +16,84 @@ namespace UE::StateTree::Tests
 {
 struct FStateTreeTest_WeakContext_FinishTask : FStateTreeTestBase
 {
+	//Tree 1 : Global Task
+		//	Root : Task 
+		//		State1 : Task -> Root
+		//			State2 : Task -> Root
 	virtual bool InstantTest() override
 	{
-		//Tree 1 (with globals)
-		//	Root
-		//		State1 -> Next
-		//			State2 -> Success
-
 		struct FWeakContext
 		{
-			FStateTreeWeakExecutionContext ContextTree1GlobalTask1;
-			FStateTreeWeakExecutionContext ContextTree1GlobalTask2;
-			FStateTreeWeakExecutionContext ContextTree1RootTask1;
-			FStateTreeWeakExecutionContext ContextTree1RootTask2;
-			FStateTreeWeakExecutionContext ContextTree1State1Task1;
-			FStateTreeWeakExecutionContext ContextTree1State1Task2;
-			FStateTreeWeakExecutionContext ContextTree1State2Task1;
-			FStateTreeWeakExecutionContext ContextTree1State2Task2;
-			FStateTreeWeakTaskRef TaskTree1GlobalTask1;
-			FStateTreeWeakTaskRef TaskTree1GlobalTask2;
-			FStateTreeWeakTaskRef TaskTree1RootTask1;
-			FStateTreeWeakTaskRef TaskTree1RootTask2;
-			FStateTreeWeakTaskRef TaskTree1State1Task1;
-			FStateTreeWeakTaskRef TaskTree1State1Task2;
-			FStateTreeWeakTaskRef TaskTree1State2Task1;
-			FStateTreeWeakTaskRef TaskTree1State2Task2;
+			FStateTreeWeakExecutionContext ContextTree1GlobalTask;
+			FStateTreeWeakExecutionContext ContextTree1RootTask;
+			FStateTreeWeakExecutionContext ContextTree1State1Task;
+			FStateTreeWeakExecutionContext ContextTree1State2Task;
 
-			bool bWhen_Tree1GlobalTask1_Finish_Tree1GlobalTask2 = false;
-			bool bWhen_TaskTree1State1Task1_Finish_TaskTree1State1Task2 = false;
-			bool bWhen_TaskTree1State1Task2_Finish_TaskTree1State1Task1 = false;
-			bool bWhen_TaskTree1State1Task2_Finish_TaskTree1State1Task2 = false;
-			bool bWhen_TaskTree1State1Task2_Finish_Tree1GlobalTask2 = false;
+			bool bGlobalFinishTaskSuccessOnTick = false;
+			bool bState1lFinishTaskFailOnTick = false;
 		};
+
 		FWeakContext WeakContext;
 
+		// Building up the State Tree
 		UStateTree& StateTree1 = NewStateTree();
 		{
 			UStateTreeEditorData& EditorData1 = *Cast<UStateTreeEditorData>(StateTree1.EditorData);
 			{
-				TStateTreeEditorNode<FTestTask_PrintValue>& Tree1GlobalTask1 = EditorData1.AddGlobalTask<FTestTask_PrintValue>(FName(TEXT("Tree1GlobalTask1")));
-				Tree1GlobalTask1.GetNode().CustomEnterStateFunc = [&WeakContext](FStateTreeExecutionContext& Context, const FTestTask_PrintValue* Task)
-					{
-						WeakContext.ContextTree1GlobalTask1 = Context.MakeWeakExecutionContext();
-						WeakContext.TaskTree1GlobalTask1 = Context.MakeWeakTaskRef(*Task);
-					};
-				Tree1GlobalTask1.GetNode().CustomTickFunc = [&WeakContext](FStateTreeExecutionContext& Context, const FTestTask_PrintValue* Task)
-					{
-						if (WeakContext.bWhen_Tree1GlobalTask1_Finish_Tree1GlobalTask2)
+				// Global Task
+				TStateTreeEditorNode<FTestTask_PrintValue>& Tree1GlobalTask = EditorData1.AddGlobalTask<FTestTask_PrintValue>(FName(TEXT("Tree1GlobalTask")));
+				{
+					Tree1GlobalTask.GetNode().CustomEnterStateFunc = [&WeakContext](FStateTreeExecutionContext& Context, const FTestTask_PrintValue* Task)
 						{
-							WeakContext.ContextTree1GlobalTask1.FinishTask(WeakContext.TaskTree1GlobalTask1, EStateTreeFinishTaskType::Failed);
-						}
-					};
-				TStateTreeEditorNode<FTestTask_PrintValue>& Tree1GlobalTask2 = EditorData1.AddGlobalTask<FTestTask_PrintValue>(FName(TEXT("Tree1GlobalTask2")));
-				Tree1GlobalTask2.GetNode().CustomEnterStateFunc = [&WeakContext](FStateTreeExecutionContext& Context, const FTestTask_PrintValue* Task)
-					{
-						WeakContext.ContextTree1GlobalTask2 = Context.MakeWeakExecutionContext();
-						WeakContext.TaskTree1GlobalTask2 = Context.MakeWeakTaskRef(*Task);
-					};
-			}
-			UStateTreeState& Root1 = EditorData1.AddSubTree("Tree1StateRoot");
-			{
-				TStateTreeEditorNode<FTestTask_PrintValue>& Tree1StateRootTask1 = Root1.AddTask<FTestTask_PrintValue>(FName(TEXT("Tree1RootTask1")));
-				Tree1StateRootTask1.GetNode().CustomEnterStateFunc = [&WeakContext](FStateTreeExecutionContext& Context, const FTestTask_PrintValue* Task)
-					{
-						WeakContext.ContextTree1RootTask1 = Context.MakeWeakExecutionContext();
-						WeakContext.TaskTree1RootTask1 = Context.MakeWeakTaskRef(*Task);
-					};
-				TStateTreeEditorNode<FTestTask_PrintValue>& Tree1StateRootTask2 = Root1.AddTask<FTestTask_PrintValue>(FName(TEXT("Tree1RootTask2")));
-				Tree1StateRootTask2.GetNode().CustomEnterStateFunc = [&WeakContext](FStateTreeExecutionContext& Context, const FTestTask_PrintValue* Task)
-					{
-						WeakContext.ContextTree1RootTask2 = Context.MakeWeakExecutionContext();
-						WeakContext.TaskTree1RootTask2 = Context.MakeWeakTaskRef(*Task);
-					};
-			}
-			UStateTreeState& State1 = Root1.AddChildState("Tree1State1");
-			{
-				TStateTreeEditorNode<FTestTask_PrintValue>& Tree1State1Task1 = State1.AddTask<FTestTask_PrintValue>(FName(TEXT("Tree1State1Task1")));
-				Tree1State1Task1.GetNode().CustomEnterStateFunc = [&WeakContext](FStateTreeExecutionContext& Context, const FTestTask_PrintValue* Task)
-					{
-						WeakContext.ContextTree1State1Task1 = Context.MakeWeakExecutionContext();
-						WeakContext.TaskTree1State1Task1 = Context.MakeWeakTaskRef(*Task);
-					};
-				Tree1State1Task1.GetNode().CustomTickFunc = [&WeakContext](FStateTreeExecutionContext& Context, const FTestTask_PrintValue* Task)
-					{
-						if (WeakContext.bWhen_TaskTree1State1Task1_Finish_TaskTree1State1Task2)
+							WeakContext.ContextTree1GlobalTask = Context.MakeWeakExecutionContext();
+						};
+					Tree1GlobalTask.GetNode().CustomTickFunc = [&WeakContext](FStateTreeExecutionContext& Context, const FTestTask_PrintValue* Task)
 						{
-							WeakContext.ContextTree1State1Task2.FinishTask(WeakContext.TaskTree1State1Task2, EStateTreeFinishTaskType::Failed);
-						}
-					};				
-				TStateTreeEditorNode<FTestTask_PrintValue>& Tree1State1Task2 = State1.AddTask<FTestTask_PrintValue>(FName(TEXT("Tree1State1Task2")));
-				Tree1State1Task2.GetNode().CustomEnterStateFunc = [&WeakContext](FStateTreeExecutionContext& Context, const FTestTask_PrintValue* Task)
-					{
-						WeakContext.ContextTree1State1Task2 = Context.MakeWeakExecutionContext();
-						WeakContext.TaskTree1State1Task2 = Context.MakeWeakTaskRef(*Task);
-					};
-				Tree1State1Task2.GetNode().CustomTickFunc = [&WeakContext](FStateTreeExecutionContext& Context, const FTestTask_PrintValue* Task)
-					{
-						if (WeakContext.bWhen_TaskTree1State1Task2_Finish_TaskTree1State1Task2)
+							if (WeakContext.bGlobalFinishTaskSuccessOnTick)
+							{
+								WeakContext.ContextTree1GlobalTask.FinishTask(EStateTreeFinishTaskType::Succeeded);
+							}
+						};
+				}
+
+				// Root State
+				UStateTreeState& Root = EditorData1.AddSubTree("Tree1StateRoot");
+				{
+					TStateTreeEditorNode<FTestTask_PrintValue>& Tree1RootTask = Root.AddTask<FTestTask_PrintValue>(FName(TEXT("Tree1RootTask")));
+					Tree1RootTask.GetNode().CustomEnterStateFunc = [&WeakContext](FStateTreeExecutionContext& Context, const FTestTask_PrintValue* Task)
 						{
-							WeakContext.ContextTree1State1Task2.FinishTask(WeakContext.TaskTree1State1Task2, EStateTreeFinishTaskType::Failed);
-						}
-						if (WeakContext.bWhen_TaskTree1State1Task2_Finish_TaskTree1State1Task1)
+							WeakContext.ContextTree1RootTask = Context.MakeWeakExecutionContext();
+						};
+
+				}
+
+				// State 1
+				UStateTreeState& State1 = Root.AddChildState("Tree1State1");
+				{
+					TStateTreeEditorNode<FTestTask_PrintValue>& Tree1State1Task = State1.AddTask<FTestTask_PrintValue>(FName(TEXT("Tree1State1Task")));
+					Tree1State1Task.GetNode().CustomEnterStateFunc = [&WeakContext](FStateTreeExecutionContext& Context, const FTestTask_PrintValue* Task)
 						{
-							WeakContext.ContextTree1State1Task1.FinishTask(WeakContext.TaskTree1State1Task1, EStateTreeFinishTaskType::Failed);
-						}
-						if (WeakContext.bWhen_TaskTree1State1Task2_Finish_Tree1GlobalTask2)
+							WeakContext.ContextTree1State1Task = Context.MakeWeakExecutionContext();
+						};
+					Tree1State1Task.GetNode().CustomTickFunc = [&WeakContext](FStateTreeExecutionContext& Context, const FTestTask_PrintValue* Task)
 						{
-							WeakContext.ContextTree1GlobalTask2.FinishTask(WeakContext.TaskTree1GlobalTask2, EStateTreeFinishTaskType::Failed);
-						}
-					};
-			}
-			UStateTreeState& State2 = State1.AddChildState("Tree1State2");
-			{
-				TStateTreeEditorNode<FTestTask_PrintValue>& Tree1State2Task1 = State2.AddTask<FTestTask_PrintValue>(FName(TEXT("Tree1State2Task1")));
-				Tree1State2Task1.GetNode().CustomEnterStateFunc = [&WeakContext](FStateTreeExecutionContext& Context, const FTestTask_PrintValue* Task)
-					{
-						WeakContext.ContextTree1State2Task1 = Context.MakeWeakExecutionContext();
-						WeakContext.TaskTree1State2Task1 = Context.MakeWeakTaskRef(*Task);
-					};
-				TStateTreeEditorNode<FTestTask_PrintValue>& Tree1State2Task2 = State2.AddTask<FTestTask_PrintValue>(FName(TEXT("Tree1State2Task2")));
-				Tree1State2Task2.GetNode().CustomEnterStateFunc = [&WeakContext](FStateTreeExecutionContext& Context, const FTestTask_PrintValue* Task)
-					{
-						WeakContext.ContextTree1State2Task2 = Context.MakeWeakExecutionContext();
-						WeakContext.TaskTree1State2Task2 = Context.MakeWeakTaskRef(*Task);
-					};
-			}
+							if (WeakContext.bState1lFinishTaskFailOnTick)
+							{
+								WeakContext.ContextTree1State1Task.FinishTask(EStateTreeFinishTaskType::Failed);
+							}
+						};
+				}
+
+				// State 2
+				UStateTreeState& State2 = State1.AddChildState("Tree1State2");
+				{
+					TStateTreeEditorNode<FTestTask_PrintValue>& Tree1State2Task = State2.AddTask<FTestTask_PrintValue>(FName(TEXT("Tree1State2Task")));
+					Tree1State2Task.GetNode().CustomEnterStateFunc = [&WeakContext](FStateTreeExecutionContext& Context, const FTestTask_PrintValue* Task)
+						{
+							WeakContext.ContextTree1State2Task = Context.MakeWeakExecutionContext();
+						};
+				}
+			};
 		}
 
 		// Compile tree
@@ -165,16 +118,13 @@ struct FStateTreeTest_WeakContext_FinishTask : FStateTreeTestBase
 				AITEST_EQUAL("Start should complete with Running", Status, EStateTreeRunStatus::Running);
 				AITEST_TRUE(TEXT("In correct states"), Exec.ExpectInActiveStates("Tree1StateRoot", "Tree1State1", "Tree1State2"));
 
-				AITEST_TRUE(TEXT("Start should EnterState"), Exec.Expect("Tree1GlobalTask1", TEXT("EnterState0"))
-					.Then("Tree1GlobalTask2", TEXT("EnterState0"))
-					.Then("Tree1RootTask1", TEXT("EnterState0"))
-					.Then("Tree1RootTask2", TEXT("EnterState0"))
-					.Then("Tree1State1Task1", TEXT("EnterState0"))
-					.Then("Tree1State1Task2", TEXT("EnterState0"))
-					.Then("Tree1State2Task1", TEXT("EnterState0"))
-					.Then("Tree1State2Task2", TEXT("EnterState0")));
+				AITEST_TRUE(TEXT("Start should EnterState"), Exec.Expect("Tree1GlobalTask", TEXT("EnterState0"))
+							.Then("Tree1RootTask", TEXT("EnterState0"))
+							.Then("Tree1State1Task", TEXT("EnterState0"))
+							.Then("Tree1State2Task", TEXT("EnterState0")));
 				Exec.LogClear();
 			}
+
 			// Test that everything tick and there are no transitions.
 			{
 				FTestStateTreeExecutionContext Exec(StateTree1, StateTree1, InstanceData);
@@ -182,134 +132,84 @@ struct FStateTreeTest_WeakContext_FinishTask : FStateTreeTestBase
 				AITEST_EQUAL("Tick should complete with Running", Status, EStateTreeRunStatus::Running);
 				AITEST_TRUE(TEXT("In correct states"), Exec.ExpectInActiveStates("Tree1StateRoot", "Tree1State1", "Tree1State2"));
 
-				AITEST_TRUE(TEXT("Start should Tick"), Exec.Expect("Tree1GlobalTask1", TEXT("Tick0"))
-					.Then("Tree1GlobalTask2", TEXT("Tick0"))
-					.Then("Tree1RootTask1", TEXT("Tick0"))
-					.Then("Tree1RootTask2", TEXT("Tick0"))
-					.Then("Tree1State1Task1", TEXT("Tick0"))
-					.Then("Tree1State1Task2", TEXT("Tick0"))
-					.Then("Tree1State2Task1", TEXT("Tick0"))
-					.Then("Tree1State2Task2", TEXT("Tick0")));
+				AITEST_TRUE(TEXT("Tick should Tick"), Exec.Expect("Tree1GlobalTask", TEXT("Tick0"))
+							.Then("Tree1RootTask", TEXT("Tick0"))
+							.Then("Tree1State1Task", TEXT("Tick0"))
+							.Then("Tree1State2Task", TEXT("Tick0")));
 				Exec.LogClear();
 			}
-			// Test Weak::FinishTask outside of the tick.
-			WeakContext.ContextTree1RootTask2.FinishTask(WeakContext.TaskTree1RootTask2, EStateTreeFinishTaskType::Succeeded);
-			{
-				FTestStateTreeExecutionContext Exec(StateTree1, StateTree1, InstanceData);
-				EStateTreeRunStatus Status = Exec.Tick(0.1f);
-				AITEST_EQUAL("Tick should complete with Running", Status, EStateTreeRunStatus::Running);
-				AITEST_TRUE(TEXT("In correct states"), Exec.ExpectInActiveStates("Tree1StateRoot", "Tree1State1", "Tree1State2"));
 
-				AITEST_TRUE(TEXT("Tick should Tick"), Exec.Expect("Tree1GlobalTask1", TEXT("Tick0"))
-					.Then("Tree1GlobalTask2", TEXT("Tick0"))
-					.Then("Tree1State2Task2", TEXT("ExitState0"))
-					.Then("Tree1State2Task1", TEXT("ExitState0"))
-					.Then("Tree1State1Task2", TEXT("ExitState0"))
-					.Then("Tree1State1Task1", TEXT("ExitState0"))
-					.Then("Tree1RootTask2", TEXT("ExitState0"))
-					.Then("Tree1RootTask1", TEXT("ExitState0"))
-					.Then("Tree1RootTask1", TEXT("EnterState0"))
-					.Then("Tree1RootTask2", TEXT("EnterState0"))
-					.Then("Tree1State1Task1", TEXT("EnterState0"))
-					.Then("Tree1State1Task2", TEXT("EnterState0"))
-					.Then("Tree1State2Task1", TEXT("EnterState0"))
-					.Then("Tree1State2Task2", TEXT("EnterState0")));
-				AITEST_TRUE(TEXT("FinishTask should not Tick"), Exec.Expect("Tree1RootTask1", TEXT("Tick0")));
-				AITEST_FALSE(TEXT("FinishTask should not Tick"), Exec.Expect("Tree1RootTask2", TEXT("Tick0")));
-				AITEST_TRUE(TEXT("FinishTask should not Tick"), Exec.Expect("Tree1State1Task1", TEXT("Tick0")));
-				AITEST_TRUE(TEXT("FinishTask should not Tick"), Exec.Expect("Tree1State1Task2", TEXT("Tick0")));
-				AITEST_TRUE(TEXT("FinishTask should not Tick"), Exec.Expect("Tree1State2Task1", TEXT("Tick0")));
-				AITEST_TRUE(TEXT("FinishTask should not Tick"), Exec.Expect("Tree1State2Task2", TEXT("Tick0")));
-				Exec.LogClear();
-			}
-			// Test Weak::FinishTask of a task after the current task, during the tick tasks. Complete with Failure. The following task should not run.
+			// Test Finish GlobalTask inside the tick
 			{
-				WeakContext.bWhen_TaskTree1State1Task1_Finish_TaskTree1State1Task2 = true;
+				WeakContext.bGlobalFinishTaskSuccessOnTick = true;
 				FTestStateTreeExecutionContext Exec(StateTree1, StateTree1, InstanceData);
 				EStateTreeRunStatus Status = Exec.Tick(0.1f);
-				AITEST_EQUAL("Tick should complete with Running", Status, EStateTreeRunStatus::Running);
-				AITEST_TRUE(TEXT("In correct states"), Exec.ExpectInActiveStates("Tree1StateRoot", "Tree1State1", "Tree1State2"));
-
-				AITEST_TRUE(TEXT("Tick should Tick"), Exec.Expect("Tree1GlobalTask1", TEXT("Tick0"))
-					.Then("Tree1GlobalTask2", TEXT("Tick0"))
-					.Then("Tree1RootTask1", TEXT("Tick0"))
-					.Then("Tree1RootTask2", TEXT("Tick0"))
-					.Then("Tree1State1Task1", TEXT("Tick0")));
-				AITEST_FALSE(TEXT("FinishTask should not Tick"), Exec.Expect("Tree1State1Task2", TEXT("Tick0")));
-				AITEST_FALSE(TEXT("FinishTask should not Tick"), Exec.Expect("Tree1State2Task1", TEXT("Tick0")));
-				AITEST_FALSE(TEXT("FinishTask should not Tick"), Exec.Expect("Tree1State2Task2", TEXT("Tick0")));
-				AITEST_TRUE(TEXT("Tick should Tick"), Exec.Expect("Tree1State2Task2", TEXT("ExitState0"))
-					.Then("Tree1State2Task1", TEXT("ExitState0"))
-					.Then("Tree1State1Task2", TEXT("ExitState0"))
-					.Then("Tree1State1Task2", TEXT("EnterState0"))
-					.Then("Tree1State2Task1", TEXT("EnterState0"))
-					.Then("Tree1State2Task2", TEXT("EnterState0")));
-				Exec.LogClear();
-				WeakContext.bWhen_TaskTree1State1Task1_Finish_TaskTree1State1Task2 = false;
+				AITEST_EQUAL("Tick should complete with Succeeded", Status, EStateTreeRunStatus::Succeeded);
+				AITEST_TRUE(TEXT("Tick should Tick"), Exec.Expect("Tree1GlobalTask", TEXT("Tick0"))
+							.Then("Tree1State2Task", TEXT("ExitState0"))
+							.Then("Tree1State1Task", TEXT("ExitState0"))
+							.Then("Tree1RootTask", TEXT("ExitState0")));
+				WeakContext.bGlobalFinishTaskSuccessOnTick = false;
 			}
-			// Test Weak::FinishTask of a task before the current task, during the tick tasks.
+
+			// Finished global task stop the execution. Reset the execution.
 			{
-				WeakContext.bWhen_TaskTree1State1Task2_Finish_TaskTree1State1Task1 = true;
 				FTestStateTreeExecutionContext Exec(StateTree1, StateTree1, InstanceData);
-				EStateTreeRunStatus Status = Exec.Tick(0.1f);
-				AITEST_EQUAL("Tick should complete with Running", Status, EStateTreeRunStatus::Running);
-				AITEST_TRUE(TEXT("In correct states"), Exec.ExpectInActiveStates("Tree1StateRoot", "Tree1State1", "Tree1State2"));
-
-				AITEST_TRUE(TEXT("Tick should Tick"), Exec.Expect("Tree1GlobalTask1", TEXT("Tick0"))
-					.Then("Tree1GlobalTask2", TEXT("Tick0"))
-					.Then("Tree1RootTask1", TEXT("Tick0"))
-					.Then("Tree1RootTask2", TEXT("Tick0"))
-					.Then("Tree1State1Task1", TEXT("Tick0"))
-					.Then("Tree1State1Task2", TEXT("Tick0")));
-				AITEST_FALSE(TEXT("FinishTask should not Tick"), Exec.Expect("Tree1State2Task1", TEXT("Tick0")));
-				AITEST_FALSE(TEXT("FinishTask should not Tick"), Exec.Expect("Tree1State2Task2", TEXT("Tick0")));
-				AITEST_TRUE(TEXT("Tick should Tick"), Exec.Expect("Tree1State2Task2", TEXT("ExitState0"))
-					.Then("Tree1State2Task1", TEXT("ExitState0"))
-					.Then("Tree1State1Task2", TEXT("ExitState0"))
-					.Then("Tree1State1Task2", TEXT("EnterState0"))
-					.Then("Tree1State2Task1", TEXT("EnterState0"))
-					.Then("Tree1State2Task2", TEXT("EnterState0")));
+				Exec.Start();
 				Exec.LogClear();
-				WeakContext.bWhen_TaskTree1State1Task2_Finish_TaskTree1State1Task1 = false;
 			}
-			// Test Weak::FinishTask of a global task, during the global tick tasks.
+
+			// Test Finish GlobalTask outside the tick
 			{
-				WeakContext.bWhen_Tree1GlobalTask1_Finish_Tree1GlobalTask2 = true;
 				FTestStateTreeExecutionContext Exec(StateTree1, StateTree1, InstanceData);
+				WeakContext.ContextTree1GlobalTask.FinishTask(EStateTreeFinishTaskType::Failed);
 				EStateTreeRunStatus Status = Exec.Tick(0.1f);
-				AITEST_EQUAL("Tick should complete with Running", Status, EStateTreeRunStatus::Failed);
-				AITEST_TRUE(TEXT("Tick should Tick"), Exec.Expect("Tree1GlobalTask1", TEXT("Tick0"))
-					.Then("Tree1State2Task2", TEXT("ExitState0"))
-					.Then("Tree1State2Task1", TEXT("ExitState0"))
-					.Then("Tree1State1Task2", TEXT("ExitState0"))
-					.Then("Tree1GlobalTask2", TEXT("ExitState0"))
-					.Then("Tree1GlobalTask1", TEXT("ExitState0")));
-				AITEST_FALSE(TEXT("FinishTask should not Tick"), Exec.Expect("Tree1State2Task1", TEXT("Tick0")));
-				Exec.LogClear();
-				WeakContext.bWhen_Tree1GlobalTask1_Finish_Tree1GlobalTask2 = false;
+				AITEST_EQUAL("Tick should complete with Failed", Status, EStateTreeRunStatus::Failed);
+				AITEST_FALSE(TEXT("Tick should not Tick"), Exec.Expect("Tree1GlobalTask", TEXT("Tick0")));
+				AITEST_FALSE(TEXT("Tick should not Tick"), Exec.Expect("Tree1RootTask", TEXT("Tick0")));
+				AITEST_TRUE(TEXT("Tick should Tick"), Exec.Expect("Tree1State2Task", TEXT("ExitState0"))
+							.Then("Tree1State1Task", TEXT("ExitState0"))
+							.Then("Tree1RootTask", TEXT("ExitState0")));
 			}
+
 			// Finished global task stop the execution. Reset the execution.
 			{
 				FTestStateTreeExecutionContext Exec(StateTree1, StateTree1, InstanceData);
-				Exec.Stop();
 				Exec.Start();
 				Exec.LogClear();
 			}
-			// Test Weak::FinishTask of global task, during the tick tasks.
+
+			// Test Finish StateTask inside the tick
 			{
-				WeakContext.bWhen_TaskTree1State1Task2_Finish_Tree1GlobalTask2 = true;
+				WeakContext.bState1lFinishTaskFailOnTick = true;
 				FTestStateTreeExecutionContext Exec(StateTree1, StateTree1, InstanceData);
 				EStateTreeRunStatus Status = Exec.Tick(0.1f);
-				AITEST_EQUAL("Tick should complete with Failed", Status, EStateTreeRunStatus::Failed);
+				AITEST_EQUAL("Tick should complete with Running", Status, EStateTreeRunStatus::Running);
+				AITEST_TRUE(TEXT("Tick should Tick"), Exec.Expect("Tree1GlobalTask", TEXT("Tick0"))
+							.Then("Tree1RootTask", TEXT("Tick0"))
+							.Then("Tree1State1Task", TEXT("Tick0"))
+							.Then("Tree1State2Task", TEXT("ExitState0"))
+							.Then("Tree1State1Task", TEXT("ExitState0"))
+							.Then("Tree1RootTask", TEXT("ExitState0")));
+
+				WeakContext.bGlobalFinishTaskSuccessOnTick = false;
+			}
 
-				AITEST_TRUE(TEXT("Tick should Tick"), Exec.Expect("Tree1GlobalTask1", TEXT("Tick0"))
-					.Then("Tree1State1Task2", TEXT("Tick0")));
-				AITEST_TRUE(TEXT("Tick should exit"), Exec.Expect("Tree1State2Task2", TEXT("ExitState0"))
-					.Then("Tree1GlobalTask1", TEXT("ExitState0")));
-				Exec.LogClear();
-				WeakContext.bWhen_TaskTree1State1Task2_Finish_Tree1GlobalTask2 = false;
+			// Test Finish StateTask outside the tick
+			{
+				FTestStateTreeExecutionContext Exec(StateTree1, StateTree1, InstanceData);
+				WeakContext.ContextTree1State1Task.FinishTask(EStateTreeFinishTaskType::Succeeded);
+				EStateTreeRunStatus Status = Exec.Tick(0.1f);
+				AITEST_EQUAL("Tick should complete with Running", Status, EStateTreeRunStatus::Running);
+				AITEST_TRUE(TEXT("Tick should Tick"), Exec.Expect("Tree1GlobalTask", TEXT("Tick0"))
+							.Then("Tree1RootTask", TEXT("Tick0"))
+							.Then("Tree1State2Task", TEXT("Tick0"))
+							.Then("Tree1State2Task", TEXT("ExitState0"))
+							.Then("Tree1State1Task", TEXT("ExitState0"))
+							.Then("Tree1RootTask", TEXT("ExitState0")));
 			}
 
+			// Stop the Exec
 			{
 				FTestStateTreeExecutionContext Exec(StateTree1, StateTree1, InstanceData);
 				Exec.Stop();
@@ -320,7 +220,6 @@ struct FStateTreeTest_WeakContext_FinishTask : FStateTreeTestBase
 	}
 };
 IMPLEMENT_AI_INSTANT_TEST(FStateTreeTest_WeakContext_FinishTask, "System.StateTree.WeakContext.FinishTask");
-
 } // namespace UE::StateTree::Tests
 
 UE_ENABLE_OPTIMIZATION_SHIP
