Remove CVar `StateTree.RuntimeValidation.MultithreadAccessDetector`.

This will make it possible (in followup CLs) to adopt our existing MT access-detector scope macros.

diff --git a/Engine/Plugins/Runtime/StateTree/Source/StateTreeModule/Private/StateTreeInstanceData.cpp b/Engine/Plugins/Runtime/StateTree/Source/StateTreeModule/Private/StateTreeInstanceData.cpp
--- a/Engine/Plugins/Runtime/StateTree/Source/StateTreeModule/Private/StateTreeInstanceData.cpp
+++ b/Engine/Plugins/Runtime/StateTree/Source/StateTreeModule/Private/StateTreeInstanceData.cpp
@@ -17,16 +17,6 @@
 const FGuid FStateTreeInstanceStorageCustomVersion::GUID(0x60C4F0DE, 0x8B264C34, 0xAA937201, 0x5DFF09CC);
 FCustomVersionRegistration GRegisterStateTreeInstanceStorageCustomVersion(FStateTreeInstanceStorageCustomVersion::GUID, FStateTreeInstanceStorageCustomVersion::LatestVersion, TEXT("StateTreeInstanceStorage"));
 
-namespace UE::StateTree::Debug::Private
-{
-	bool RuntimeValidationThreadingAccessDetector = true;
-	static FAutoConsoleVariableRef CVarRuntimeValidationThreadingAccessDetector(
-		TEXT("StateTree.RuntimeValidation.MultithreadAccessDetector"),
-		RuntimeValidationThreadingAccessDetector,
-		TEXT("Test if an execution context is accessed from different threads at the same time.")
-	);
-} //UE::StateTree::Debug::Private
-
 namespace UE::StateTree
 {
 	/**
@@ -371,34 +361,22 @@ void FStateTreeInstanceStorage::Reset()
 
 void FStateTreeInstanceStorage::AcquireReadAccess()
 {
-	if (UE::StateTree::Debug::Private::RuntimeValidationThreadingAccessDetector)
-	{
-		UE_MT_ACQUIRE_READ_ACCESS(AccessDetector);
-	}
+	UE_MT_ACQUIRE_READ_ACCESS(AccessDetector);
 }
 
 void FStateTreeInstanceStorage::ReleaseReadAccess()
 {
-	if (UE::StateTree::Debug::Private::RuntimeValidationThreadingAccessDetector)
-	{
-		UE_MT_RELEASE_READ_ACCESS(AccessDetector);
-	}
+	UE_MT_RELEASE_READ_ACCESS(AccessDetector);
 }
 
 void FStateTreeInstanceStorage::AcquireWriteAccess()
 {
-	if (UE::StateTree::Debug::Private::RuntimeValidationThreadingAccessDetector)
-	{
-		UE_MT_ACQUIRE_WRITE_ACCESS(AccessDetector);
-	}
+	UE_MT_ACQUIRE_WRITE_ACCESS(AccessDetector);
 }
 
 void FStateTreeInstanceStorage::ReleaseWriteAccess()
 {
-	if (UE::StateTree::Debug::Private::RuntimeValidationThreadingAccessDetector)
-	{
-		UE_MT_RELEASE_WRITE_ACCESS(AccessDetector);
-	}
+	UE_MT_RELEASE_WRITE_ACCESS(AccessDetector);
 }
 
 UE::StateTree::Debug::FRuntimeValidation FStateTreeInstanceStorage::GetRuntimeValidation() const
