StateTree: Adding filtering of state selection behavior options based on schema.

diff --git a/Engine/Plugins/Runtime/StateTree/Source/StateTreeEditorModule/Private/Customizations/StateTreeStateDetails.cpp b/Engine/Plugins/Runtime/StateTree/Source/StateTreeEditorModule/Private/Customizations/StateTreeStateDetails.cpp
--- a/Engine/Plugins/Runtime/StateTree/Source/StateTreeEditorModule/Private/Customizations/StateTreeStateDetails.cpp
+++ b/Engine/Plugins/Runtime/StateTree/Source/StateTreeEditorModule/Private/Customizations/StateTreeStateDetails.cpp
@@ -5,6 +5,7 @@
 #include "PropertyCustomizationHelpers.h"
 #include "IPropertyUtilities.h"
 #include "PropertyBagDetails.h"
+#include "PropertyRestriction.h"
 #include "StateTree.h"
 #include "StateTreeEditor.h"
 #include "StateTreeEditorData.h"
@@ -181,6 +182,19 @@ void FStateTreeStateDetails::CustomizeDetails(IDetailLayoutBuilder& DetailBuilde
 
 	if (StateType == EStateTreeStateType::State || StateType == EStateTreeStateType::Subtree)
 	{
+		// Restrict IsStateSelectionAllowed
+		TSharedRef<FPropertyRestriction> SchemaSupported = MakeShared<FPropertyRestriction>(LOCTEXT("StateSelectionAllowedRestriction", "The schema restricts the selection behavior."));
+		const UEnum* EnumPtr = StaticEnum<EStateTreeStateSelectionBehavior>();
+		for (int32 Index = 0; Index < EnumPtr->NumEnums(); ++Index)
+		{
+			const EStateTreeStateSelectionBehavior EnumValue = static_cast<EStateTreeStateSelectionBehavior>(EnumPtr->GetValueByIndex(Index));
+			if (!Schema->IsStateSelectionAllowed(EnumValue))
+			if (Index == 1 || Index == 2 || Index == 3)
+			{
+				SchemaSupported->AddHiddenValue(EnumPtr->GetNameStringByIndex(Index));
+			}
+		}
+		SelectionBehaviorProperty->AddRestriction(SchemaSupported);
 		StateCategory.AddProperty(SelectionBehaviorProperty);
 	}
 	else if (StateType == EStateTreeStateType::Linked)
diff --git a/Engine/Plugins/Runtime/StateTree/Source/StateTreeEditorModule/Private/StateTreeCompiler.cpp b/Engine/Plugins/Runtime/StateTree/Source/StateTreeEditorModule/Private/StateTreeCompiler.cpp
--- a/Engine/Plugins/Runtime/StateTree/Source/StateTreeEditorModule/Private/StateTreeCompiler.cpp
+++ b/Engine/Plugins/Runtime/StateTree/Source/StateTreeEditorModule/Private/StateTreeCompiler.cpp
@@ -563,6 +563,7 @@ bool FStateTreeCompiler::CreateStates()
 bool FStateTreeCompiler::CreateStateRecursive(UStateTreeState& State, const FStateTreeStateHandle Parent)
 {
 	check(StateTree);
+	check(Schema);
 
 	FStateTreeCompilerLogStateScope LogStateScope(&State, Log);
 
@@ -593,6 +594,14 @@ bool FStateTreeCompiler::CreateStateRecursive(UStateTreeState& State, const FSta
 	CompactState.Type = State.Type;
 	CompactState.SelectionBehavior = State.SelectionBehavior;
 
+	if (!Schema->IsStateSelectionAllowed(CompactState.SelectionBehavior))
+	{
+		Log.Reportf(EMessageSeverity::Warning,
+			TEXT("The State '%s' has a restricted selection behavior for the schema."),
+			*CompactState.Name.ToString());
+		return false;
+	}
+
 	SourceStates.Add(&State);
 	IDToState.Add(State.ID, StateIdx);
 
diff --git a/Engine/Plugins/Runtime/StateTree/Source/StateTreeModule/Public/StateTreeSchema.h b/Engine/Plugins/Runtime/StateTree/Source/StateTreeModule/Public/StateTreeSchema.h
--- a/Engine/Plugins/Runtime/StateTree/Source/StateTreeModule/Public/StateTreeSchema.h
+++ b/Engine/Plugins/Runtime/StateTree/Source/StateTreeModule/Public/StateTreeSchema.h
@@ -11,7 +11,7 @@ struct FStateTreeExternalDataDesc;
 struct FStateTreeLinker;
 enum class EStateTreeDataSourceType : uint8;
 enum class EStateTreeParameterDataType : uint8;
-
+enum class EStateTreeStateSelectionBehavior : uint8;
 
 #if WITH_EDITOR
 struct FStateTreeDetailsCustomizationSettings
@@ -60,6 +60,12 @@ public:
 		return false;
 	}
 
+	/** @return True if the state selection behavior is supported. */
+	virtual bool IsStateSelectionAllowed(EStateTreeStateSelectionBehavior InBehavior) const
+	{
+		return true;
+	}
+
 	/**
 	 * Helper function to check if a class is any of the Blueprint extendable item classes (Eval, Task, Condition).
 	 * Can be used to quickly accept all of those classes in IsClassAllowed().
