[AutoRTFM] Wrap an MTAccessDetector with open blocks to avoid a crash.

This solution isn't ideal; we have a transactionally-safe access detector which does a better job of detecting access violations when transactions are involved. However, switching over to this class is a more complex fix which didn't make the cut for 35.00. This fix is a simpler option for the short term to avoid crashing when AutoRTFM is on. It won't introduce any false positives, and the behavior will be completely unchanged when AutoRTFM is not enabled.

diff --git a/Engine/Plugins/Runtime/StateTree/Source/StateTreeModule/Private/StateTreeInstanceData.cpp b/Engine/Plugins/Runtime/StateTree/Source/StateTreeModule/Private/StateTreeInstanceData.cpp
--- a/Engine/Plugins/Runtime/StateTree/Source/StateTreeModule/Private/StateTreeInstanceData.cpp
+++ b/Engine/Plugins/Runtime/StateTree/Source/StateTreeModule/Private/StateTreeInstanceData.cpp
@@ -11,6 +11,7 @@
 #include "Serialization/CustomVersion.h"
 #include "Serialization/PropertyLocalizationDataGathering.h"
 #include "StateTreeDelegate.h"
+#include "AutoRTFM.h"
 
 #include UE_INLINE_GENERATED_CPP_BY_NAME(StateTreeInstanceData)
 
@@ -359,24 +360,53 @@ void FStateTreeInstanceStorage::Reset()
 #endif
 }
 
+// #jira SOL-8070: Ideally, we should use the transactionally-safe access detector instead of relying on OPEN and ONABORT blocks here.
 void FStateTreeInstanceStorage::AcquireReadAccess()
 {
-	UE_MT_ACQUIRE_READ_ACCESS(AccessDetector);
+	UE_AUTORTFM_OPEN
+	{
+		UE_MT_ACQUIRE_READ_ACCESS(AccessDetector);
+	};
+	UE_AUTORTFM_ONABORT(this)
+	{
+		UE_MT_RELEASE_READ_ACCESS(AccessDetector);
+	};
 }
 
 void FStateTreeInstanceStorage::ReleaseReadAccess()
 {
-	UE_MT_RELEASE_READ_ACCESS(AccessDetector);
+	UE_AUTORTFM_OPEN
+	{
+		UE_MT_RELEASE_READ_ACCESS(AccessDetector);
+	};
+	UE_AUTORTFM_ONABORT(this)
+	{
+		UE_MT_ACQUIRE_READ_ACCESS(AccessDetector);
+	};
 }
 
 void FStateTreeInstanceStorage::AcquireWriteAccess()
 {
-	UE_MT_ACQUIRE_WRITE_ACCESS(AccessDetector);
+	UE_AUTORTFM_OPEN
+	{
+		UE_MT_ACQUIRE_WRITE_ACCESS(AccessDetector);
+	};
+	UE_AUTORTFM_ONABORT(this)
+	{
+		UE_MT_RELEASE_WRITE_ACCESS(AccessDetector);
+	};
 }
 
 void FStateTreeInstanceStorage::ReleaseWriteAccess()
 {
-	UE_MT_RELEASE_WRITE_ACCESS(AccessDetector);
+	UE_AUTORTFM_OPEN
+	{
+		UE_MT_RELEASE_WRITE_ACCESS(AccessDetector);
+	};
+	UE_AUTORTFM_ONABORT(this)
+	{
+		UE_MT_ACQUIRE_WRITE_ACCESS(AccessDetector);
+	};
 }
 
 UE::StateTree::Debug::FRuntimeValidation FStateTreeInstanceStorage::GetRuntimeValidation() const
