[StateTreeEditor] added refresh on undo/redo for the Editor Data details view to make sure cached data in binding row extensions are up-to-date

diff --git a/Engine/Plugins/Runtime/StateTree/Source/StateTreeEditorModule/Private/Customizations/StateTreeEditorDataDetails.cpp b/Engine/Plugins/Runtime/StateTree/Source/StateTreeEditorModule/Private/Customizations/StateTreeEditorDataDetails.cpp
--- a/Engine/Plugins/Runtime/StateTree/Source/StateTreeEditorModule/Private/Customizations/StateTreeEditorDataDetails.cpp
+++ b/Engine/Plugins/Runtime/StateTree/Source/StateTreeEditorModule/Private/Customizations/StateTreeEditorDataDetails.cpp
@@ -2,21 +2,20 @@
 
 #include "StateTreeEditorDataDetails.h"
 #include "Brushes/SlateRoundedBoxBrush.h"
-#include "IDetailChildrenBuilder.h"
 #include "DetailLayoutBuilder.h"
-#include "PropertyCustomizationHelpers.h"
-#include "StateTreeEditorData.h"
-#include "Kismet2/BlueprintEditorUtils.h"
+#include "IDetailChildrenBuilder.h"
 #include "IPropertyUtilities.h"
+#include "Kismet2/BlueprintEditorUtils.h"
 #include "PropertyBagDetails.h"
+#include "PropertyCustomizationHelpers.h"
+#include "StateTreeEditorData.h"
+#include "StateTreeEditorNodeUtils.h"
 #include "StateTreeEditorStyle.h"
 #include "StateTreeSchema.h"
 #include "Styling/StyleColors.h"
 #include "Widgets/Images/SImage.h"
 #include "Widgets/Layout/SBorder.h"
-#include "Widgets/Layout/SBox.h"
 #include "Widgets/Text/STextBlock.h"
-#include "StateTreeEditorNodeUtils.h"
 
 #define LOCTEXT_NAMESPACE "StateTreeEditor"
 
@@ -27,6 +26,9 @@ TSharedRef<IDetailCustomization> FStateTreeEditorDataDetails::MakeInstance()
 
 void FStateTreeEditorDataDetails::CustomizeDetails(IDetailLayoutBuilder& DetailBuilder)
 {
+	const TSharedRef<IPropertyUtilities> PropUtils = DetailBuilder.GetPropertyUtilities();
+	WeakPropertyUtilities = PropUtils.ToWeakPtr();
+
 	// Find StateTreeEditorData associated with this panel.
 	const UStateTreeEditorData* EditorData = nullptr;
 	TArray<TWeakObjectPtr<UObject>> Objects;
@@ -147,7 +149,6 @@ void FStateTreeEditorDataDetails::CustomizeDetails(IDetailLayoutBuilder& DetailB
 		ParametersCategory.SetSortOrder(2);
 		{
 			// Show parameters as a category.
-			TSharedRef<IPropertyUtilities> PropUtils = DetailBuilder.GetPropertyUtilities();
 			TSharedPtr<IPropertyHandle> PropertyBagParametersProperty = DetailBuilder.GetProperty(GET_MEMBER_NAME_CHECKED(UStateTreeEditorData, RootParameterPropertyBag)); // FInstancedPropertyBag
 			PropertyBagParametersProperty->MarkHiddenByCustomization();
 			check(PropertyBagParametersProperty);
@@ -187,7 +188,6 @@ void FStateTreeEditorDataDetails::CustomizeDetails(IDetailLayoutBuilder& DetailB
 	}
 	else
 	{
-		TSharedRef<IPropertyUtilities> PropUtils = DetailBuilder.GetPropertyUtilities();
 		TSharedPtr<IPropertyHandle> PropertyBagParametersProperty = DetailBuilder.GetProperty(GET_MEMBER_NAME_CHECKED(UStateTreeEditorData, RootParameterPropertyBag)); // FInstancedPropertyBag
 		PropertyBagParametersProperty->MarkHiddenByCustomization();
 	}	
@@ -239,22 +239,20 @@ void FStateTreeEditorDataDetails::CustomizeDetails(IDetailLayoutBuilder& DetailB
 	// Refresh the UI when the Schema changes.	
 	TSharedPtr<IPropertyHandle> SchemaProperty = DetailBuilder.GetProperty(GET_MEMBER_NAME_CHECKED(UStateTreeEditorData, Schema));
 	check(SchemaProperty.IsValid());
-	TSharedPtr<IPropertyUtilities> PropUtils = DetailBuilder.GetPropertyUtilities();
-	SchemaProperty->SetOnPropertyValueChanged(FSimpleDelegate::CreateLambda([PropUtils] ()
-	{
-		if (PropUtils.IsValid())
+	SchemaProperty->SetOnPropertyValueChanged(FSimpleDelegate::CreateLambda([WeakPropertyUtilities = WeakPropertyUtilities] ()
 		{
-			PropUtils->ForceRefresh();
-		}
-	}));
-	SchemaProperty->SetOnChildPropertyValueChanged(FSimpleDelegate::CreateLambda([PropUtils] ()
-	{
-		if (PropUtils.IsValid())
+			if (const TSharedPtr<IPropertyUtilities> PropertyUtilities = WeakPropertyUtilities.Pin())
+			{
+				PropertyUtilities->ForceRefresh();
+			}
+		}));
+	SchemaProperty->SetOnChildPropertyValueChanged(FSimpleDelegate::CreateLambda([WeakPropertyUtilities = WeakPropertyUtilities] ()
 		{
-			PropUtils->ForceRefresh();
-		}
-	}));
-
+			if (const TSharedPtr<IPropertyUtilities> PropertyUtilities = WeakPropertyUtilities.Pin())
+			{
+				PropertyUtilities->ForceRefresh();
+			}
+		}));
 }
 
 void FStateTreeEditorDataDetails::MakeArrayCategory(IDetailLayoutBuilder& DetailBuilder, FName CategoryName, const FText& DisplayName, int32 SortOrder, TSharedPtr<IPropertyHandle> PropertyHandle)
@@ -282,4 +280,22 @@ void FStateTreeEditorDataDetails::MakeArrayCategory(IDetailLayoutBuilder& Detail
 	Category.AddCustomBuilder(Builder, /*bForAdvanced*/ false);
 }
 
+void FStateTreeEditorDataDetails::PostUndo(bool bSuccess)
+{
+	// Refresh view on undo or redo so that the customization based on e.g. Global tasks bindings, parameters, etc. will be reflected correctly.
+	if (const TSharedPtr<IPropertyUtilities> PropertyUtilities = WeakPropertyUtilities.Pin())
+	{
+		PropertyUtilities->ForceRefresh();
+	}
+}
+
+void FStateTreeEditorDataDetails::PostRedo(bool bSuccess)
+{
+	// Refresh view on undo or redo so that the customization based on e.g. Global tasks bindings, parameters, etc. will be reflected correctly.
+	if (const TSharedPtr<IPropertyUtilities> PropertyUtilities = WeakPropertyUtilities.Pin())
+	{
+		PropertyUtilities->ForceRefresh();
+	}
+}
+
 #undef LOCTEXT_NAMESPACE
diff --git a/Engine/Plugins/Runtime/StateTree/Source/StateTreeEditorModule/Private/Customizations/StateTreeEditorDataDetails.h b/Engine/Plugins/Runtime/StateTree/Source/StateTreeEditorModule/Private/Customizations/StateTreeEditorDataDetails.h
--- a/Engine/Plugins/Runtime/StateTree/Source/StateTreeEditorModule/Private/Customizations/StateTreeEditorDataDetails.h
+++ b/Engine/Plugins/Runtime/StateTree/Source/StateTreeEditorModule/Private/Customizations/StateTreeEditorDataDetails.h
@@ -1,14 +1,16 @@
-﻿// Copyright Epic Games, Inc. All Rights Reserved.
+// Copyright Epic Games, Inc. All Rights Reserved.
 
 #pragma once
 
+#include "EditorUndoClient.h"
 #include "IDetailCustomization.h"
 
-class IDetailLayoutBuilder;
 class IDetailChildrenBuilder;
+class IDetailLayoutBuilder;
 class IPropertyHandle;
+class IPropertyUtilities;
 
-class FStateTreeEditorDataDetails : public IDetailCustomization
+class FStateTreeEditorDataDetails : public IDetailCustomization, FSelfRegisteringEditorUndoClient
 {
 public:
 	/** Makes a new instance of this detail layout class for a specific detail view requesting it */
@@ -19,4 +21,10 @@ public:
 
 protected:
 	void MakeArrayCategory(IDetailLayoutBuilder& DetailBuilder, FName CategoryName, const FText& DisplayName, int32 SortOrder, TSharedPtr<IPropertyHandle> PropertyHandle);
+
+	//~ FEditorUndoClient
+	virtual void PostUndo(bool bSuccess) override;
+	virtual void PostRedo(bool bSuccess) override;
+
+	TWeakPtr<IPropertyUtilities> WeakPropertyUtilities;
 };
diff --git a/Engine/Plugins/Runtime/StateTree/Source/StateTreeEditorModule/Private/Customizations/StateTreeStateDetails.cpp b/Engine/Plugins/Runtime/StateTree/Source/StateTreeEditorModule/Private/Customizations/StateTreeStateDetails.cpp
--- a/Engine/Plugins/Runtime/StateTree/Source/StateTreeEditorModule/Private/Customizations/StateTreeStateDetails.cpp
+++ b/Engine/Plugins/Runtime/StateTree/Source/StateTreeEditorModule/Private/Customizations/StateTreeStateDetails.cpp
@@ -15,8 +15,6 @@
 #include "StateTreeSchema.h"
 #include "StateTreeStateParametersDetails.h"
 #include "Debugger/StateTreeDebuggerUIExtensions.h"
-#include "Widgets/Layout/SBorder.h"
-#include "Widgets/Layout/SBox.h"
 #include "Widgets/SBoxPanel.h"
 #include "Widgets/Images/SImage.h"
 #include "Widgets/Text/STextBlock.h"
@@ -32,6 +30,9 @@ TSharedRef<IDetailCustomization> FStateTreeStateDetails::MakeInstance()
 
 void FStateTreeStateDetails::CustomizeDetails(IDetailLayoutBuilder& DetailBuilder)
 {
+	const TSharedRef<IPropertyUtilities> PropUtils = DetailBuilder.GetPropertyUtilities();
+	WeakPropertyUtilities = PropUtils.ToWeakPtr();
+
 	TArray<TWeakObjectPtr<UObject>> Objects;
 	DetailBuilder.GetObjectsBeingCustomized(Objects);
 	
@@ -67,7 +68,6 @@ void FStateTreeStateDetails::CustomizeDetails(IDetailLayoutBuilder& DetailBuilde
 	const FString SchemaPath = Schema ? Schema->GetClass()->GetPathName() : FString();
 	TWeakObjectPtr<UStateTreeEditorData> WeakEditorData = EditorData;
 
-	
 	const TSharedPtr<IPropertyHandle> IDProperty = DetailBuilder.GetProperty(GET_MEMBER_NAME_CHECKED(UStateTreeState, ID));
 	const TSharedPtr<IPropertyHandle> NameProperty = DetailBuilder.GetProperty(GET_MEMBER_NAME_CHECKED(UStateTreeState, Name));
 	const TSharedPtr<IPropertyHandle> TagProperty = DetailBuilder.GetProperty(GET_MEMBER_NAME_CHECKED(UStateTreeState, Tag));
@@ -88,8 +88,6 @@ void FStateTreeStateDetails::CustomizeDetails(IDetailLayoutBuilder& DetailBuilde
 	const TSharedPtr<IPropertyHandle> RequiredEventToEnterProperty = DetailBuilder.GetProperty(GET_MEMBER_NAME_CHECKED(UStateTreeState, RequiredEventToEnter));
 	const TSharedPtr<IPropertyHandle> CheckPrerequisitesProperty = DetailBuilder.GetProperty(GET_MEMBER_NAME_CHECKED(UStateTreeState, bCheckPrerequisitesWhenActivatingChildDirectly));
 	const TSharedPtr<IPropertyHandle> WeightProperty = DetailBuilder.GetProperty(GET_MEMBER_NAME_CHECKED(UStateTreeState, Weight));
-
-	PropUtils = DetailBuilder.GetPropertyUtilities();
 	
 	// Never show enabled
 	EnabledProperty->MarkHiddenByCustomization();
@@ -108,7 +106,7 @@ void FStateTreeStateDetails::CustomizeDetails(IDetailLayoutBuilder& DetailBuilde
 	StateCategory.SetSortOrder(0);
 	{
 		TSharedRef<SHorizontalBox> HeaderContent = SNew(SHorizontalBox)
-			.IsEnabled(PropUtils.ToSharedRef(), &IPropertyUtilities::IsPropertyEditingEnabled)
+			.IsEnabled(PropUtils, &IPropertyUtilities::IsPropertyEditingEnabled)
 			+ SHorizontalBox::Slot()
 			.FillWidth(1.f)
 			.HAlign(HAlign_Right)
@@ -243,7 +241,7 @@ void FStateTreeStateDetails::CustomizeDetails(IDetailLayoutBuilder& DetailBuilde
 
 		
 		const TSharedRef<SHorizontalBox> HeaderContentWidget = SNew(SHorizontalBox)
-			.IsEnabled(PropUtils.ToSharedRef(), &IPropertyUtilities::IsPropertyEditingEnabled);
+			.IsEnabled(PropUtils, &IPropertyUtilities::IsPropertyEditingEnabled);
 
 		HeaderContentWidget->AddSlot()
 			.VAlign(VAlign_Center)
@@ -273,7 +271,7 @@ void FStateTreeStateDetails::CustomizeDetails(IDetailLayoutBuilder& DetailBuilde
 				[
 					FPropertyBagDetails::MakeAddPropertyWidget(
 						ParametersParametersProperty,
-						PropUtils.ToSharedRef(),
+						PropUtils,
 						EPropertyBagPropertyType::Bool,
 						FLinearColor(UE::StateTree::Colors::Blue)).ToSharedRef()
 				];
@@ -393,21 +391,21 @@ void FStateTreeStateDetails::CustomizeDetails(IDetailLayoutBuilder& DetailBuilde
 		/*SortOrder*/5);
 
 	// Refresh the UI when the type changes.	
-	TypeProperty->SetOnPropertyValueChanged(FSimpleDelegate::CreateLambda([PropUtils = PropUtils] ()
-	{
-		if (PropUtils.IsValid())
+	TypeProperty->SetOnPropertyValueChanged(FSimpleDelegate::CreateLambda([WeakPropertyUtilities = WeakPropertyUtilities] ()
 		{
-			PropUtils->ForceRefresh();
-		}
-	}));
+			if (const TSharedPtr<IPropertyUtilities> PropertyUtilities = WeakPropertyUtilities.Pin())
+			{
+				PropertyUtilities->ForceRefresh();
+			}
+		}));
 }
 
 void FStateTreeStateDetails::PostUndo(bool bSuccess)
 {
 	// Refresh view on undo or redo so that the customization based on e.g. State type will be reflected correctly.
-	if (PropUtils.IsValid())
+	if (const TSharedPtr<IPropertyUtilities> PropertyUtilities = WeakPropertyUtilities.Pin())
 	{
-		PropUtils->ForceRefresh();
+		PropertyUtilities->ForceRefresh();
 	}
 	
 }
@@ -415,9 +413,9 @@ void FStateTreeStateDetails::PostUndo(bool bSuccess)
 void FStateTreeStateDetails::PostRedo(bool bSuccess)
 {
 	// Refresh view on undo or redo so that the customization based on e.g. State type will be reflected correctly.
-	if (PropUtils.IsValid())
+	if (const TSharedPtr<IPropertyUtilities> PropertyUtilities = WeakPropertyUtilities.Pin())
 	{
-		PropUtils->ForceRefresh(); 
+		PropertyUtilities->ForceRefresh();
 	}
 }
 
diff --git a/Engine/Plugins/Runtime/StateTree/Source/StateTreeEditorModule/Private/Customizations/StateTreeStateDetails.h b/Engine/Plugins/Runtime/StateTree/Source/StateTreeEditorModule/Private/Customizations/StateTreeStateDetails.h
--- a/Engine/Plugins/Runtime/StateTree/Source/StateTreeEditorModule/Private/Customizations/StateTreeStateDetails.h
+++ b/Engine/Plugins/Runtime/StateTree/Source/StateTreeEditorModule/Private/Customizations/StateTreeStateDetails.h
@@ -28,5 +28,5 @@ protected:
 	virtual void PostUndo(bool bSuccess) override;
 	virtual void PostRedo(bool bSuccess) override;
 
-	TSharedPtr<IPropertyUtilities> PropUtils;
+	TWeakPtr<IPropertyUtilities> WeakPropertyUtilities;
 };
