[StateTreeDebugger] registers to be notified when the associated StateTree asset get successfully recompiled to clear any previous recorded data that could mismatch new compiled tree data version.

diff --git a/Engine/Plugins/Runtime/StateTree/Source/StateTreeEditorModule/Private/Debugger/SStateTreeDebuggerView.cpp b/Engine/Plugins/Runtime/StateTree/Source/StateTreeEditorModule/Private/Debugger/SStateTreeDebuggerView.cpp
--- a/Engine/Plugins/Runtime/StateTree/Source/StateTreeEditorModule/Private/Debugger/SStateTreeDebuggerView.cpp
+++ b/Engine/Plugins/Runtime/StateTree/Source/StateTreeEditorModule/Private/Debugger/SStateTreeDebuggerView.cpp
@@ -356,6 +356,16 @@ void SStateTreeDebuggerView::Construct(const FArguments& InArgs, const UStateTre
 			bRecording = TraceStatus == EStateTreeTraceStatus::TracesStarted;
 		});
 
+	// Registers a delegate to be notified when the associated StateTree asset get successfully recompiled
+	// to clear any previous recorded data that could mismatch new compiled tree data version.
+	UE::StateTree::Delegates::OnPostCompile.AddSPLambda(this, [this](const UStateTree& RecompiledStateTree)
+		{
+			if (StateTree == &RecompiledStateTree)
+			{
+				ResetTracks();
+			}
+		});
+
 	// Bind callbacks to the debugger delegates
 	Debugger->OnNewSession.BindSP(this, &SStateTreeDebuggerView::OnNewSession);
 	Debugger->OnNewInstance.BindSP(this, &SStateTreeDebuggerView::OnNewInstance);
