StateTree: Fix rebuild state tree view with undo/redo.

diff --git a/Engine/Plugins/Runtime/StateTree/Source/StateTreeEditorModule/Private/SStateTreeView.cpp b/Engine/Plugins/Runtime/StateTree/Source/StateTreeEditorModule/Private/SStateTreeView.cpp
--- a/Engine/Plugins/Runtime/StateTree/Source/StateTreeEditorModule/Private/SStateTreeView.cpp
+++ b/Engine/Plugins/Runtime/StateTree/Source/StateTreeEditorModule/Private/SStateTreeView.cpp
@@ -420,7 +420,7 @@ void SStateTreeView::HandleUserSettingsChanged()
 void SStateTreeView::HandleModelAssetChanged()
 {
 	// this only refresh the list. i.e. each row widget will not be refreshed
-	// bItemsDirty = true;
+	bItemsDirty = true;
 
 	// we need to rebuild the list to update each row widget
 	TreeView->RebuildList();
