StateTree: Remove temporary test
[RN]

diff --git a/Engine/Plugins/Runtime/StateTree/Source/StateTreeEditorModule/Private/Customizations/StateTreeStateDetails.cpp b/Engine/Plugins/Runtime/StateTree/Source/StateTreeEditorModule/Private/Customizations/StateTreeStateDetails.cpp
--- a/Engine/Plugins/Runtime/StateTree/Source/StateTreeEditorModule/Private/Customizations/StateTreeStateDetails.cpp
+++ b/Engine/Plugins/Runtime/StateTree/Source/StateTreeEditorModule/Private/Customizations/StateTreeStateDetails.cpp
@@ -187,7 +187,6 @@ void FStateTreeStateDetails::CustomizeDetails(IDetailLayoutBuilder& DetailBuilde
 		{
 			const EStateTreeStateSelectionBehavior EnumValue = static_cast<EStateTreeStateSelectionBehavior>(EnumPtr->GetValueByIndex(Index));
 			if (!Schema->IsStateSelectionAllowed(EnumValue))
-			if (Index == 1 || Index == 2 || Index == 3)
 			{
 				SchemaSupported->AddHiddenValue(EnumPtr->GetNameStringByIndex(Index));
 			}
