StateTree: Update StatePath to include the frame ID. Add static version for the search methods. It helps when the allocator is different from the one in FActiveStatePath. Move the element outside of FActiveStatePath.

#rnx

diff --git a/Engine/Plugins/Runtime/StateTree/Source/StateTreeModule/Private/StateTreeExecutionTypes.cpp b/Engine/Plugins/Runtime/StateTree/Source/StateTreeModule/Private/StateTreeExecutionTypes.cpp
--- a/Engine/Plugins/Runtime/StateTree/Source/StateTreeModule/Private/StateTreeExecutionTypes.cpp
+++ b/Engine/Plugins/Runtime/StateTree/Source/StateTreeModule/Private/StateTreeExecutionTypes.cpp
@@ -61,14 +61,14 @@ UE::StateTree::FActiveStatePath FStateTreeExecutionState::GetActiveStatePath() c
 		return UE::StateTree::FActiveStatePath();
 	}
 
-	TArray<UE::StateTree::FActiveStatePath::FElement> Elements;
+	TArray<UE::StateTree::FActiveState> Elements;
 	Elements.Reserve(NewNum);
 
 	for (const FStateTreeExecutionFrame& Frame : ActiveFrames)
 	{
 		for (int32 StateIndex = 0; StateIndex < Frame.ActiveStates.Num(); ++StateIndex)
 		{
-			Elements.Emplace(Frame.ActiveStates.States[StateIndex], Frame.ActiveStates.StateIDs[StateIndex]);
+			Elements.Emplace(Frame.FrameID, Frame.ActiveStates.StateIDs[StateIndex], Frame.ActiveStates.States[StateIndex]);
 		}
 	}
 
diff --git a/Engine/Plugins/Runtime/StateTree/Source/StateTreeModule/Private/StateTreeStatePath.cpp b/Engine/Plugins/Runtime/StateTree/Source/StateTreeModule/Private/StateTreeStatePath.cpp
--- a/Engine/Plugins/Runtime/StateTree/Source/StateTreeModule/Private/StateTreeStatePath.cpp
+++ b/Engine/Plugins/Runtime/StateTree/Source/StateTreeModule/Private/StateTreeStatePath.cpp
@@ -10,15 +10,15 @@ namespace UE::StateTree
 	const FActiveFrameID FActiveFrameID::Invalid = FActiveFrameID();
 	const FActiveStateID FActiveStateID::Invalid = FActiveStateID();
 
-	FActiveStatePath::FActiveStatePath(TNotNull<const UStateTree*> InStateTree, const TArrayView<const FElement> InElements)
+	FActiveStatePath::FActiveStatePath(TNotNull<const UStateTree*> InStateTree, const TArrayView<const FActiveState> InElements)
 		: StateTree(InStateTree)
 		, States(InElements)
 	{
 #if WITH_STATETREE_DEBUG
 		bool bIsValid = true;
-		for (const FElement& Element : InElements)
+		for (const FActiveState& Element : InElements)
 		{
-			if (!Element.GetHandle().IsValid() || Element.GetHandle().IsCompletionState())
+			if (!Element.GetStateHandle().IsValid() || Element.GetStateHandle().IsCompletionState())
 			{
 				bIsValid = false;
 				break;
@@ -28,15 +28,15 @@ namespace UE::StateTree
 #endif
 	}
 
-	FActiveStatePath::FActiveStatePath(TNotNull<const UStateTree*> InStateTree, TArray<FElement> InElements)
+	FActiveStatePath::FActiveStatePath(TNotNull<const UStateTree*> InStateTree, TArray<FActiveState> InElements)
 		: StateTree(InStateTree)
 		, States(MoveTemp(InElements))
 	{
 #if WITH_STATETREE_DEBUG
 		bool bIsValid = true;
-		for (const FElement& Element : InElements)
+		for (const FActiveState& Element : InElements)
 		{
-			if (!Element.GetHandle().IsValid() || Element.GetHandle().IsCompletionState())
+			if (!Element.GetStateHandle().IsValid() || Element.GetStateHandle().IsCompletionState())
 			{
 				bIsValid = false;
 				break;
@@ -46,89 +46,131 @@ namespace UE::StateTree
 #endif
 	}
 
+	bool FActiveStatePath::Matches(const TArrayView<const FActiveState>& A, const TArrayView<const FActiveState>& B)
+	{
+		return A.Num() == B.Num() && Intersect(A, B).Num() == A.Num();
+	}
+
 	bool FActiveStatePath::Matches(const FActiveStatePath& Other) const
 	{
-		return States.Num() == Other.States.Num() && Intersect(Other).Num() == States.Num();
+		return Matches(MakeConstArrayView(States), MakeConstArrayView(Other.States));
 	}
 
-	bool FActiveStatePath::Matches(FElement Other) const
+	bool FActiveStatePath::Matches(const TArrayView<const FActiveState>& States, FActiveState Other)
 	{
 		return States.Num() > 0 && States.Last() == Other;
 	}
+
+	bool FActiveStatePath::Matches(FActiveState Other) const
+	{
+		return Matches(MakeConstArrayView(States), Other);
+	}
+
+	bool FActiveStatePath::Matches(const TArrayView<const FActiveState>& States, FActiveStateID Other)
+	{
+		return States.Num() > 0 && States.Last().GetStateID() == Other;
+	}
 	
 	bool FActiveStatePath::Matches(FActiveStateID Other) const
 	{
-		return States.Num() > 0 && States.Last().GetActiveStateID() == Other;
+		return Matches(MakeConstArrayView(States), Other);
 	}
 
-	const TArrayView<const FActiveStatePath::FElement> FActiveStatePath::Intersect(const FActiveStatePath& Other) const
+	const TArrayView<const FActiveState> FActiveStatePath::Intersect(const TArrayView<const FActiveState>& A, const TArrayView<const FActiveState>& B)
 	{
 		// If the ID is the same, then are equal.
-		if (States.Num() == Other.States.Num()
-			&& States.Num() > 0
-			&& States.Last() == Other.States.Last())
+		if (A.Num() == B.Num()
+			&& A.Num() > 0
+			&& A.Last() == B.Last())
 		{
-			return States;
+			return A;
 		}
 
 		int32 MatchIndex = 0;
-		const int32 MinNum = FMath::Min(States.Num(), Other.States.Num());
+		const int32 MinNum = FMath::Min(A.Num(), B.Num());
 		for (MatchIndex = 0; MatchIndex < MinNum; ++MatchIndex)
 		{
-			if (States[MatchIndex] != Other.States[MatchIndex])
+			if (A[MatchIndex] != B[MatchIndex])
 			{
 				break;
 			}
 		}
-		return MakeConstArrayView(States.GetData(), MatchIndex);
+		return MakeConstArrayView(A.GetData(), MatchIndex);
 	}
 
-	bool FActiveStatePath::StartsWith(const FActiveStatePath& Other) const
+	const TArrayView<const FActiveState> FActiveStatePath::Intersect(const FActiveStatePath& Other) const
+	{
+		return Intersect(MakeConstArrayView(States), MakeConstArrayView(Other.States));
+	}
+
+	bool FActiveStatePath::StartsWith(const TArrayView<const FActiveState>& A, const TArrayView<const FActiveState>& B)
 	{
 		// No need to check all the element. If the first and last match then they all matches.
-		return Other.States.Num() > 0 ? Contains(Other.States.Last()) && States[0] == Other.States[0] : false;
+		return B.Num() > 0 ? Contains(A, B.Last()) && A[0] == B[0] : false;
+	}
+
+	bool FActiveStatePath::StartsWith(const FActiveStatePath& Other) const
+	{
+		return StartsWith(MakeConstArrayView(States), MakeConstArrayView(Other.States));
 	}
 
-	bool FActiveStatePath::Contains(FElement Other) const
+	bool FActiveStatePath::Contains(const TArrayView<const FActiveState>& States, FActiveState Other)
 	{
-		for (int32 Index = States.Num() - 1; Index >= 0; --Index)
-		{
-			if (States[Index] == Other)
+		return States.Contains(Other);
+	}
+
+	bool FActiveStatePath::Contains(FActiveState Other) const
+	{
+		return Contains(MakeConstArrayView(States), Other);
+	}
+
+	bool FActiveStatePath::Contains(const TArrayView<const FActiveState>& States, FActiveStateID Other)
+	{
+		return Other.IsValid()
+			&& States.ContainsByPredicate([Other](const FActiveState& Element)
 			{
-				return true;
-			}
-		}
-		return false;
+				return Element.GetStateID() == Other;
+			});
 	}
 
 	bool FActiveStatePath::Contains(FActiveStateID Other) const
 	{
-		if (Other.IsValid())
-		{
-			for (int32 Index = States.Num() - 1; Index >= 0; --Index)
-			{
-				if (States[Index].GetActiveStateID() == Other)
+		return Contains(MakeConstArrayView(States), Other);
+	}
+
+	int32 FActiveStatePath::IndexOf(const TArrayView<const FActiveState>& States, FActiveState Other)
+	{
+		return States.IndexOfByKey(Other);
+	}
+
+	int32 FActiveStatePath::IndexOf(FActiveState Other) const
+	{
+		return IndexOf(MakeConstArrayView(States), Other);
+	}
+
+	int32 FActiveStatePath::IndexOf(const TArrayView<const FActiveState>& States, FActiveStateID Other)
+	{
+		return Other.IsValid()
+			? States.IndexOfByPredicate([Other](const FActiveState& Element)
 				{
-					return true;
-				}
-			}
-		}
-		return false;
+					return Element.GetStateID() == Other;
+				})
+			: INDEX_NONE;
 	}
 
-	TValueOrError<FString, void> FActiveStatePath::Describe() const
+	int32 FActiveStatePath::IndexOf(FActiveStateID Other) const
 	{
-		const UStateTree* CurrentTree = StateTree.Get();
-		if (CurrentTree == nullptr)
-		{
-			return MakeError();
-		}
+		return IndexOf(MakeConstArrayView(States), Other);
+	}
 
+	TValueOrError<FString, void> FActiveStatePath::Describe(TNotNull<const UStateTree*> StateTree, const TArrayView<const FActiveState>& States)
+	{
+		const UStateTree* CurrentTree = StateTree;
 		TStringBuilder<256> Builder;
 		bool bNewTree = true;
-		for (const FElement& Element : States)
+		for (const FActiveState& Element : States)
 		{
-			if (!Element.GetHandle().IsValid() || Element.GetHandle().IsCompletionState())
+			if (!Element.GetStateHandle().IsValid() || Element.GetStateHandle().IsCompletionState())
 			{
 				return MakeError();
 			}
@@ -146,15 +188,15 @@ namespace UE::StateTree
 					bNewTree = false;
 				}
 
-				if (CurrentTree->GetStates().IsValidIndex(Element.GetHandle().Index))
+				if (CurrentTree->GetStates().IsValidIndex(Element.GetStateHandle().Index))
 				{
 					return MakeError();
 				}
 
-				const FCompactStateTreeState& State = CurrentTree->GetStates()[Element.GetHandle().Index];
+				const FCompactStateTreeState& State = CurrentTree->GetStates()[Element.GetStateHandle().Index];
 				Builder << State.Name;
 				Builder << TEXT('(');
-				Builder << Element.GetHandle().Index;
+				Builder << Element.GetStateHandle().Index;
 				Builder << TEXT(')');
 
 				if (State.Type == EStateTreeStateType::LinkedAsset)
@@ -170,4 +212,13 @@ namespace UE::StateTree
 		return MakeValue(Builder.ToString());
 	}
 
-}
\ No newline at end of file
+	TValueOrError<FString, void> FActiveStatePath::Describe() const
+	{
+		const UStateTree* CurrentTree = StateTree.Get();
+		if (CurrentTree == nullptr)
+		{
+			return MakeError();
+		}
+		return Describe(CurrentTree, MakeConstArrayView(States));
+	}
+}
diff --git a/Engine/Plugins/Runtime/StateTree/Source/StateTreeModule/Public/StateTreeStatePath.h b/Engine/Plugins/Runtime/StateTree/Source/StateTreeModule/Public/StateTreeStatePath.h
--- a/Engine/Plugins/Runtime/StateTree/Source/StateTreeModule/Public/StateTreeStatePath.h
+++ b/Engine/Plugins/Runtime/StateTree/Source/StateTreeModule/Public/StateTreeStatePath.h
@@ -2,7 +2,6 @@
 
 #pragma once
 
-#include "CoreMinimal.h"
 #include "Containers/ArrayView.h"
 #include "Misc/NotNull.h"
 #include "StateTreeTypes.h"
@@ -21,7 +20,7 @@ namespace UE::StateTree
 		static constexpr uint32 InvalidValue = 0;
 
 	public:
-		static const FActiveFrameID Invalid;
+		UE_API static const FActiveFrameID Invalid;
 
 		explicit FActiveFrameID() = default;
 		explicit FActiveFrameID(uint32 NewID)
@@ -49,7 +48,7 @@ namespace UE::StateTree
 		static constexpr uint32 InvalidValue = 0;
 
 	public:
-		static const FActiveStateID Invalid;
+		UE_API static const FActiveStateID Invalid;
 
 		explicit FActiveStateID() = default;
 		explicit FActiveStateID(uint32 NewID)
@@ -70,6 +69,44 @@ namespace UE::StateTree
 		uint32 Value = InvalidValue;
 	};
 
+	/** A state entry in the state path. */
+	struct FActiveState
+	{
+		FActiveState() = default;
+		inline explicit FActiveState(FActiveFrameID InFrameID, FActiveStateID InStateID, FStateTreeStateHandle InHandle);
+
+		[[nodiscard]] bool operator==(const FActiveState&) const = default;
+		[[nodiscard]] bool operator!=(const FActiveState&) const = default;
+
+		[[nodiscard]] bool IsValid() const
+		{
+			return StateHandle.IsValid();
+		}
+
+		[[nodiscard]] FActiveFrameID GetFrameID() const
+		{
+			return FrameID;
+		}
+
+		[[nodiscard]] FActiveStateID GetStateID() const
+		{
+			return StateID;
+		}
+
+		[[nodiscard]] FStateTreeStateHandle GetStateHandle() const
+		{
+			return StateHandle;
+		}
+
+	private:
+		/** The unique ID of the state for this instance. */
+		FActiveFrameID FrameID;
+		/** The unique ID of the state for this instance. */
+		FActiveStateID StateID;
+		/** The index of the state handle the state tree asset. */
+		FStateTreeStateHandle StateHandle;
+	};
+
 	/**
 	 * Describe the state list used to get to a specific state.
 	 * Since a state can be a subtree and can be linked to other states, the path to activate the subtree can be different.
@@ -87,44 +124,10 @@ namespace UE::StateTree
 	 */
 	class FActiveStatePath
 	{
-	public:
-		/** An state entry in the state path. */
-		struct FElement
-		{
-			explicit FElement(FStateTreeStateHandle InHandle, FActiveStateID InStateID)
-				: ActiveStateID(InStateID)
-				, TreeStateIndex(InHandle)
-			{}
-
-			[[nodiscard]] bool operator==(const FElement&) const = default;
-			[[nodiscard]] bool operator!=(const FElement&) const = default;
-
-			[[nodiscard]] bool IsValid() const
-			{
-				return TreeStateIndex.IsValid();
-			}
-
-			[[nodiscard]] FActiveStateID GetActiveStateID() const
-			{
-				return ActiveStateID;
-			}
-
-			[[nodiscard]] FStateTreeStateHandle GetHandle() const
-			{
-				return TreeStateIndex;
-			}
-
-		private:
-			/** The unique ID of the state for this instance. */
-			FActiveStateID ActiveStateID;
-			/** The index of the state inside the state tree asset. */
-			FStateTreeStateHandle TreeStateIndex;
-		};
-
 	public:
 		explicit FActiveStatePath() = default;
-		UE_API explicit FActiveStatePath(TNotNull<const UStateTree*> InStateTree, const TArrayView<const FElement> InElements);
-		UE_API explicit FActiveStatePath(TNotNull<const UStateTree*> InStateTree, TArray<FElement> InElements);
+		UE_API explicit FActiveStatePath(TNotNull<const UStateTree*> InStateTree, const TArrayView<const FActiveState> InElements);
+		UE_API explicit FActiveStatePath(TNotNull<const UStateTree*> InStateTree, TArray<FActiveState> InElements);
 
 		/** @return the number of elements in the path. */
 		[[nodiscard]] int32 Num() const
@@ -133,42 +136,83 @@ namespace UE::StateTree
 		}
 
 		/** @return a view of the states in the path. */
-		[[nodiscard]] const TArrayView<const FElement> GetView() const
+		[[nodiscard]] const TArrayView<const FActiveState> GetView() const
 		{
 			return States;
 		}
 
+		/** @return true if both paths match exactly. */
+		[[nodiscard]] UE_API static bool Matches(const TArrayView<const FActiveState>& A, const TArrayView<const FActiveState>& B);
+
 		/** @return true if both paths match exactly. */
 		[[nodiscard]] UE_API bool Matches(const FActiveStatePath& Other) const;
 
 		/** @return true if the last element in the path matches. */
-		[[nodiscard]] UE_API bool Matches(FElement Other) const;
+		[[nodiscard]] UE_API static bool Matches(const TArrayView<const FActiveState>& Path, FActiveState Other);
+
+		/** @return true if the last element in the path matches. */
+		[[nodiscard]] UE_API bool Matches(FActiveState Other) const;
+
+		/** @return true if the last element in the path matches. */
+		[[nodiscard]] UE_API static bool Matches(const TArrayView<const FActiveState>& Path, FActiveStateID Other);
 
 		/** @return true if the last element in the path matches. */
 		[[nodiscard]] UE_API bool Matches(FActiveStateID Other) const;
 
+		/** @return the elements that are in both paths. */
+		[[nodiscard]] UE_API static const TArrayView<const FActiveState> Intersect(const TArrayView<const FActiveState>& A, const TArrayView<const FActiveState>& B);
+
 		/** @return the path where the 2 paths intersect. */
-		[[nodiscard]] UE_API const TArrayView<const FElement> Intersect(const FActiveStatePath& Other) const;
+		[[nodiscard]] UE_API const TArrayView<const FActiveState> Intersect(const FActiveStatePath& Other) const;
 
-		/** @return the path that is not in the other path. */
-		[[nodiscard]] UE_API const TArrayView<const FElement> Difference(const FActiveStatePath& Other) const;
+		/** @return true if the path starts with and contains the other path. */
+		[[nodiscard]]UE_API  static bool StartsWith(const TArrayView<const FActiveState>& Path, const TArrayView<const FActiveState>& B);
 
 		/** @return true if this path starts with and contains the other path. */
 		[[nodiscard]] UE_API bool StartsWith(const FActiveStatePath& Other) const;
 
+		/** @return true if the path contains anywhere the other element. */
+		[[nodiscard]] UE_API static bool Contains(const TArrayView<const FActiveState>& Path, FActiveState Other);
+		
 		/** @return true if this path contains anywhere the other element. */
-		[[nodiscard]] UE_API bool Contains(FElement Other) const;
+		[[nodiscard]] UE_API bool Contains(FActiveState Other) const;
+
+		/** @return true if the path contains anywhere the other element. */
+		[[nodiscard]] UE_API static bool Contains(const TArrayView<const FActiveState>& Path, FActiveStateID Other);
 
 		/** @return true if this path contains anywhere the other element. */
 		[[nodiscard]] UE_API bool Contains(FActiveStateID Other) const;
 
+		/** @return true the index of element inside the path. */
+		[[nodiscard]] UE_API static int32 IndexOf(const TArrayView<const FActiveState>& Path, FActiveState Other);
+
+		/** @return true the index of element inside the path. */
+		[[nodiscard]]UE_API int32 IndexOf(FActiveState Other) const;
+
+		/** @return true the index of element inside the path. */
+		[[nodiscard]] UE_API static int32 IndexOf(const TArrayView<const FActiveState>& Path, FActiveStateID Other);
+
+		/** @return true the index of element inside the path. */
+		[[nodiscard]] UE_API int32 IndexOf(FActiveStateID Other) const;
+
+		/** @return debugging information about the path. */
+		[[nodiscard]] UE_API static TValueOrError<FString, void> Describe(TNotNull<const UStateTree*> StateTree, const TArrayView<const FActiveState>& Path);
+		
 		/** @return debugging information about the path. */
 		[[nodiscard]] UE_API TValueOrError<FString, void> Describe() const;
 
 	private:
 		TWeakObjectPtr<const UStateTree> StateTree;
-		TArray<FElement> States;
+		TArray<FActiveState> States;
 	};
+
+	// Inlined functions
+	FActiveState::FActiveState(FActiveFrameID InFrameID, FActiveStateID InStateID, FStateTreeStateHandle InHandle)
+		: FrameID(InFrameID)
+		, StateID(InStateID)
+		, StateHandle(InHandle)
+	{
+	}
 }
 
 #undef UE_API
diff --git a/Engine/Plugins/Runtime/StateTree/Source/StateTreeTestSuite/Private/StateTreeStatePathTest.cpp b/Engine/Plugins/Runtime/StateTree/Source/StateTreeTestSuite/Private/StateTreeStatePathTest.cpp
--- a/Engine/Plugins/Runtime/StateTree/Source/StateTreeTestSuite/Private/StateTreeStatePathTest.cpp
+++ b/Engine/Plugins/Runtime/StateTree/Source/StateTreeTestSuite/Private/StateTreeStatePathTest.cpp
@@ -204,7 +204,7 @@ struct FStateTreeTest_StatePath_LinkStates : FStateTreeTestBase
 			Sub1.Type = EStateTreeStateType::Subtree;
 			//Root
 			{
-				UStateTreeState& State = Root.AddChildState("Tree2RootStateSub1A", EStateTreeStateType::Linked);
+				UStateTreeState& State = Root.AddChildState("Tree3RootStateSub1A", EStateTreeStateType::Linked);
 				State.SetLinkedState(Sub1.GetLinkToState());
 				State.AddTransition(EStateTreeTransitionTrigger::OnStateSucceeded, EStateTreeTransitionType::Succeeded);
 			}
@@ -285,53 +285,58 @@ struct FStateTreeTest_StatePath_LinkStates : FStateTreeTestBase
 
 		// Start tests
 		{
+			using namespace UE::StateTree;
+
 			const EStateTreeRunStatus Status = Exec.Start();
 			AITEST_EQUAL("Start should complete with Running", Status, EStateTreeRunStatus::Running);
 			AITEST_TRUE("Should be in the correct state", Exec.ExpectInActiveStates("Tree1Root", "Tree1RootState1", "Tree1State1StateSub1A", "Tree1Sub1", "Tree1Sub1StateLinkTree2A", "Tree2StateRoot", "Tree2RootStateSub1A", "Tree2Sub1", "Tree2Sub1State1"));
 
-			const UE::StateTree::FActiveStatePath ExecPath = InstanceData.GetExecutionState()->GetActiveStatePath();
+			const FActiveStatePath ExecPath = InstanceData.GetExecutionState()->GetActiveStatePath();
 			AITEST_TRUE("Has the correct number of path elements", ExecPath.Num() == 9);
 			AITEST_TRUE("Has the correct number of active states", InstanceData.GetExecutionState()->ActiveFrames.Num() == 4);
 			{
-				AITEST_TRUE("Frame for Tree1Root is active", InstanceData.GetExecutionState()->ActiveFrames[0].FrameID == UE::StateTree::FActiveFrameID(++ActiveCounter));
-				AITEST_TRUE("State Tree1Root is active", ExecPath.Contains(UE::StateTree::FActiveStateID(++ActiveCounter)));
-				AITEST_TRUE("State Tree1Root is active", ExecPath.Contains(UE::StateTree::FActiveStatePath::FElement(FStateTreeStateHandle(0), UE::StateTree::FActiveStateID(ActiveCounter))));
-				AITEST_TRUE("State Tree1RootState1 is active", ExecPath.Contains(UE::StateTree::FActiveStateID(++ActiveCounter)));
-				AITEST_TRUE("State Tree1RootState1 is active", ExecPath.Contains(UE::StateTree::FActiveStatePath::FElement(FStateTreeStateHandle(1), UE::StateTree::FActiveStateID(ActiveCounter))));
-				AITEST_TRUE("State Tree1State1StateSub1A is active", ExecPath.Contains(UE::StateTree::FActiveStateID(++ActiveCounter)));
-				AITEST_TRUE("State Tree1State1StateSub1A is active", ExecPath.Contains(UE::StateTree::FActiveStatePath::FElement(FStateTreeStateHandle(2), UE::StateTree::FActiveStateID(4))));
+				const FActiveFrameID FirstFrameID = InstanceData.GetExecutionState()->ActiveFrames[0].FrameID;
+				AITEST_TRUE("Frame for Tree1Root is active", FirstFrameID == FActiveFrameID(++ActiveCounter));
+				AITEST_TRUE("State Tree1Root is active", ExecPath.Contains(FActiveStateID(++ActiveCounter)));
+				AITEST_TRUE("State Tree1Root is active", ExecPath.Contains(FActiveState(FirstFrameID, FActiveStateID(ActiveCounter), FStateTreeStateHandle(0))));
+				AITEST_TRUE("State Tree1RootState1 is active", ExecPath.Contains(FActiveStateID(++ActiveCounter)));
+				AITEST_TRUE("State Tree1RootState1 is active", ExecPath.Contains(FActiveState(FirstFrameID, FActiveStateID(ActiveCounter), FStateTreeStateHandle(1))));
+				AITEST_TRUE("State Tree1State1StateSub1A is active", ExecPath.Contains(FActiveStateID(++ActiveCounter)));
+				AITEST_TRUE("State Tree1State1StateSub1A is active", ExecPath.Contains(FActiveState(FirstFrameID, FActiveStateID(4), FStateTreeStateHandle(2))));
 			}
 			{
-				AITEST_TRUE("Frame for Tree1Sub1 is active", InstanceData.GetExecutionState()->ActiveFrames[1].FrameID == UE::StateTree::FActiveFrameID(++ActiveCounter));
-				AITEST_TRUE("State Tree1Sub1 is active", ExecPath.Contains(UE::StateTree::FActiveStateID(++ActiveCounter)));
-				AITEST_TRUE("State Tree1Sub1StateLinkTree2A is active", ExecPath.Contains(UE::StateTree::FActiveStateID(++ActiveCounter)));
+				AITEST_TRUE("Frame for Tree1Sub1 is active", InstanceData.GetExecutionState()->ActiveFrames[1].FrameID == FActiveFrameID(++ActiveCounter));
+				AITEST_TRUE("State Tree1Sub1 is active", ExecPath.Contains(FActiveStateID(++ActiveCounter)));
+				AITEST_TRUE("State Tree1Sub1StateLinkTree2A is active", ExecPath.Contains(FActiveStateID(++ActiveCounter)));
 			}
 			{
-				AITEST_TRUE("Frame for Tree2StateRoot is active", InstanceData.GetExecutionState()->ActiveFrames[2].FrameID == UE::StateTree::FActiveFrameID(++ActiveCounter));
-				AITEST_TRUE("State Tree2StateRoot is active", ExecPath.Contains(UE::StateTree::FActiveStateID(++ActiveCounter)));
-				AITEST_TRUE("State Tree2RootStateSub1A is active", ExecPath.Contains(UE::StateTree::FActiveStateID(++ActiveCounter)));
+				AITEST_TRUE("Frame for Tree2StateRoot is active", InstanceData.GetExecutionState()->ActiveFrames[2].FrameID == FActiveFrameID(++ActiveCounter));
+				AITEST_TRUE("State Tree2StateRoot is active", ExecPath.Contains(FActiveStateID(++ActiveCounter)));
+				AITEST_TRUE("State Tree2RootStateSub1A is active", ExecPath.Contains(FActiveStateID(++ActiveCounter)));
 			}
 			{
-				AITEST_TRUE("Frame for Tree2Sub1 is active", InstanceData.GetExecutionState()->ActiveFrames[3].FrameID == UE::StateTree::FActiveFrameID(++ActiveCounter));
-				AITEST_TRUE("State Tree2Sub1 is active", ExecPath.Contains(UE::StateTree::FActiveStateID(++ActiveCounter)));
-				AITEST_TRUE("State Tree2Sub1State1 is active", ExecPath.Contains(UE::StateTree::FActiveStateID(++ActiveCounter)));
+				AITEST_TRUE("Frame for Tree2Sub1 is active", InstanceData.GetExecutionState()->ActiveFrames[3].FrameID == FActiveFrameID(++ActiveCounter));
+				AITEST_TRUE("State Tree2Sub1 is active", ExecPath.Contains(FActiveStateID(++ActiveCounter)));
+				AITEST_TRUE("State Tree2Sub1State1 is active", ExecPath.Contains(FActiveStateID(++ActiveCounter)));
 			}
 			{
-				AITEST_FALSE("No accidental increment", ExecPath.Contains(UE::StateTree::FActiveStateID(ActiveCounter+1)));
+				AITEST_FALSE("No accidental increment", ExecPath.Contains(FActiveStateID(ActiveCounter+1)));
 			}
 			Exec.LogClear();
 		}
 		{
+			using namespace UE::StateTree;
+
 			PreTestFrameId();
 			EStateTreeRunStatus Status = Exec.Tick(1.0f);
 			AITEST_EQUAL("Start should complete with Running", Status, EStateTreeRunStatus::Running);
 			AITEST_TRUE("Should be in the correct state", Exec.ExpectInActiveStates("Tree1Root", "Tree1RootState1", "Tree1State1StateSub1A", "Tree1Sub1", "Tree1Sub1StateLinkTree2A", "Tree2StateRoot", "Tree2RootStateSub1A", "Tree2Sub1", "Tree2Sub1State2"));
 
-			UE::StateTree::FActiveStatePath ExecPath = InstanceData.GetExecutionState()->GetActiveStatePath();
+			FActiveStatePath ExecPath = InstanceData.GetExecutionState()->GetActiveStatePath();
 			AITEST_TRUE("Has the correct number of path elements", ExecPath.Num() == 9);
 			AITEST_TRUE("Has the correct number of active states", InstanceData.GetExecutionState()->ActiveFrames.Num() == 4);
-			AITEST_TRUE("State Tree2Sub1State2 is active", ExecPath.Contains(UE::StateTree::FActiveStateID(++ActiveCounter)));
-			AITEST_FALSE("No accidental increment", ExecPath.Contains(UE::StateTree::FActiveStateID(ActiveCounter + 1)));
+			AITEST_TRUE("State Tree2Sub1State2 is active", ExecPath.Contains(FActiveStateID(++ActiveCounter)));
+			AITEST_FALSE("No accidental increment", ExecPath.Contains(FActiveStateID(ActiveCounter + 1)));
 			AITEST_TRUE("All frames are the same", TestFrameId(4));
 			Exec.LogClear();
 		}
