StateTree: Prevent enter conditions and utilities from binding to task instance data. A parent state can try to enter a child state before it's accepted by the transition. The task instance data is only available once the transition is completed.

diff --git a/Engine/Plugins/Runtime/StateTree/Source/StateTreeEditorModule/Private/StateTreeCompiler.cpp b/Engine/Plugins/Runtime/StateTree/Source/StateTreeEditorModule/Private/StateTreeCompiler.cpp
--- a/Engine/Plugins/Runtime/StateTree/Source/StateTreeEditorModule/Private/StateTreeCompiler.cpp
+++ b/Engine/Plugins/Runtime/StateTree/Source/StateTreeEditorModule/Private/StateTreeCompiler.cpp
@@ -2350,11 +2350,26 @@ bool FStateTreeCompiler::GetAndValidateBindings(const FStateTreeBindableStructDe
 
 		if (!bSourceAccessible)
 		{
-			Log.Reportf(EMessageSeverity::Error, TargetStruct,
-						TEXT("Property at %s cannot be bound to %s, because the binding source %s is not updated before %s in the tree."),
-						*UE::StateTree::GetDescAndPathAsString(*SourceStruct, Binding.GetSourcePath()),
-						*UE::StateTree::GetDescAndPathAsString(TargetStruct, Binding.GetTargetPath()),
-						*SourceStruct->ToString(), *TargetStruct.ToString());
+			TInstancedStruct<FPropertyBindingBindableStructDescriptor> SourceStructDescriptor;
+			const bool bFoundSourceStructDescriptor = EditorData->GetBindableStructByID(SourceStructID, SourceStructDescriptor);
+			if (bFoundSourceStructDescriptor
+				&& SourceStructDescriptor.Get<FStateTreeBindableStructDesc>().DataSource == EStateTreeBindableStructSource::Task
+				&& !UE::StateTree::AcceptTaskInstanceData(TargetStruct.DataSource))
+			{
+				Log.Reportf(EMessageSeverity::Error, TargetStruct,
+					TEXT("Property at %s cannot be bound to %s, because the binding source %s is an instance data that is possibly not instantiated before %s in the tree."),
+					* UE::StateTree::GetDescAndPathAsString(*SourceStruct, Binding.GetSourcePath()),
+					* UE::StateTree::GetDescAndPathAsString(TargetStruct, Binding.GetTargetPath()),
+					* SourceStruct->ToString(), * TargetStruct.ToString());
+			}
+			else
+			{
+				Log.Reportf(EMessageSeverity::Error, TargetStruct,
+					TEXT("Property at %s cannot be bound to %s, because the binding source %s is not updated before %s in the tree."),
+					*UE::StateTree::GetDescAndPathAsString(*SourceStruct, Binding.GetSourcePath()),
+					*UE::StateTree::GetDescAndPathAsString(TargetStruct, Binding.GetTargetPath()),
+					*SourceStruct->ToString(), *TargetStruct.ToString());
+			}
 			bSucceeded = false;
 			continue;
 		}
diff --git a/Engine/Plugins/Runtime/StateTree/Source/StateTreeEditorModule/Private/StateTreeEditorData.cpp b/Engine/Plugins/Runtime/StateTree/Source/StateTreeEditorModule/Private/StateTreeEditorData.cpp
--- a/Engine/Plugins/Runtime/StateTree/Source/StateTreeEditorModule/Private/StateTreeEditorData.cpp
+++ b/Engine/Plugins/Runtime/StateTree/Source/StateTreeEditorModule/Private/StateTreeEditorData.cpp
@@ -28,6 +28,14 @@
 
 namespace UE::StateTree::Editor
 {
+	FAutoConsoleVariable CVarLogEnableBindingSelectionNodeToInstanceData(
+		TEXT("StateTree.Compiler.EnableBindingSelectionNodeToInstanceData"),
+		true,
+		TEXT("Enable binding from enter condition, utility/consideration and state argument to bind to task instance data.\n")
+		TEXT("The task instance data is only available once the transition is completed.")
+		TEXT("A parent state can enter a child state during state selection (before the transition completes).")
+	);
+
 	const FString GlobalStateName(TEXT("Global"));
 	const FString PropertyFunctionStateName(TEXT("Property Functions"));
 	const FName ParametersNodeName(TEXT("Parameters"));
@@ -371,11 +379,16 @@ void UStateTreeEditorData::GetAccessibleStructsInExecutionPath(const TConstArray
 	const UStateTree* StateTree = GetTypedOuter<UStateTree>();
 	checkf(StateTree, TEXT("UStateTreeEditorData should only be allocated within a UStateTree"));
 
+	bool bAcceptTaskInstanceData = true;
 	TInstancedStruct<FPropertyBindingBindableStructDescriptor> TargetStructDesc;
 	bool bIsTargetPropertyFunction = false;
 	if (GetBindableStructByID(TargetStructID, TargetStructDesc))
 	{
 		bIsTargetPropertyFunction = TargetStructDesc.Get<FStateTreeBindableStructDesc>().DataSource == EStateTreeBindableStructSource::PropertyFunction;
+		if (!UE::StateTree::Editor::CVarLogEnableBindingSelectionNodeToInstanceData->GetBool())
+		{
+			bAcceptTaskInstanceData = UE::StateTree::AcceptTaskInstanceData(TargetStructDesc.Get<FStateTreeBindableStructDesc>().DataSource);
+		}
 	}
 
 	EStateTreeVisitor BaseProgress = VisitGlobalNodes([&OutStructDescs, TargetStructID]
@@ -402,7 +415,8 @@ void UStateTreeEditorData::GetAccessibleStructsInExecutionPath(const TConstArray
 				continue;
 			}
 			
-			const EStateTreeVisitor StateProgress = VisitStateNodes(*State, [&OutStructDescs, &BindableDescs, &Path, TargetStructID, bIsTargetPropertyFunction, this]
+			const EStateTreeVisitor StateProgress = VisitStateNodes(*State,
+				[&OutStructDescs, &BindableDescs, &Path, TargetStructID, bIsTargetPropertyFunction, bAcceptTaskInstanceData, this]
 				(const UStateTreeState* State, const FStateTreeBindableStructDesc& Desc, const FStateTreeDataView Value)
 				{
 					// Stop iterating as soon as we find the target node.
@@ -416,11 +430,17 @@ void UStateTreeEditorData::GetAccessibleStructsInExecutionPath(const TConstArray
 					switch (Desc.DataSource)
 					{
 						case EStateTreeBindableStructSource::StateParameter:
-						case EStateTreeBindableStructSource::Task:
 						case EStateTreeBindableStructSource::StateEvent:
 							BindableDescs.Add(TInstancedStruct<FStateTreeBindableStructDesc>::Make(Desc));
 							break;
 
+						case EStateTreeBindableStructSource::Task:
+							if (bAcceptTaskInstanceData)
+							{
+								BindableDescs.Add(TInstancedStruct<FStateTreeBindableStructDesc>::Make(Desc));
+							}
+							break;
+
 						case EStateTreeBindableStructSource::TransitionEvent:
 						{
 							// Checking if BindableStruct's owning Transition contains the Target.
diff --git a/Engine/Plugins/Runtime/StateTree/Source/StateTreeModule/Private/StateTreePropertyBindings.cpp b/Engine/Plugins/Runtime/StateTree/Source/StateTreeModule/Private/StateTreePropertyBindings.cpp
--- a/Engine/Plugins/Runtime/StateTree/Source/StateTreeModule/Private/StateTreePropertyBindings.cpp
+++ b/Engine/Plugins/Runtime/StateTree/Source/StateTreeModule/Private/StateTreePropertyBindings.cpp
@@ -18,6 +18,13 @@
 
 namespace UE::StateTree
 {
+	bool AcceptTaskInstanceData(EStateTreeBindableStructSource Target)
+	{
+		return Target != EStateTreeBindableStructSource::StateParameter
+			&& Target != EStateTreeBindableStructSource::Condition
+			&& Target != EStateTreeBindableStructSource::Consideration;
+	}
+
 	FString GetDescAndPathAsString(const FStateTreeBindableStructDesc& Desc, const FPropertyBindingPath& Path)
 	{
 		return PropertyBinding::GetDescriptorAndPathAsString(Desc, Path);
diff --git a/Engine/Plugins/Runtime/StateTree/Source/StateTreeModule/Public/StateTreePropertyBindings.h b/Engine/Plugins/Runtime/StateTree/Source/StateTreeModule/Public/StateTreePropertyBindings.h
--- a/Engine/Plugins/Runtime/StateTree/Source/StateTreeModule/Public/StateTreePropertyBindings.h
+++ b/Engine/Plugins/Runtime/StateTree/Source/StateTreeModule/Public/StateTreePropertyBindings.h
@@ -50,6 +50,15 @@ enum class EStateTreeBindableStructSource : uint8
 	Transition,
 };
 
+namespace UE::StateTree
+{
+	/**
+	 * Can that binding type accept a task instance data to be the source.
+	 * Condition and Utility are constructed before the task instance data is constructed.
+	 */
+	[[nodiscard]] UE_API bool AcceptTaskInstanceData(EStateTreeBindableStructSource Target);
+}
+
 
 /**
  * Describes how the copy should be performed.
