StateTree: Resetting temporary instances after a transition. Previously, the instance data were reset only when the transition succeeded.

diff --git a/Engine/Plugins/Runtime/StateTree/Source/StateTreeModule/Private/StateTreeExecutionContext.cpp b/Engine/Plugins/Runtime/StateTree/Source/StateTreeModule/Private/StateTreeExecutionContext.cpp
--- a/Engine/Plugins/Runtime/StateTree/Source/StateTreeModule/Private/StateTreeExecutionContext.cpp
+++ b/Engine/Plugins/Runtime/StateTree/Source/StateTreeModule/Private/StateTreeExecutionContext.cpp
@@ -1188,6 +1188,8 @@ EStateTreeRunStatus FStateTreeExecutionContext::Start(FStartParameters Parameter
 			}
 		}
 
+		InstanceData.ResetTemporaryInstances();
+
 		if (Exec.LastTickStatus == EStateTreeRunStatus::Unset)
 		{
 			// Should not happen. This may happen if initial state could not be selected.
@@ -1515,6 +1517,8 @@ void FStateTreeExecutionContext::TickTriggerTransitionsInternal()
 	static constexpr int32 MaxIterations = 5;
 	for (int32 Iter = 0; Iter < MaxIterations; Iter++)
 	{
+		ON_SCOPE_EXIT{ InstanceData.ResetTemporaryInstances(); };
+
 		// Trigger conditional transitions or state succeed/failed transitions. First tick transition is handled here too.
 		if (TriggerTransitions())
 		{
