StateTree: Enable EnableBindingSelectionNodeToInstanceData by default.

diff --git a/Engine/Plugins/Runtime/StateTree/Source/StateTreeEditorModule/Private/StateTreeEditorData.cpp b/Engine/Plugins/Runtime/StateTree/Source/StateTreeEditorModule/Private/StateTreeEditorData.cpp
--- a/Engine/Plugins/Runtime/StateTree/Source/StateTreeEditorModule/Private/StateTreeEditorData.cpp
+++ b/Engine/Plugins/Runtime/StateTree/Source/StateTreeEditorModule/Private/StateTreeEditorData.cpp
@@ -30,7 +30,7 @@ namespace UE::StateTree::Editor
 {
 	FAutoConsoleVariable CVarLogEnableBindingSelectionNodeToInstanceData(
 		TEXT("StateTree.Compiler.EnableBindingSelectionNodeToInstanceData"),
-		true,
+		false,
 		TEXT("Enable binding from enter condition, utility/consideration and state argument to bind to task instance data.\n")
 		TEXT("The task instance data is only available once the transition is completed.")
 		TEXT("A parent state can enter a child state during state selection (before the transition completes).")
diff --git a/Engine/Plugins/Runtime/StateTree/Source/StateTreeModule/Private/StateTreePropertyBindings.cpp b/Engine/Plugins/Runtime/StateTree/Source/StateTreeModule/Private/StateTreePropertyBindings.cpp
--- a/Engine/Plugins/Runtime/StateTree/Source/StateTreeModule/Private/StateTreePropertyBindings.cpp
+++ b/Engine/Plugins/Runtime/StateTree/Source/StateTreeModule/Private/StateTreePropertyBindings.cpp
@@ -20,6 +20,7 @@ namespace UE::StateTree
 {
 	bool AcceptTaskInstanceData(EStateTreeBindableStructSource Target)
 	{
+		// Condition and Utility are constructed before the task instance data is constructed.
 		return Target != EStateTreeBindableStructSource::StateParameter
 			&& Target != EStateTreeBindableStructSource::Condition
 			&& Target != EStateTreeBindableStructSource::Consideration;
diff --git a/Engine/Plugins/Runtime/StateTree/Source/StateTreeModule/Public/StateTreePropertyBindings.h b/Engine/Plugins/Runtime/StateTree/Source/StateTreeModule/Public/StateTreePropertyBindings.h
--- a/Engine/Plugins/Runtime/StateTree/Source/StateTreeModule/Public/StateTreePropertyBindings.h
+++ b/Engine/Plugins/Runtime/StateTree/Source/StateTreeModule/Public/StateTreePropertyBindings.h
@@ -52,10 +52,7 @@ enum class EStateTreeBindableStructSource : uint8
 
 namespace UE::StateTree
 {
-	/**
-	 * Can that binding type accept a task instance data to be the source.
-	 * Condition and Utility are constructed before the task instance data is constructed.
-	 */
+	/** Can that binding type accept a task instance data for a source. */
 	[[nodiscard]] UE_API bool AcceptTaskInstanceData(EStateTreeBindableStructSource Target);
 }
 
