[AutoRTFM] Make UStateTree::GetSharedInstanceData transactionally safe.

This is reachable from Verse when using the `class_and_team_selector_device`. We now use a transactionally-safe lock, and maintain the cached thread index in an always-open block.

diff --git a/Engine/Plugins/Runtime/StateTree/Source/StateTreeModule/Private/StateTree.cpp b/Engine/Plugins/Runtime/StateTree/Source/StateTreeModule/Private/StateTree.cpp
--- a/Engine/Plugins/Runtime/StateTree/Source/StateTreeModule/Private/StateTree.cpp
+++ b/Engine/Plugins/Runtime/StateTree/Source/StateTreeModule/Private/StateTree.cpp
@@ -15,6 +15,9 @@
 #include "Misc/EnumerateRange.h"
 #include "UObject/AssetRegistryTagsContext.h"
 #include "StateTreePropertyFunctionBase.h"
+#include "AutoRTFM.h"
+
+#include <atomic>
 
 #if WITH_EDITOR
 #include "Editor.h"
@@ -99,7 +102,8 @@ FGuid UStateTree::GetTransitionIdFromIndex(const FStateTreeIndex16 Index) const
 	return Entry != nullptr ? Entry->Id : FGuid();
 }
 
-TSharedPtr<FStateTreeInstanceData> UStateTree::GetSharedInstanceData() const
+UE_AUTORTFM_ALWAYS_OPEN
+static int32 GetThreadIndexForSharedInstanceData()
 {
 	// Create a unique index for each thread.
 	static std::atomic_int ThreadIndexCounter {0};
@@ -109,9 +113,16 @@ TSharedPtr<FStateTreeInstanceData> UStateTree::GetSharedInstanceData() const
 		ThreadIndex = ThreadIndexCounter.fetch_add(1);
 	}
 
+	return ThreadIndex;
+}
+
+TSharedPtr<FStateTreeInstanceData> UStateTree::GetSharedInstanceData() const
+{
+	int32 ThreadIndex = GetThreadIndexForSharedInstanceData();
+
 	// If shared instance data for this thread exists, return it.
 	{
-		FReadScopeLock ReadLock(PerThreadSharedInstanceDataLock);
+		UE::TReadScopeLock ReadLock(PerThreadSharedInstanceDataLock);
 		if (ThreadIndex < PerThreadSharedInstanceData.Num())
 		{
 			return PerThreadSharedInstanceData[ThreadIndex];
@@ -119,7 +130,7 @@ TSharedPtr<FStateTreeInstanceData> UStateTree::GetSharedInstanceData() const
 	}
 
 	// Not initialized yet, create new instances up to the index.
-	FWriteScopeLock WriteLock(PerThreadSharedInstanceDataLock);
+	UE::TWriteScopeLock WriteLock(PerThreadSharedInstanceDataLock);
 
 	// It is possible that multiple threads are waiting for the write lock,
 	// which means that execution may get here so that 'ThreadIndex' is already in valid range.
@@ -415,7 +426,7 @@ void UStateTree::AddReferencedObjects(UObject* InThis, FReferenceCollector& Coll
 	const UStateTree* StateTree = Cast<UStateTree>(InThis);
 	check(StateTree);
 	
-	FReadScopeLock ReadLock(StateTree->PerThreadSharedInstanceDataLock);
+	UE::TReadScopeLock ReadLock(StateTree->PerThreadSharedInstanceDataLock);
 
 	for (const TSharedPtr<FStateTreeInstanceData>& InstanceData : StateTree->PerThreadSharedInstanceData)
 	{
@@ -519,7 +530,7 @@ void UStateTree::ResetLinked()
 	OutOfDateStructs.Reset();
 #endif
 
-	FWriteScopeLock WriteLock(PerThreadSharedInstanceDataLock);
+	UE::TWriteScopeLock WriteLock(PerThreadSharedInstanceDataLock);
 	PerThreadSharedInstanceData.Reset();
 }
 
diff --git a/Engine/Plugins/Runtime/StateTree/Source/StateTreeModule/Public/StateTree.h b/Engine/Plugins/Runtime/StateTree/Source/StateTreeModule/Public/StateTree.h
--- a/Engine/Plugins/Runtime/StateTree/Source/StateTreeModule/Public/StateTree.h
+++ b/Engine/Plugins/Runtime/StateTree/Source/StateTreeModule/Public/StateTree.h
@@ -7,6 +7,7 @@
 #include "StateTreePropertyBindings.h"
 #include "StateTreeInstanceData.h"
 #include "StateTreeTypes.h"
+#include "Misc/TransactionallySafeRWLock.h"
 #include "UObject/ObjectKey.h"
 #include "StateTree.generated.h"
 
@@ -321,7 +322,7 @@ private:
 	UPROPERTY()
 	FStateTreeInstanceData SharedInstanceData;
 
-	mutable FRWLock PerThreadSharedInstanceDataLock;
+	mutable FTransactionallySafeRWLock PerThreadSharedInstanceDataLock;
 	mutable TArray<TSharedPtr<FStateTreeInstanceData>> PerThreadSharedInstanceData;
 	
 	/** List of names external data enforced by the schema, created at compilation. */
