StateTree: Fixed calling StateCompleted twice on tasks that finished immediately when the tree starts.
#rnx

diff --git a/Engine/Plugins/Runtime/StateTree/Source/StateTreeModule/Private/StateTreeExecutionContext.cpp b/Engine/Plugins/Runtime/StateTree/Source/StateTreeModule/Private/StateTreeExecutionContext.cpp
--- a/Engine/Plugins/Runtime/StateTree/Source/StateTreeModule/Private/StateTreeExecutionContext.cpp
+++ b/Engine/Plugins/Runtime/StateTree/Source/StateTreeModule/Private/StateTreeExecutionContext.cpp
@@ -1432,6 +1432,7 @@ void FStateTreeExecutionContext::TickUpdateTasksInternal(float DeltaTime)
 		DelayedState.TimeLeft -= DeltaTime;
 	}
 
+	const EStateTreeRunStatus PreviousTickStatus = Exec.LastTickStatus;
 	auto LogRequestStop = [&Exec, this]()
 		{
 			if (Exec.RequestedStop != EStateTreeRunStatus::Unset) // -V547
@@ -1440,12 +1441,17 @@ void FStateTreeExecutionContext::TickUpdateTasksInternal(float DeltaTime)
 				STATETREE_LOG(Log, TEXT("Global tasks completed (%s), stopping the tree"), *UEnum::GetDisplayValueAsText(Exec.RequestedStop).ToString());
 			}
 		};
-	auto TickTaskLogic = [&Exec, &LogRequestStop, this](float DeltaTime)
+	auto TickTaskLogic = [&Exec, &LogRequestStop, PreviousTickStatus, this](float DeltaTime)
 		{
 			// Tick tasks on active states.
 			Exec.LastTickStatus = TickTasks(DeltaTime);
 			// Report state completed immediately (and no global task completes)
-			if (Exec.LastTickStatus != EStateTreeRunStatus::Running && Exec.RequestedStop == EStateTreeRunStatus::Unset)
+			if (Exec.LastTickStatus != EStateTreeRunStatus::Running && Exec.RequestedStop == EStateTreeRunStatus::Unset && PreviousTickStatus == EStateTreeRunStatus::Running)
+			{
+				StateCompleted();
+			}
+			LogRequestStop();
+
 			{
 				StateCompleted();
 			}
