[State Tree] GetInstance and GetPropertyReferenceInstance for Strong Execution Context
- moved async helper struct FActivePathInfo to public
- refactored GetDataView functionalities from ExecContext.cpp to InstanceData.cpp

diff --git a/Engine/Plugins/Runtime/StateTree/Source/StateTreeModule/Private/StateTreeAsyncExecutionContext.cpp b/Engine/Plugins/Runtime/StateTree/Source/StateTreeModule/Private/StateTreeAsyncExecutionContext.cpp
--- a/Engine/Plugins/Runtime/StateTree/Source/StateTreeModule/Private/StateTreeAsyncExecutionContext.cpp
+++ b/Engine/Plugins/Runtime/StateTree/Source/StateTreeModule/Private/StateTreeAsyncExecutionContext.cpp
@@ -2,7 +2,16 @@
 
 #include "StateTreeAsyncExecutionContext.h"
 #include "StateTreeExecutionContext.h"
-#include "StateTreeTaskBase.h"
+
+namespace UE::StateTree::Async
+{
+	const FStateTreeNodeBase& FActivePathInfo::GetNode() const
+	{
+		check(IsValid() && Frame->StateTree && NodeIndex.IsValid());
+
+		return *(Frame->StateTree->GetNode(NodeIndex.AsInt32()).GetPtr<const FStateTreeNodeBase>());
+	}
+}
 
 template <bool bWithWriteAccess>
 TStateTreeStrongExecutionContext<bWithWriteAccess>::TStateTreeStrongExecutionContext(const FStateTreeWeakExecutionContext& WeakContext)
@@ -68,7 +77,7 @@ bool TStateTreeStrongExecutionContext<bWithWriteAccess>::RequestTransition(FStat
 {
 	static_assert(bWithWriteAccess);
 
-	FActivePathInfo ActivePath = GetActivePathInfo();
+	UE::StateTree::Async::FActivePathInfo ActivePath = GetActivePathInfo();
 	if (ActivePath.IsValid())
 	{
 		FStateTreeTransitionRequest Request;
@@ -100,7 +109,7 @@ const
 		return true;
 	}
 
-	FActivePathInfo ActivePath = GetActivePathInfo();
+	UE::StateTree::Async::FActivePathInfo ActivePath = GetActivePathInfo();
 	if (ActivePath.IsValid())
 	{
 		FStateTreeExecutionState& Exec = Storage->GetMutableExecutionState();
@@ -128,7 +137,7 @@ bool TStateTreeStrongExecutionContext<bWithWriteAccess>::BindDelegate(const FSta
 		return true;
 	}
 
-	FActivePathInfo ActivePath = GetActivePathInfo();
+	UE::StateTree::Async::FActivePathInfo ActivePath = GetActivePathInfo();
 	if (ActivePath.IsValid())
 	{
 		FStateTreeExecutionState& Exec = Storage->GetMutableExecutionState();
@@ -174,41 +183,42 @@ bool TStateTreeStrongExecutionContext<bWithWriteAccess>::FinishTask(EStateTreeFi
 
 	bool bSucceed = false;
 
-	FActivePathInfo ActivePath = GetActivePathInfo();
+	UE::StateTree::Async::FActivePathInfo ActivePath = GetActivePathInfo();
 	if (ActivePath.IsValid())
 	{
 		FStateTreeExecutionState& Exec = Storage->GetMutableExecutionState();
-		if (ensure(ActivePath.Frame->StateTree))
+		FStateTreeExecutionFrame* NonConstFrame = const_cast<FStateTreeExecutionFrame*>(ActivePath.Frame);
+		if (ensure(NonConstFrame->StateTree))
 		{
 			using namespace UE::StateTree;
 
-			const UStateTree* FrameStateTree = ActivePath.Frame->StateTree;
+			const UStateTree* FrameStateTree = NonConstFrame->StateTree;
 
-			const int32 AssetNodeIndex = NodeIndex.AsInt32();
+			const int32 AssetNodeIndex = ActivePath.NodeIndex.AsInt32();
 			const int32 GlobalTaskBeginIndex = FrameStateTree->GlobalTasksBegin;
 			const int32 GlobalTaskEndIndex = FrameStateTree->GlobalTasksBegin + FrameStateTree->GlobalTasksNum;
 			const bool bIsGlobalTask = AssetNodeIndex >= GlobalTaskBeginIndex && AssetNodeIndex < GlobalTaskEndIndex;
 			const ETaskCompletionStatus TaskStatus = ExecutionContext::CastToTaskStatus(FinishType);
 
 			bool bCompleted = false;
-			if (ActivePath.Frame->bIsGlobalFrame && bIsGlobalTask)
+			if (NonConstFrame->bIsGlobalFrame && bIsGlobalTask)
 			{
 				const int32 FrameTaskIndex = AssetNodeIndex - GlobalTaskBeginIndex;
-				FTasksCompletionStatus GlobalTasksStatus = ActivePath.Frame->ActiveTasksStatus.GetStatus(ActivePath.Frame->StateTree);
+				FTasksCompletionStatus GlobalTasksStatus = NonConstFrame->ActiveTasksStatus.GetStatus(NonConstFrame->StateTree);
 				GlobalTasksStatus.SetStatusWithPriority(FrameTaskIndex, TaskStatus);
 				bCompleted = GlobalTasksStatus.IsCompleted();
 				bSucceed = true;
 			}
 			else
 			{
-				const FCompactStateTreeState& State = ActivePath.Frame->StateTree->States[ActivePath.StateHandle.Index];
+				const FCompactStateTreeState& State = NonConstFrame->StateTree->States[ActivePath.StateHandle.Index];
 				const int32 StateTaskBeginIndex = State.TasksBegin;
 				const int32 StateTaskEndIndex = State.TasksBegin + State.TasksNum;
 				const bool bIsStateTask = AssetNodeIndex >= StateTaskBeginIndex && AssetNodeIndex < StateTaskEndIndex;
 				if (bIsStateTask)
 				{
 					const int32 StateTaskIndex = AssetNodeIndex - State.TasksBegin;
-					FTasksCompletionStatus StateTasksStatus = ActivePath.Frame->ActiveTasksStatus.GetStatus(State);
+					FTasksCompletionStatus StateTasksStatus = NonConstFrame->ActiveTasksStatus.GetStatus(State);
 					StateTasksStatus.SetStatusWithPriority(StateTaskIndex, TaskStatus);
 					bCompleted = StateTasksStatus.IsCompleted();
 					bSucceed = true;
@@ -247,7 +257,7 @@ bool TStateTreeStrongExecutionContext<bWithWriteAccess>::UpdateScheduledTickRequ
 }
 
 template <bool bWithWriteAccess>
-typename TStateTreeStrongExecutionContext<bWithWriteAccess>::FActivePathInfo TStateTreeStrongExecutionContext<bWithWriteAccess>::GetActivePathInfo() const
+UE::StateTree::Async::FActivePathInfo TStateTreeStrongExecutionContext<bWithWriteAccess>::GetActivePathInfo() const
 {
 	if (IsValidInstanceStorage())
 	{
@@ -261,7 +271,7 @@ typename TStateTreeStrongExecutionContext<bWithWriteAccess>::FActivePathInfo TSt
 			if (!StateID.IsValid() || StateHandle.IsValid())
 			{
 				FStateTreeExecutionFrame* ParentFrame = FrameIndex != 0 ? &Exec.ActiveFrames[FrameIndex - 1] : nullptr;
-				return {Frame, ParentFrame, StateHandle};
+				return {Frame, ParentFrame, StateHandle, NodeIndex};
 			}
 		}
 	}
@@ -269,6 +279,28 @@ typename TStateTreeStrongExecutionContext<bWithWriteAccess>::FActivePathInfo TSt
 	return {};
 }
 
+template <bool bWithWriteAccess>
+FStateTreeDataView TStateTreeStrongExecutionContext<bWithWriteAccess>::GetInstanceDataPtrInternal() const
+{
+	using namespace UE::StateTree;
+
+	Async::FActivePathInfo ActivePath = GetActivePathInfo();
+	if (ActivePath.IsValid())
+	{
+		const FStateTreeNodeBase& Node = ActivePath.GetNode();
+		FStateTreeDataView InstanceDataView = InstanceData::GetDataViewOrTemporary(
+			*Storage,
+			/* SharedInstanceStorage */ nullptr,
+			ActivePath.ParentFrame,
+			*ActivePath.Frame,
+			ActivePath.GetNode().InstanceDataHandle);
+
+		return InstanceDataView;
+	}
+
+	return {};
+}
+
 template <bool bWithWriteAccess>
 template<bool bWriteAccess UE_REQUIRES(bWriteAccess)>
 void TStateTreeStrongExecutionContext<bWithWriteAccess>::ScheduleNextTick(TNotNull<UObject*> Owner, TNotNull<const UStateTree*> RootStateTree, FStateTreeInstanceStorage& Storage)
@@ -316,26 +348,23 @@ FStateTreeWeakExecutionContext::FStateTreeWeakExecutionContext(const FStateTreeE
 
 bool FStateTreeWeakExecutionContext::SendEvent(const FGameplayTag Tag, const FConstStructView Payload, const FName Origin) const
 {
-	constexpr bool bWithWriteAccess = true;
-	return MakeStrongExecutionContext<bWithWriteAccess>().SendEvent(Tag, Payload, Origin);
+	return MakeStrongExecutionContext().SendEvent(Tag, Payload, Origin);
 }
 
 bool FStateTreeWeakExecutionContext::RequestTransition(FStateTreeStateHandle TargetState, EStateTreeTransitionPriority Priority, const EStateTreeSelectionFallback Fallback) const
 {
-	constexpr bool bWithWriteAccess = true;
-	return MakeStrongExecutionContext<bWithWriteAccess>().RequestTransition(TargetState, Priority, Fallback);
+	return MakeStrongExecutionContext().RequestTransition(TargetState, Priority, Fallback);
 }
 
 bool FStateTreeWeakExecutionContext::BroadcastDelegate(const FStateTreeDelegateDispatcher& Dispatcher) const
 {
-	constexpr bool bWithWriteAccess = true;
-	return MakeStrongExecutionContext<bWithWriteAccess>().BroadcastDelegate(Dispatcher);
+	return MakeStrongExecutionContext().BroadcastDelegate(Dispatcher);
 }
 
 bool FStateTreeWeakExecutionContext::BindDelegate(const FStateTreeDelegateListener& Listener, FSimpleDelegate Delegate) const
 {
 	constexpr bool bWithWriteAccess = true;
-	return MakeStrongExecutionContext<bWithWriteAccess>().BindDelegate(Listener, Delegate);
+	return MakeStrongExecutionContext().BindDelegate(Listener, Delegate);
 }
 
 PRAGMA_DISABLE_DEPRECATION_WARNINGS
@@ -353,14 +382,12 @@ bool FStateTreeWeakExecutionContext::RemoveDelegateListener(const FStateTreeDele
 
 bool FStateTreeWeakExecutionContext::UnbindDelegate(const FStateTreeDelegateListener& Listener) const
 {
-	constexpr bool bWithWriteAccess = true;
-	return MakeStrongExecutionContext<bWithWriteAccess>().UnbindDelegate(Listener);
+	return MakeStrongExecutionContext().UnbindDelegate(Listener);
 }
 
 bool FStateTreeWeakExecutionContext::FinishTask(EStateTreeFinishTaskType FinishType) const
 {
-	constexpr bool bWithWriteAccess = true;
-	return MakeStrongExecutionContext<bWithWriteAccess>().FinishTask(FinishType);
+	return MakeStrongExecutionContext().FinishTask(FinishType);
 }
 
 PRAGMA_DISABLE_DEPRECATION_WARNINGS
@@ -372,6 +399,5 @@ PRAGMA_ENABLE_DEPRECATION_WARNINGS
 
 bool FStateTreeWeakExecutionContext::UpdateScheduledTickRequest(UE::StateTree::FScheduledTickHandle Handle, FStateTreeScheduledTick ScheduledTick) const
 {
-	constexpr bool bWithWriteAccess = true;
-	return MakeStrongExecutionContext<bWithWriteAccess>().UpdateScheduledTickRequest(Handle, ScheduledTick);
+	return MakeStrongExecutionContext().UpdateScheduledTickRequest(Handle, ScheduledTick);
 }
diff --git a/Engine/Plugins/Runtime/StateTree/Source/StateTreeModule/Private/StateTreeExecutionContext.cpp b/Engine/Plugins/Runtime/StateTree/Source/StateTreeModule/Private/StateTreeExecutionContext.cpp
--- a/Engine/Plugins/Runtime/StateTree/Source/StateTreeModule/Private/StateTreeExecutionContext.cpp
+++ b/Engine/Plugins/Runtime/StateTree/Source/StateTreeModule/Private/StateTreeExecutionContext.cpp
@@ -2119,12 +2119,20 @@ FStateTreeDataView FStateTreeExecutionContext::GetDataView(const FStateTreeExecu
 				return {};
 			}
 
-			return GetDataViewFromInstanceStorage(Storage, CurrentlyProcessedSharedInstanceStorage, ParentFrame, CurrentFrame, Handle);
+			return UE::StateTree::InstanceData::GetDataView(Storage, CurrentlyProcessedSharedInstanceStorage, ParentFrame, CurrentFrame, Handle);
+		}
+
+	case EStateTreeDataSourceType::ExternalGlobalParameterData:
+		{
+			checkf(false, TEXT("External global parameter data currently not supported for linked state-trees"));
+			break;
 		}
 
 	default:
-		return GetDataViewFromInstanceStorage(Storage, CurrentlyProcessedSharedInstanceStorage, ParentFrame, CurrentFrame, Handle);
+		return UE::StateTree::InstanceData::GetDataView(Storage, CurrentlyProcessedSharedInstanceStorage, ParentFrame, CurrentFrame, Handle);
 	}
+
+	return {};
 }
 
 FStateTreeDataView FStateTreeExecutionContext::GetDataView(const FStateTreeExecutionFrame* ParentFrame, const FStateTreeExecutionFrame& CurrentFrame, const FPropertyBindingCopyInfo& CopyInfo)
@@ -2164,87 +2172,6 @@ EStateTreeRunStatus FStateTreeExecutionContext::ForceTransition(const FRecordedS
 	return EnterState(TransitionResult.GetValue());
 }
 
-FStateTreeDataView FStateTreeExecutionContext::GetDataViewFromInstanceStorage(FStateTreeInstanceStorage& InstanceDataStorage, FStateTreeInstanceStorage* CurrentlyProcessedSharedInstanceStorage, const FStateTreeExecutionFrame* ParentFrame, const FStateTreeExecutionFrame& CurrentFrame, const FStateTreeDataHandle Handle)
-{
-	switch (Handle.GetSource())
-	{
-	case EStateTreeDataSourceType::None:
-		return {};
-
-	case EStateTreeDataSourceType::GlobalInstanceData:
-		return InstanceDataStorage.GetMutableStruct(CurrentFrame.GlobalInstanceIndexBase.Get() + Handle.GetIndex());
-	case EStateTreeDataSourceType::GlobalInstanceDataObject:
-		return InstanceDataStorage.GetMutableObject(CurrentFrame.GlobalInstanceIndexBase.Get() + Handle.GetIndex());
-		
-	case EStateTreeDataSourceType::ActiveInstanceData:
-		return InstanceDataStorage.GetMutableStruct(CurrentFrame.ActiveInstanceIndexBase.Get() + Handle.GetIndex());
-	case EStateTreeDataSourceType::ActiveInstanceDataObject:
-		return InstanceDataStorage.GetMutableObject(CurrentFrame.ActiveInstanceIndexBase.Get() + Handle.GetIndex());
-
-	case EStateTreeDataSourceType::SharedInstanceData:
-		check(CurrentlyProcessedSharedInstanceStorage);
-		return CurrentlyProcessedSharedInstanceStorage->GetMutableStruct(Handle.GetIndex());
-	case EStateTreeDataSourceType::SharedInstanceDataObject:
-		check(CurrentlyProcessedSharedInstanceStorage);
-		return CurrentlyProcessedSharedInstanceStorage->GetMutableObject(Handle.GetIndex());
-
-	case EStateTreeDataSourceType::GlobalParameterData:
-		// Defined in parent frame or is root state tree parameters
-		if (ParentFrame)
-		{
-			return GetDataViewFromInstanceStorage(InstanceDataStorage, CurrentlyProcessedSharedInstanceStorage, nullptr, *ParentFrame, CurrentFrame.GlobalParameterDataHandle);
-		}
-		return InstanceDataStorage.GetMutableGlobalParameters();
-
-	case EStateTreeDataSourceType::ExternalGlobalParameterData:
-	    {		    
-			checkf(false, TEXT("External global parameter data currently not supported for linked state-trees"));
-	    }
-	    break;
-
-	case EStateTreeDataSourceType::SubtreeParameterData:
-		{
-			// Defined in parent frame.
-			if (ParentFrame)
-			{
-				// Linked subtree, params defined in parent scope.
-				return GetDataViewFromInstanceStorage(InstanceDataStorage, CurrentlyProcessedSharedInstanceStorage, nullptr, *ParentFrame, CurrentFrame.StateParameterDataHandle);
-			}
-			// Standalone subtree, params define as state params.
-			FCompactStateTreeParameters& SubtreeParams = InstanceDataStorage.GetMutableStruct(CurrentFrame.ActiveInstanceIndexBase.Get() + Handle.GetIndex()).Get<FCompactStateTreeParameters>();
-			return SubtreeParams.Parameters.GetMutableValue();
-		}
-
-	case EStateTreeDataSourceType::StateParameterData:
-		{
-			FCompactStateTreeParameters& StateParams = InstanceDataStorage.GetMutableStruct(CurrentFrame.ActiveInstanceIndexBase.Get() + Handle.GetIndex()).Get<FCompactStateTreeParameters>();
-			return StateParams.Parameters.GetMutableValue();
-		}
-
-	case EStateTreeDataSourceType::StateEvent:
-		{
-			// Return FStateTreeEvent from shared event.
-			FStateTreeSharedEvent& SharedEvent = InstanceDataStorage.GetMutableStruct(CurrentFrame.ActiveInstanceIndexBase.Get() + Handle.GetIndex()).Get<FStateTreeSharedEvent>();
-			if (ensure(SharedEvent.IsValid()))
-			{
-				// Events are read only, but we cannot express that in FStateTreeDataView.
-				return FStateTreeDataView(FStructView::Make(*SharedEvent.GetMutable()));
-			}
-			return {};
-		}
-
-	case EStateTreeDataSourceType::ContextData:
-	case EStateTreeDataSourceType::ExternalData:
-	case EStateTreeDataSourceType::TransitionEvent:
-		return {};
-	
-	default:
-		checkf(false, TEXT("Unhandle case %s"), *UEnum::GetValueAsString(Handle.GetSource()));
-	}
-
-	return {};
-}
-
 const FStateTreeExecutionFrame* FStateTreeExecutionContext::FindFrame(const UStateTree* StateTree, FStateTreeStateHandle RootState, TConstArrayView<FStateTreeExecutionFrame> Frames, const FStateTreeExecutionFrame*& OutParentFrame)
 {
 	const int32 FrameIndex = Frames.IndexOfByPredicate([&StateTree, RootState](const FStateTreeExecutionFrame& Frame)
@@ -2268,36 +2195,11 @@ const FStateTreeExecutionFrame* FStateTreeExecutionContext::FindFrame(const USta
 
 bool FStateTreeExecutionContext::IsHandleSourceValid(const FStateTreeExecutionFrame* ParentFrame, const FStateTreeExecutionFrame& CurrentFrame, const FStateTreeDataHandle Handle) const
 {
-	// Checks that the instance data is valid for specific handle types.
-	// 
-	// The CurrentFrame may not be yet properly initialized, for that reason we need to check
-	// that the path to the handle makes sense (it's part of the active states) as well as that
-	// we actually have instance data for the handle (index is valid).
-	// 
-	// The (base) indices can be invalid if the frame/state is not entered yet.
-	// For active instance data we need to check that the frame is initialized for a specific state,
-	// as well as that the instance data is initialized.
-
 	switch (Handle.GetSource())
 	{
 	case EStateTreeDataSourceType::None:
 		return true;
 
-	case EStateTreeDataSourceType::GlobalInstanceData:
-	case EStateTreeDataSourceType::GlobalInstanceDataObject:
-		return CurrentFrame.GlobalInstanceIndexBase.IsValid()
-			&& Storage.IsValidIndex(CurrentFrame.GlobalInstanceIndexBase.Get() + Handle.GetIndex());
-
-	case EStateTreeDataSourceType::ActiveInstanceData:
-	case EStateTreeDataSourceType::ActiveInstanceDataObject:
-		return CurrentFrame.ActiveInstanceIndexBase.IsValid()
-			&& CurrentFrame.ActiveStates.Contains(Handle.GetState(), CurrentFrame.NumCurrentlyActiveStates)
-			&& Storage.IsValidIndex(CurrentFrame.ActiveInstanceIndexBase.Get() + Handle.GetIndex());
-		
-	case EStateTreeDataSourceType::SharedInstanceData:
-	case EStateTreeDataSourceType::SharedInstanceDataObject:
-		return true;
-
 	case EStateTreeDataSourceType::ContextData:
 		return true;
 
@@ -2305,37 +2207,6 @@ bool FStateTreeExecutionContext::IsHandleSourceValid(const FStateTreeExecutionFr
 		return CurrentFrame.ExternalDataBaseIndex.IsValid()
 			&& ContextAndExternalDataViews.IsValidIndex(CurrentFrame.ExternalDataBaseIndex.Get() + Handle.GetIndex());
 
-	case EStateTreeDataSourceType::GlobalParameterData:
-		return ParentFrame
-			? IsHandleSourceValid(nullptr, *ParentFrame, CurrentFrame.GlobalParameterDataHandle)
-			: CurrentFrame.GlobalParameterDataHandle.IsValid();
-
-	case EStateTreeDataSourceType::ExternalGlobalParameterData:
-		{
-			checkf(false, TEXT("External global parameter data currently not supported for linked state-trees"));
-		}
-		break;
-	case EStateTreeDataSourceType::SubtreeParameterData:
-		if (ParentFrame)
-		{
-			// If the current subtree state is not instantiated yet, we cannot assume that the parameter data is instantiated in the parent frame either. 
-			if (!CurrentFrame.ActiveInstanceIndexBase.IsValid())
-			{
-				return false;
-			}
-			// Linked subtree, params defined in parent scope.
-			return IsHandleSourceValid(nullptr, *ParentFrame, CurrentFrame.StateParameterDataHandle);
-		}
-		// Standalone subtree, params define as state params.
-		return CurrentFrame.ActiveInstanceIndexBase.IsValid()
-			&& CurrentFrame.ActiveStates.Contains(Handle.GetState(), CurrentFrame.NumCurrentlyActiveStates)
-			&& Storage.IsValidIndex(CurrentFrame.ActiveInstanceIndexBase.Get() + Handle.GetIndex());
-
-	case EStateTreeDataSourceType::StateParameterData:
-		return CurrentFrame.ActiveInstanceIndexBase.IsValid()
-			&& CurrentFrame.ActiveStates.Contains(Handle.GetState(), CurrentFrame.NumCurrentlyActiveStates)
-			&& Storage.IsValidIndex(CurrentFrame.ActiveInstanceIndexBase.Get() + Handle.GetIndex());
-
 	case EStateTreeDataSourceType::TransitionEvent:
 		return CurrentlyProcessedTransitionEvent != nullptr;
 
@@ -2345,8 +2216,14 @@ bool FStateTreeExecutionContext::IsHandleSourceValid(const FStateTreeExecutionFr
 			&& CurrentFrame.ActiveStates.Contains(Handle.GetState())
 			&& Storage.IsValidIndex(CurrentFrame.ActiveInstanceIndexBase.Get() + Handle.GetIndex()));
 
+	case EStateTreeDataSourceType::ExternalGlobalParameterData:
+	{
+		checkf(false, TEXT("External global parameter data currently not supported for linked state-trees"));
+		break;
+	}
+
 	default:
-		checkf(false, TEXT("Unhandle case %s"), *UEnum::GetValueAsString(Handle.GetSource()));
+		return UE::StateTree::InstanceData::Internal::IsHandleSourceValid(Storage, ParentFrame, CurrentFrame, Handle);
 	}
 
 	return false;
@@ -2370,60 +2247,7 @@ FStateTreeDataView FStateTreeExecutionContext::GetDataViewOrTemporary(const FSta
 		return GetDataView(ParentFrame, CurrentFrame, Handle);
 	}
 	
-	switch (Handle.GetSource())
-	{
-	case EStateTreeDataSourceType::GlobalInstanceData:
-	case EStateTreeDataSourceType::ActiveInstanceData:
-		return Storage.GetMutableTemporaryStruct(CurrentFrame, Handle);
-		
-	case EStateTreeDataSourceType::GlobalInstanceDataObject:
-	case EStateTreeDataSourceType::ActiveInstanceDataObject:
-		return Storage.GetMutableTemporaryObject(CurrentFrame, Handle);
-
-	case EStateTreeDataSourceType::GlobalParameterData:
-		if (ParentFrame)
-		{
-			if (FCompactStateTreeParameters* Params = Storage.GetMutableTemporaryStruct(*ParentFrame, CurrentFrame.GlobalParameterDataHandle).GetPtr<FCompactStateTreeParameters>())
-			{
-				return Params->Parameters.GetMutableValue();
-			}
-		}
-		break;
-		
-	case EStateTreeDataSourceType::ExternalGlobalParameterData:
-	    {
-			checkf(false, TEXT("External global parameter data currently not supported for linked state-trees"));
-	    }
-		break;
-		
-	case EStateTreeDataSourceType::SubtreeParameterData:
-		if (ParentFrame)
-		{
-			// Linked subtree, params defined in parent scope.
-			if (FCompactStateTreeParameters* Params = Storage.GetMutableTemporaryStruct(*ParentFrame, CurrentFrame.StateParameterDataHandle).GetPtr<FCompactStateTreeParameters>())
-			{
-				return Params->Parameters.GetMutableValue();
-			}
-		}
-		// Standalone subtree, params define as state params.
-		if (FCompactStateTreeParameters* Params = Storage.GetMutableTemporaryStruct(CurrentFrame, Handle).GetPtr<FCompactStateTreeParameters>())
-		{
-			return Params->Parameters.GetMutableValue();
-		}
-		break;
-
-	case EStateTreeDataSourceType::StateParameterData:
-		if (FCompactStateTreeParameters* Params = Storage.GetMutableTemporaryStruct(CurrentFrame, Handle).GetPtr<FCompactStateTreeParameters>())
-		{
-			return Params->Parameters.GetMutableValue();
-		}
-		break;
-
-	default:
-		return {};
-	}
-
-	return {};
+	return GetTemporaryDataView(ParentFrame, CurrentFrame, Handle);
 }
 
 FStateTreeDataView FStateTreeExecutionContext::GetDataViewOrTemporary(const FStateTreeExecutionFrame* ParentFrame, const FStateTreeExecutionFrame& CurrentFrame, const FPropertyBindingCopyInfo& CopyInfo)
@@ -2438,6 +2262,21 @@ FStateTreeDataView FStateTreeExecutionContext::GetDataViewOrTemporary(const FSta
 	return GetDataViewOrTemporary(ParentFrame, CurrentFrame, Handle);	
 }
 
+FStateTreeDataView FStateTreeExecutionContext::GetTemporaryDataView(const FStateTreeExecutionFrame* ParentFrame,
+	const FStateTreeExecutionFrame& CurrentFrame, const FStateTreeDataHandle Handle)
+{
+	switch (Handle.GetSource())
+	{
+	case EStateTreeDataSourceType::ExternalGlobalParameterData:
+		checkf(false, TEXT("External global parameter data currently not supported for linked state-trees"));
+		return {};
+
+	default:
+		return UE::StateTree::InstanceData::Internal::GetTemporaryDataView(Storage, ParentFrame, CurrentFrame, Handle);
+	}
+
+}
+
 FStateTreeDataView FStateTreeExecutionContext::AddTemporaryInstance(const FStateTreeExecutionFrame& Frame, const FStateTreeIndex16 OwnerNodeIndex, const FStateTreeDataHandle DataHandle, FConstStructView NewInstanceData)
 {
 	const FStructView NewInstance = Storage.AddTemporaryInstance(Owner, Frame, OwnerNodeIndex, DataHandle, NewInstanceData);
diff --git a/Engine/Plugins/Runtime/StateTree/Source/StateTreeModule/Private/StateTreeInstanceData.cpp b/Engine/Plugins/Runtime/StateTree/Source/StateTreeModule/Private/StateTreeInstanceData.cpp
--- a/Engine/Plugins/Runtime/StateTree/Source/StateTreeModule/Private/StateTreeInstanceData.cpp
+++ b/Engine/Plugins/Runtime/StateTree/Source/StateTreeModule/Private/StateTreeInstanceData.cpp
@@ -97,6 +97,218 @@ FStateTreeInstanceStorage::FStateTreeInstanceStorage()
 FStateTreeInstanceStorage::FStateTreeInstanceStorage() = default;
 #endif
 
+namespace UE::StateTree::InstanceData
+{
+	namespace Internal
+	{
+		bool IsHandleSourceValid(
+			const FStateTreeInstanceStorage& InstanceStorage,
+			const FStateTreeExecutionFrame* ParentFrame,
+			const FStateTreeExecutionFrame& CurrentFrame,
+			const FStateTreeDataHandle Handle)
+		{
+			// Checks that the instance data is valid for specific handle types.
+			// 
+			// The CurrentFrame may not be yet properly initialized, for that reason we need to check
+			// that the path to the handle makes sense (it's part of the active states) as well as that
+			// we actually have instance data for the handle (index is valid).
+			// 
+			// The (base) indices can be invalid if the frame/state is not entered yet.
+			// For active instance data we need to check that the frame is initialized for a specific state,
+			// as well as that the instance data is initialized.
+
+			switch (Handle.GetSource())
+			{
+			case EStateTreeDataSourceType::None:
+				return true;
+
+			case EStateTreeDataSourceType::GlobalInstanceData:
+			case EStateTreeDataSourceType::GlobalInstanceDataObject:
+				return CurrentFrame.GlobalInstanceIndexBase.IsValid()
+					&& InstanceStorage.IsValidIndex(CurrentFrame.GlobalInstanceIndexBase.Get() + Handle.GetIndex());
+
+			case EStateTreeDataSourceType::ActiveInstanceData:
+			case EStateTreeDataSourceType::ActiveInstanceDataObject:
+			case EStateTreeDataSourceType::StateParameterData:
+				return CurrentFrame.ActiveInstanceIndexBase.IsValid()
+					&& CurrentFrame.ActiveStates.Contains(Handle.GetState(), CurrentFrame.NumCurrentlyActiveStates)
+					&& InstanceStorage.IsValidIndex(CurrentFrame.ActiveInstanceIndexBase.Get() + Handle.GetIndex());
+
+			case EStateTreeDataSourceType::SharedInstanceData:
+			case EStateTreeDataSourceType::SharedInstanceDataObject:
+				return true;
+
+			case EStateTreeDataSourceType::GlobalParameterData:
+				return ParentFrame
+					? IsHandleSourceValid(InstanceStorage, nullptr, *ParentFrame, CurrentFrame.GlobalParameterDataHandle)
+					: CurrentFrame.GlobalParameterDataHandle.IsValid();
+
+			case EStateTreeDataSourceType::SubtreeParameterData:
+				if (ParentFrame)
+				{
+					// If the current subtree state is not instantiated yet, we cannot assume that the parameter data is instantiated in the parent frame either. 
+					if (!CurrentFrame.ActiveInstanceIndexBase.IsValid())
+					{
+						return false;
+					}
+					// Linked subtree, params defined in parent scope.
+					return IsHandleSourceValid(InstanceStorage, nullptr, *ParentFrame, CurrentFrame.StateParameterDataHandle);
+				}
+				// Standalone subtree, params define as state params.
+				return CurrentFrame.ActiveInstanceIndexBase.IsValid()
+					&& CurrentFrame.ActiveStates.Contains(Handle.GetState(), CurrentFrame.NumCurrentlyActiveStates)
+					&& InstanceStorage.IsValidIndex(CurrentFrame.ActiveInstanceIndexBase.Get() + Handle.GetIndex());
+
+			default:
+				checkf(false, TEXT("Unhandled case or unsupported type for InstanceDataStorage %s"), *UEnum::GetValueAsString(Handle.GetSource()));
+			}
+
+			return false;
+		}
+
+		FStateTreeDataView GetTemporaryDataView(
+			FStateTreeInstanceStorage& InstanceStorage, 
+			const FStateTreeExecutionFrame* ParentFrame, 
+			const FStateTreeExecutionFrame& CurrentFrame, 
+			const FStateTreeDataHandle Handle)
+		{
+			switch (Handle.GetSource())
+			{
+			case EStateTreeDataSourceType::GlobalInstanceData:
+			case EStateTreeDataSourceType::ActiveInstanceData:
+				return InstanceStorage.GetMutableTemporaryStruct(CurrentFrame, Handle);
+
+			case EStateTreeDataSourceType::GlobalInstanceDataObject:
+			case EStateTreeDataSourceType::ActiveInstanceDataObject:
+				return InstanceStorage.GetMutableTemporaryObject(CurrentFrame, Handle);
+
+			case EStateTreeDataSourceType::GlobalParameterData:
+				if (ParentFrame)
+				{
+					if (FCompactStateTreeParameters* Params = InstanceStorage.GetMutableTemporaryStruct(*ParentFrame, CurrentFrame.GlobalParameterDataHandle).GetPtr<FCompactStateTreeParameters>())
+					{
+						return Params->Parameters.GetMutableValue();
+					}
+				}
+				break;
+
+			case EStateTreeDataSourceType::SubtreeParameterData:
+				if (ParentFrame)
+				{
+					// Linked subtree, params defined in parent scope.
+					if (FCompactStateTreeParameters* Params = InstanceStorage.GetMutableTemporaryStruct(*ParentFrame, CurrentFrame.StateParameterDataHandle).GetPtr<FCompactStateTreeParameters>())
+					{
+						return Params->Parameters.GetMutableValue();
+					}
+				}
+				// Standalone subtree, params define as state params.
+				if (FCompactStateTreeParameters* Params = InstanceStorage.GetMutableTemporaryStruct(CurrentFrame, Handle).GetPtr<FCompactStateTreeParameters>())
+				{
+					return Params->Parameters.GetMutableValue();
+				}
+				break;
+
+			case EStateTreeDataSourceType::StateParameterData:
+				if (FCompactStateTreeParameters* Params = InstanceStorage.GetMutableTemporaryStruct(CurrentFrame, Handle).GetPtr<FCompactStateTreeParameters>())
+				{
+					return Params->Parameters.GetMutableValue();
+				}
+				break;
+
+			default:
+				checkf(false, TEXT("Unhandled case or unsupported type for InstanceDataStorage %s"), *UEnum::GetValueAsString(Handle.GetSource()));
+			}
+
+			return {};
+		}
+
+	}
+
+	FStateTreeDataView GetDataView(
+		FStateTreeInstanceStorage& InstanceStorage, 
+		FStateTreeInstanceStorage* SharedInstanceStorage,
+		const FStateTreeExecutionFrame* ParentFrame, 
+		const FStateTreeExecutionFrame& CurrentFrame, 
+		const FStateTreeDataHandle Handle)
+	{
+		switch (Handle.GetSource())
+		{
+		case EStateTreeDataSourceType::GlobalInstanceData:
+			return InstanceStorage.GetMutableStruct(CurrentFrame.GlobalInstanceIndexBase.Get() + Handle.GetIndex());
+		case EStateTreeDataSourceType::GlobalInstanceDataObject:
+			return InstanceStorage.GetMutableObject(CurrentFrame.GlobalInstanceIndexBase.Get() + Handle.GetIndex());
+
+		case EStateTreeDataSourceType::ActiveInstanceData:
+			return InstanceStorage.GetMutableStruct(CurrentFrame.ActiveInstanceIndexBase.Get() + Handle.GetIndex());
+		case EStateTreeDataSourceType::ActiveInstanceDataObject:
+			return InstanceStorage.GetMutableObject(CurrentFrame.ActiveInstanceIndexBase.Get() + Handle.GetIndex());
+
+		case EStateTreeDataSourceType::SharedInstanceData:
+			check(SharedInstanceStorage);
+			return SharedInstanceStorage->GetMutableStruct(Handle.GetIndex());
+		case EStateTreeDataSourceType::SharedInstanceDataObject:
+			check(SharedInstanceStorage);
+			return SharedInstanceStorage->GetMutableObject(Handle.GetIndex());
+
+		case EStateTreeDataSourceType::GlobalParameterData:
+			// Defined in parent frame or is root state tree parameters
+			if (ParentFrame)
+			{
+				return GetDataView(InstanceStorage, SharedInstanceStorage, nullptr, *ParentFrame, CurrentFrame.GlobalParameterDataHandle);
+			}
+			return InstanceStorage.GetMutableGlobalParameters();
+
+		case EStateTreeDataSourceType::SubtreeParameterData:
+		{
+			// Defined in parent frame.
+			if (ParentFrame)
+			{
+				// Linked subtree, params defined in parent scope.
+				return GetDataView(InstanceStorage, SharedInstanceStorage, nullptr, *ParentFrame, CurrentFrame.StateParameterDataHandle);
+			}
+			// Standalone subtree, params define as state params.
+			FCompactStateTreeParameters& SubtreeParams = InstanceStorage.GetMutableStruct(CurrentFrame.ActiveInstanceIndexBase.Get() + Handle.GetIndex()).Get<FCompactStateTreeParameters>();
+			return SubtreeParams.Parameters.GetMutableValue();
+		}
+
+		case EStateTreeDataSourceType::StateParameterData:
+		{
+			FCompactStateTreeParameters& StateParams = InstanceStorage.GetMutableStruct(CurrentFrame.ActiveInstanceIndexBase.Get() + Handle.GetIndex()).Get<FCompactStateTreeParameters>();
+			return StateParams.Parameters.GetMutableValue();
+		}
+
+		case EStateTreeDataSourceType::StateEvent:
+		{
+			// Return FStateTreeEvent from shared event.
+			FStateTreeSharedEvent& SharedEvent = InstanceStorage.GetMutableStruct(CurrentFrame.ActiveInstanceIndexBase.Get() + Handle.GetIndex()).Get<FStateTreeSharedEvent>();
+			if (ensure(SharedEvent.IsValid()))
+			{
+				// Events are read only, but we cannot express that in FStateTreeDataView.
+				return FStateTreeDataView(FStructView::Make(*SharedEvent.GetMutable()));
+			}
+			return {};
+		}
+
+		default:
+			checkf(false, TEXT("Unhandled case or unsupported type for InstanceDataStorage %s"), *UEnum::GetValueAsString(Handle.GetSource()));
+		}
+
+		return {};
+	}
+
+	FStateTreeDataView GetDataViewOrTemporary(FStateTreeInstanceStorage& InstanceStorage, FStateTreeInstanceStorage* SharedInstanceStorage,
+		const FStateTreeExecutionFrame* ParentFrame, const FStateTreeExecutionFrame& CurrentFrame, const FStateTreeDataHandle Handle)
+	{
+		using namespace Internal;
+		if (IsHandleSourceValid(InstanceStorage, ParentFrame, CurrentFrame, Handle))
+		{
+			return GetDataView(InstanceStorage, SharedInstanceStorage, ParentFrame, CurrentFrame, Handle);
+		}
+
+		return GetTemporaryDataView(InstanceStorage, ParentFrame, CurrentFrame, Handle);
+	}
+}
+
 FStateTreeInstanceStorage::FStateTreeInstanceStorage(const FStateTreeInstanceStorage& Other)
 	: InstanceStructs(Other.InstanceStructs)
 	, ExecutionState(Other.ExecutionState)
diff --git a/Engine/Plugins/Runtime/StateTree/Source/StateTreeModule/Public/StateTreeAsyncExecutionContext.h b/Engine/Plugins/Runtime/StateTree/Source/StateTreeModule/Public/StateTreeAsyncExecutionContext.h
--- a/Engine/Plugins/Runtime/StateTree/Source/StateTreeModule/Public/StateTreeAsyncExecutionContext.h
+++ b/Engine/Plugins/Runtime/StateTree/Source/StateTreeModule/Public/StateTreeAsyncExecutionContext.h
@@ -2,6 +2,7 @@
 
 #pragma once
 #include "StateTree.h"
+#include "StateTreeInstanceData.h"
 
 #define UE_API STATETREEMODULE_API
 
@@ -21,6 +22,25 @@ struct FStateTreeWeakExecutionContext;
 
 enum class EStateTreeFinishTaskType : uint8;
 
+
+namespace UE::StateTree::Async
+{
+	struct FActivePathInfo
+	{
+		bool IsValid() const
+		{
+			return Frame != nullptr;
+		}
+
+		const FStateTreeNodeBase& GetNode() const;
+
+		const FStateTreeExecutionFrame* Frame = nullptr;
+		const FStateTreeExecutionFrame* ParentFrame = nullptr;
+		const FStateTreeStateHandle StateHandle;
+		const FStateTreeIndex16 NodeIndex;
+	};
+}
+
 /**
  * Execution context to interact with the state tree instance data asynchronously.
  * It should only be allocated on the stack.
@@ -57,6 +77,12 @@ struct TStateTreeStrongExecutionContext
 		return StateTree;
 	}
 
+	/** @return the Instance Storage. */
+	TSharedPtr<const FStateTreeInstanceStorage> GetStorage() const
+	{
+		return Storage;
+	}
+
 	//@todo: Replace UE_REQUIRES with requires and UE_API once past 5.6 for c++20 support
 
 	/**
@@ -119,6 +145,28 @@ struct TStateTreeStrongExecutionContext
 	template<bool bWriteAccess = bWithWriteAccess UE_REQUIRES(bWriteAccess)>
 	bool UpdateScheduledTickRequest(UE::StateTree::FScheduledTickHandle Handle, FStateTreeScheduledTick ScheduledTick) const;
 
+	/**
+	 * Get the Instance Data of recorded node.
+	 * @return nullptr if the frame or state containing the node is no longer active.
+	 */
+	template<typename T>
+	std::conditional_t<bWithWriteAccess, T*, const T*> GetInstanceDataPtr() const
+	{
+		FStateTreeDataView DataView = GetInstanceDataPtrInternal();
+		if (DataView.IsValid() && ensure(DataView.GetStruct() == T::StaticStruct()))
+		{
+			return static_cast<T*>(DataView.GetMutableMemory());
+		}
+
+		return nullptr;
+	}
+
+	/**
+	 * Get the info of active frame, state and node this Context is based on
+	 * @return an invalid FActivePathInfo if the frame or state containing the node is no longer active.
+	 */
+	UE_API UE::StateTree::Async::FActivePathInfo GetActivePathInfo() const;
+
 	/**
 	 * Checks if the context is valid.
 	 * Validity: Pinned Members are valid AND (WeakContext is created outside the ExecContext loop OR the recorded frame and state are still active)    
@@ -136,19 +184,7 @@ private:
 		return Owner && StateTree && Storage;
 	}
 
-	struct FActivePathInfo
-	{
-		bool IsValid() const
-		{
-			return Frame != nullptr; 
-		}
-
-		FStateTreeExecutionFrame* Frame = nullptr;
-		FStateTreeExecutionFrame* ParentFrame = nullptr;
-		FStateTreeStateHandle StateHandle;
-	};
-
-	UE_API FActivePathInfo GetActivePathInfo() const;
+	UE_API FStateTreeDataView GetInstanceDataPtrInternal() const;
 
 	template<bool bWriteAccess = bWithWriteAccess UE_REQUIRES(bWriteAccess)>
 	static void ScheduleNextTick(TNotNull<UObject*> Owner, TNotNull<const UStateTree*> RootStateTree, FStateTreeInstanceStorage& Storage);
@@ -161,10 +197,12 @@ private:
 	UE::StateTree::FActiveStateID StateID;
 	FStateTreeIndex16 NodeIndex;
 	uint8 bAccessAcquired : 1 = false;
+
+	friend struct FStateTreePropertyRef;
 };
 
 using FStateTreeStrongExecutionContext = TStateTreeStrongExecutionContext<true>;
-using FStateTreeReadOnlyStrongExecutionContext = TStateTreeStrongExecutionContext<false>;
+using FStateTreeStrongReadOnlyExecutionContext = TStateTreeStrongExecutionContext<false>;
 
 /**
  * Execution context that can be saved/copied and used asynchronously. 
@@ -203,10 +241,14 @@ public:
 		return StateTree.Pin();
 	}
 
-	template<bool bRequireWriteAccess>
-	TStateTreeStrongExecutionContext<bRequireWriteAccess> MakeStrongExecutionContext() const
+	FStateTreeStrongReadOnlyExecutionContext MakeStrongReadOnlyExecutionContext() const
+	{
+		return FStateTreeStrongReadOnlyExecutionContext(*this);
+	}
+
+	FStateTreeStrongExecutionContext MakeStrongExecutionContext() const
 	{
-		return TStateTreeStrongExecutionContext<bRequireWriteAccess>(*this);
+		return FStateTreeStrongExecutionContext(*this);
 	}
 
 	/**
diff --git a/Engine/Plugins/Runtime/StateTree/Source/StateTreeModule/Public/StateTreeExecutionContext.h b/Engine/Plugins/Runtime/StateTree/Source/StateTreeModule/Public/StateTreeExecutionContext.h
--- a/Engine/Plugins/Runtime/StateTree/Source/StateTreeModule/Public/StateTreeExecutionContext.h
+++ b/Engine/Plugins/Runtime/StateTree/Source/StateTreeModule/Public/StateTreeExecutionContext.h
@@ -687,18 +687,22 @@ PRAGMA_DISABLE_DEPRECATION_WARNINGS
 	 * If called outside of the tick processing, then the request is buffered and handled on the next tick.
 	 */
 	UE_API void FinishTask(const UE::StateTree::FFinishedTask& Task, EStateTreeFinishTaskType FinishType);
-UE_API PRAGMA_ENABLE_DEPRECATION_WARNINGS
+PRAGMA_ENABLE_DEPRECATION_WARNINGS
 
 	/** @return data view of the specified handle relative to given frame. */
-	static FStateTreeDataView GetDataViewFromInstanceStorage(FStateTreeInstanceStorage& InstanceDataStorage, FStateTreeInstanceStorage* CurrentlyProcessedSharedInstanceStorage, const FStateTreeExecutionFrame* ParentFrame, const FStateTreeExecutionFrame& CurrentFrame, const FStateTreeDataHandle Handle);
+	UE_DEPRECATED(5.6, "Use UE::StateTree::InstanceData::GetDataView instead.")
+	static FStateTreeDataView GetDataViewFromInstanceStorage(FStateTreeInstanceStorage& InstanceDataStorage, FStateTreeInstanceStorage* CurrentlyProcessedSharedInstanceStorage, const FStateTreeExecutionFrame* ParentFrame, const FStateTreeExecutionFrame& CurrentFrame, const FStateTreeDataHandle Handle)
+	{
+		return UE::StateTree::InstanceData::GetDataView(InstanceDataStorage, CurrentlyProcessedSharedInstanceStorage, ParentFrame, CurrentFrame, Handle);
+	}
 
 	/** Looks for a frame in provided list of frames. */
 	static UE_API const FStateTreeExecutionFrame* FindFrame(const UStateTree* StateTree, FStateTreeStateHandle RootState, TConstArrayView<FStateTreeExecutionFrame> Frames, const FStateTreeExecutionFrame*& OutParentFrame);
 
-	UE_DEPRECATED(5.5, "Use GetDataViewFromInstanceStorage instead.")
+	UE_DEPRECATED(5.5, "Use UE::StateTree::InstanceData::GetDataView instead.")
 	static FStateTreeDataView GetDataView(FStateTreeInstanceStorage& InstanceDataStorage, FStateTreeInstanceStorage* CurrentlyProcessedSharedInstanceStorage, const FStateTreeExecutionFrame* ParentFrame, const FStateTreeExecutionFrame& CurrentFrame, TConstArrayView<FStateTreeDataView> ContextAndExternalDataViews, const FStateTreeDataHandle Handle)
 	{
-		return GetDataViewFromInstanceStorage(InstanceDataStorage, CurrentlyProcessedSharedInstanceStorage, ParentFrame, CurrentFrame, Handle);
+		return UE::StateTree::InstanceData::GetDataView(InstanceDataStorage, CurrentlyProcessedSharedInstanceStorage, ParentFrame, CurrentFrame, Handle);
 	}
 
 	/**
@@ -1028,6 +1032,9 @@ PRAGMA_ENABLE_DEPRECATION_WARNINGS
 	UE_API FStateTreeDataView GetDataViewOrTemporary(const FStateTreeExecutionFrame* ParentFrame, const FStateTreeExecutionFrame& CurrentFrame, const FStateTreeDataHandle Handle);
 	UE_API FStateTreeDataView GetDataViewOrTemporary(const FStateTreeExecutionFrame* ParentFrame, const FStateTreeExecutionFrame& CurrentFrame, const FPropertyBindingCopyInfo& CopyInfo);
 
+	/** @return data view of the specified handle from temporary instance. */
+	UE_API FStateTreeDataView GetTemporaryDataView(const FStateTreeExecutionFrame* ParentFrame, const FStateTreeExecutionFrame& CurrentFrame, const FStateTreeDataHandle Handle);
+
 	/**
 	 * Adds a temporary instance that can be located using frame and data handle later.
 	 * @returns view to the newly added instance. If NewInstanceData is Object wrapper, the new object is returned.
diff --git a/Engine/Plugins/Runtime/StateTree/Source/StateTreeModule/Public/StateTreeInstanceData.h b/Engine/Plugins/Runtime/StateTree/Source/StateTreeModule/Public/StateTreeInstanceData.h
--- a/Engine/Plugins/Runtime/StateTree/Source/StateTreeModule/Public/StateTreeInstanceData.h
+++ b/Engine/Plugins/Runtime/StateTree/Source/StateTreeModule/Public/StateTreeInstanceData.h
@@ -25,6 +25,43 @@ namespace UE::StateTree::Debug
 }
 #endif
 
+namespace UE::StateTree::InstanceData
+{
+	namespace Internal
+	{
+		/** @return true if handle source is valid cified handle relative to given frame. */
+		bool IsHandleSourceValid(
+			const FStateTreeInstanceStorage& InstanceStorage,
+			const FStateTreeExecutionFrame* ParentFrame,
+			const FStateTreeExecutionFrame& CurrentFrame,
+			const FStateTreeDataHandle Handle);
+
+		/** @return data view of the specified handle in temporary instance. */
+		FStateTreeDataView GetTemporaryDataView(
+			FStateTreeInstanceStorage& InstanceStorage,
+			const FStateTreeExecutionFrame* ParentFrame,
+			const FStateTreeExecutionFrame& CurrentFrame,
+			const FStateTreeDataHandle Handle);
+	}
+
+	/** @return data view of the specified handle relative to given frame. */
+	UE_API FStateTreeDataView GetDataView(
+		FStateTreeInstanceStorage& InstanceStorage, 
+		FStateTreeInstanceStorage* SharedInstanceStorage, 
+		const FStateTreeExecutionFrame* ParentFrame, 
+		const FStateTreeExecutionFrame& CurrentFrame, 
+		const FStateTreeDataHandle Handle);
+
+	/** @return data view of the specified handle relative to the given frame, or tries to find a matching temporary instance. */
+	UE_API FStateTreeDataView GetDataViewOrTemporary(
+		FStateTreeInstanceStorage& InstanceStorage,
+		FStateTreeInstanceStorage* SharedInstanceStorage,
+		const FStateTreeExecutionFrame* ParentFrame,
+		const FStateTreeExecutionFrame& CurrentFrame,
+		const FStateTreeDataHandle Handle);
+
+}
+
 /** Wrapper class to store an object amongst the structs. */
 USTRUCT()
 struct FStateTreeInstanceObjectWrapper
diff --git a/Engine/Plugins/Runtime/StateTree/Source/StateTreeModule/Public/StateTreePropertyBindings.h b/Engine/Plugins/Runtime/StateTree/Source/StateTreeModule/Public/StateTreePropertyBindings.h
--- a/Engine/Plugins/Runtime/StateTree/Source/StateTreeModule/Public/StateTreePropertyBindings.h
+++ b/Engine/Plugins/Runtime/StateTree/Source/StateTreeModule/Public/StateTreePropertyBindings.h
@@ -583,7 +583,7 @@ T* FStateTreePropertyBindings::GetMutablePropertyPtr(FStateTreeDataView SourceVi
 {
 	check(SourceView.GetStruct() == PropertyAccess.SourceStructType);
 
-	if (!UE::StateTree::PropertyRefHelpers::Validator<T>::IsValid(*PropertyAccess.SourceLeafProperty))
+	if (!UE::StateTree::PropertyRefHelpers::Validator<std::remove_cv_t<T>>::IsValid(*PropertyAccess.SourceLeafProperty))
 	{
 		return nullptr;
 	}
diff --git a/Engine/Plugins/Runtime/StateTree/Source/StateTreeModule/Public/StateTreePropertyRef.h b/Engine/Plugins/Runtime/StateTree/Source/StateTreeModule/Public/StateTreePropertyRef.h
--- a/Engine/Plugins/Runtime/StateTree/Source/StateTreeModule/Public/StateTreePropertyRef.h
+++ b/Engine/Plugins/Runtime/StateTree/Source/StateTreeModule/Public/StateTreePropertyRef.h
@@ -2,8 +2,10 @@
 
 #pragma once
 
-#include "StateTreeIndexTypes.h"
+#include "StateTreeAsyncExecutionContext.h"
 #include "StateTreeExecutionContext.h"
+#include "StateTreeIndexTypes.h"
+#include "StateTreeInstanceData.h"
 #include "StateTreePropertyRefHelpers.h"
 #include "Templates/Tuple.h"
 #include "StateTreePropertyRef.generated.h"
@@ -26,11 +28,11 @@ namespace UE::StateTree::PropertyRefHelpers
 		const FStateTreePropertyBindings& PropertyBindings = ExecutionFrame.StateTree->GetPropertyBindings();
 		if (const FStateTreePropertyAccess* PropertyAccess = PropertyBindings.GetPropertyAccess(PropertyRef))
 		{
-			FStateTreeDataView SourceView = FStateTreeExecutionContext::GetDataViewFromInstanceStorage(InstanceDataStorage, nullptr, ParentExecutionFrame, ExecutionFrame, PropertyAccess->SourceDataHandle);
+			const FStateTreeDataView SourceView = InstanceData::GetDataViewOrTemporary(InstanceDataStorage, nullptr, ParentExecutionFrame, ExecutionFrame, PropertyAccess->SourceDataHandle);
 			
 			// The only possibility when PropertyRef references another PropertyRef is when source one is a global or subtree parameter, i.e lives in parent execution frame.
 			// If that's the case, referenced PropertyRef is obtained and we recursively take the address where it points to.
-			if (UE::StateTree::PropertyRefHelpers::IsPropertyRef(*PropertyAccess->SourceLeafProperty))
+			if (IsPropertyRef(*PropertyAccess->SourceLeafProperty))
 			{
 				check(PropertyAccess->SourceDataHandle.GetSource() == EStateTreeDataSourceType::GlobalParameterData ||
 					PropertyAccess->SourceDataHandle.GetSource() == EStateTreeDataSourceType::ExternalGlobalParameterData ||
@@ -86,7 +88,7 @@ namespace UE::StateTree::PropertyRefHelpers
 		if (const FStateTreePropertyAccess* PropertyAccess = PropertyBindings.GetPropertyAccess(PropertyRef))
 		{
 			// Passing empty ContextAndExternalDataViews, as PropertyRef is not allowed to point to context or external data.
-			FStateTreeDataView SourceView = FStateTreeExecutionContext::GetDataViewFromInstanceStorage(InstanceDataStorage, nullptr, ParentExecutionFrame, ExecutionFrame, PropertyAccess->SourceDataHandle);
+			const FStateTreeDataView SourceView = InstanceData::GetDataViewOrTemporary(InstanceDataStorage, nullptr, ParentExecutionFrame, ExecutionFrame, PropertyAccess->SourceDataHandle);
 			return TTuple<T*...>(PropertyBindings.GetMutablePropertyPtr<T>(SourceView, *PropertyAccess)...);
 		}
 
@@ -138,7 +140,7 @@ struct FStateTreePropertyRef
 
 	/** @return pointer to the property if possible, nullptr otherwise. */
 	template <class T>
-	T* GetMutablePtr(FStateTreeExecutionContext& Context) const
+	T* GetMutablePtr(const FStateTreeExecutionContext& Context) const
 	{
 		const FStateTreeExecutionFrame* CurrentlyProcessedFrame = Context.GetCurrentlyProcessedFrame();
 		check(CurrentlyProcessedFrame);
@@ -146,9 +148,22 @@ struct FStateTreePropertyRef
 		return UE::StateTree::PropertyRefHelpers::GetMutablePtrToProperty<T>(*this, Context.GetMutableInstanceData()->GetMutableStorage(), *CurrentlyProcessedFrame, Context.GetCurrentlyProcessedParentFrame());
 	}
 
+	/** @return pointer to the property if possible, nullptr otherwise. */
+	template<class T, bool bWithWriteAccess>
+	std::conditional_t<bWithWriteAccess, T*, const T*> GetPtrFromStrongExecutionContext(const TStateTreeStrongExecutionContext<bWithWriteAccess>& Context)
+	{
+		UE::StateTree::Async::FActivePathInfo ActivePath = Context.GetActivePathInfo();
+		if (ActivePath.IsValid())
+		{
+			return UE::StateTree::PropertyRefHelpers::GetMutablePtrToProperty<T>(*this, *Context.Storage, *ActivePath.Frame, ActivePath.ParentFrame);
+		}
+
+		return nullptr;
+	}
+
 	/** @return a tuple of pointers of the given types to the property if possible, nullptr otherwise. */
 	template <class... T>
-	TTuple<T*...> GetMutablePtrTuple(FStateTreeExecutionContext& Context) const
+	TTuple<T*...> GetMutablePtrTuple(const FStateTreeExecutionContext& Context) const
 	{
 		const FStateTreeExecutionFrame* CurrentlyProcessedFrame = Context.GetCurrentlyProcessedFrame();
 		check(CurrentlyProcessedFrame);
@@ -156,6 +171,19 @@ struct FStateTreePropertyRef
 		return UE::StateTree::PropertyRefHelpers::GetMutablePtrTupleToProperty<T...>(*this, Context.GetMutableInstanceData()->GetMutableStorage(), *CurrentlyProcessedFrame, Context.GetCurrentlyProcessedParentFrame());
 	}
 
+	/** @return a tuple of pointers of the given types to the property if possible, nullptr otherwise. */
+	template <class... T, bool bWithWriteAccess>
+	std::conditional_t<bWithWriteAccess, TTuple<T*...>, TTuple<const T*...>> GetPtrTupleFromStrongExecutionContext(const TStateTreeStrongExecutionContext<bWithWriteAccess>& Context) const
+	{
+		UE::StateTree::Async::FActivePathInfo ActivePath = Context.GetActivePathInfo();
+		if (ActivePath.IsValid())
+		{
+			return UE::StateTree::PropertyRefHelpers::GetMutablePtrTupleToProperty<T...>(*this, *Context.Storage, *ActivePath.Frame, ActivePath.ParentFrame);
+		}
+
+		return {};
+	}
+
 	/**
 	 * Used internally.
 	 * @return index to referenced property access
@@ -361,4 +389,4 @@ private:
 	TObjectPtr<UObject> TypeObject = nullptr;
 
 	friend class FStateTreeBlueprintPropertyRefDetails;
-};
+};
\ No newline at end of file
