StateTree: Blueprint task can set the default completion flags.
The Blueprint has the consideration flags duplicated. Usually, the flags are copied to the wrapper task when compiled. If we do that, we lose the "reset to default" and propagation behavior in the editor.

diff --git a/Engine/Plugins/Runtime/StateTree/Source/StateTreeEditorModule/Private/Customizations/StateTreeEditorNodeDetails.cpp b/Engine/Plugins/Runtime/StateTree/Source/StateTreeEditorModule/Private/Customizations/StateTreeEditorNodeDetails.cpp
--- a/Engine/Plugins/Runtime/StateTree/Source/StateTreeEditorModule/Private/Customizations/StateTreeEditorNodeDetails.cpp
+++ b/Engine/Plugins/Runtime/StateTree/Source/StateTreeEditorModule/Private/Customizations/StateTreeEditorNodeDetails.cpp
@@ -674,7 +674,7 @@ void FStateTreeEditorNodeDetails::CustomizeHeader(TSharedRef<class IPropertyHand
 		.PasteAction(FUIAction(FExecuteAction::CreateSP(this, &FStateTreeEditorNodeDetails::OnPasteNode)));
 
 	// Task completion
-	if (const FStateTreeEditorNode * Node = UE::StateTreeEditor::EditorNodeUtils::GetCommonNode(StructProperty))
+	if (const FStateTreeEditorNode* Node = UE::StateTreeEditor::EditorNodeUtils::GetCommonNode(StructProperty))
 	{
 		if (const FStateTreeTaskBase* TaskBase = Node->Node.GetPtr<FStateTreeTaskBase>())
 		{
@@ -693,7 +693,7 @@ void FStateTreeEditorNodeDetails::CustomizeHeader(TSharedRef<class IPropertyHand
 						.ColorAndOpacity(this, &FStateTreeEditorNodeDetails::GetToggleCompletionTaskColor)
 						.Image(FAppStyle::Get().GetBrush("Icons.Star"))
 					]
-					.IsEnabled(TaskBase->bCanEditConsideredForCompletion)
+					.IsEnabled(UE::StateTreeEditor::EditorNodeUtils::CanEditTaskConsideredForCompletion(*Node))
 				];
 		}
 	}
@@ -1622,10 +1622,8 @@ FReply FStateTreeEditorNodeDetails::HandleToggleCompletionTaskClicked()
 {
 	if (FStateTreeEditorNode* Node = UE::StateTreeEditor::EditorNodeUtils::GetMutableCommonNode(StructProperty))
 	{
-		if (FStateTreeTaskBase* TaskBase = Node->Node.GetMutablePtr<FStateTreeTaskBase>())
-		{
-			TaskBase->bConsideredForCompletion = !TaskBase->bConsideredForCompletion;
-		}
+		const bool bCurrentValue = UE::StateTreeEditor::EditorNodeUtils::IsTaskConsideredForCompletion(*Node);
+		UE::StateTreeEditor::EditorNodeUtils::SetTaskConsideredForCompletion(*Node, !bCurrentValue);
 	}
 	return FReply::Handled();
 }
@@ -1634,12 +1632,9 @@ FSlateColor FStateTreeEditorNodeDetails::GetToggleCompletionTaskColor() const
 {
 	if (const FStateTreeEditorNode* Node = UE::StateTreeEditor::EditorNodeUtils::GetCommonNode(StructProperty))
 	{
-		if (const FStateTreeTaskBase* TaskBase = Node->Node.GetPtr<FStateTreeTaskBase>())
+		if (UE::StateTreeEditor::EditorNodeUtils::IsTaskConsideredForCompletion(*Node))
 		{
-			if (TaskBase->bConsideredForCompletion)
-			{
-				return FSlateColor(EStyleColor::Foreground);
-			}
+			return FSlateColor(EStyleColor::Foreground);
 		}
 	}
 	return FSlateColor(EStyleColor::Recessed);
diff --git a/Engine/Plugins/Runtime/StateTree/Source/StateTreeEditorModule/Private/Customizations/StateTreeEditorNodeUtils.cpp b/Engine/Plugins/Runtime/StateTree/Source/StateTreeEditorModule/Private/Customizations/StateTreeEditorNodeUtils.cpp
--- a/Engine/Plugins/Runtime/StateTree/Source/StateTreeEditorModule/Private/Customizations/StateTreeEditorNodeUtils.cpp
+++ b/Engine/Plugins/Runtime/StateTree/Source/StateTreeEditorModule/Private/Customizations/StateTreeEditorNodeUtils.cpp
@@ -64,12 +64,66 @@ bool IsTaskDisabled(const TSharedPtr<IPropertyHandle>& StructProperty)
 {
 	if (const FStateTreeEditorNode* Node = GetCommonNode(StructProperty))
 	{
-		if (const FStateTreeTaskBase* TaskBase = Node->Node.GetPtr<FStateTreeTaskBase>())
+		return !IsTaskEnabled(*Node);
+	}
+
+	return false;
+}
+
+bool IsTaskEnabled(const FStateTreeEditorNode& EditorNode)
+{
+	if (const FStateTreeTaskBase* TaskBase = EditorNode.Node.GetPtr<FStateTreeTaskBase>())
+	{
+		return TaskBase->bTaskEnabled;
+	}
+	return false;
+}
+
+bool IsTaskConsideredForCompletion(const FStateTreeEditorNode& EditorNode)
+{
+	// We use the Blueprint flag to have a default value that behaves like the other flags. Sadly, it duplicates the flags.
+	if (const FStateTreeBlueprintTaskWrapper* TaskWrapper = EditorNode.Node.GetPtr<FStateTreeBlueprintTaskWrapper>())
+	{
+		if (UStateTreeTaskBlueprintBase* BPTaskBase = Cast<UStateTreeTaskBlueprintBase>(EditorNode.InstanceObject))
 		{
-			return !TaskBase->bTaskEnabled;
+			return BPTaskBase->bConsideredForCompletion;
 		}
 	}
+	else if (const FStateTreeTaskBase* TaskBase = EditorNode.Node.GetPtr<FStateTreeTaskBase>())
+	{
+		return TaskBase->bConsideredForCompletion;
+	}
+	return false;
+}
 
+void SetTaskConsideredForCompletion(FStateTreeEditorNode& EditorNode, bool bIsConsidered)
+{
+	if (const FStateTreeBlueprintTaskWrapper* TaskWrapper = EditorNode.Node.GetPtr<FStateTreeBlueprintTaskWrapper>())
+	{
+		if (UStateTreeTaskBlueprintBase* BPTaskBase = Cast<UStateTreeTaskBlueprintBase>(EditorNode.InstanceObject))
+		{
+			BPTaskBase->bConsideredForCompletion = bIsConsidered;
+		}
+	}
+	else if (FStateTreeTaskBase* TaskBase = EditorNode.Node.GetMutablePtr<FStateTreeTaskBase>())
+	{
+		TaskBase->bConsideredForCompletion = bIsConsidered;
+	}
+}
+
+bool CanEditTaskConsideredForCompletion(const FStateTreeEditorNode& EditorNode)
+{
+	if (const FStateTreeBlueprintTaskWrapper* TaskWrapper = EditorNode.Node.GetPtr<FStateTreeBlueprintTaskWrapper>())
+	{
+		if (UStateTreeTaskBlueprintBase* BPTaskBase = Cast<UStateTreeTaskBlueprintBase>(EditorNode.InstanceObject))
+		{
+			return BPTaskBase->bCanEditConsideredForCompletion;
+		}
+	}
+	else if (const FStateTreeTaskBase* TaskBase = EditorNode.Node.GetPtr<FStateTreeTaskBase>())
+	{
+		return TaskBase->bCanEditConsideredForCompletion;
+	}
 	return false;
 }
 
diff --git a/Engine/Plugins/Runtime/StateTree/Source/StateTreeEditorModule/Private/Customizations/StateTreeEditorNodeUtils.h b/Engine/Plugins/Runtime/StateTree/Source/StateTreeEditorModule/Private/Customizations/StateTreeEditorNodeUtils.h
--- a/Engine/Plugins/Runtime/StateTree/Source/StateTreeEditorModule/Private/Customizations/StateTreeEditorNodeUtils.h
+++ b/Engine/Plugins/Runtime/StateTree/Source/StateTreeEditorModule/Private/Customizations/StateTreeEditorNodeUtils.h
@@ -65,10 +65,25 @@ EStateTreeConditionEvaluationMode GetConditionEvaluationMode(const TSharedPtr<IP
 
 /**
  * @param StructProperty Property handle pointing to a FStateTreeEditorNode. 
- * @return true if the specified task is marked as disable, or false if node not found. 
+ * @return Whether specified task node is marked as disable. Will return false if node is not a valid task.
  */
 bool IsTaskDisabled(const TSharedPtr<IPropertyHandle>& StructProperty);
 
+/**
+ * @param StructProperty Property handle pointing to a FStateTreeEditorNode.
+ * @return Whether the specified task node is marked as enabled. Will return false if node is not a valid task.
+ */
+bool IsTaskEnabled(const FStateTreeEditorNode& EditorNode);
+
+/** @return Whether the specified task is marked as ConsideredForCompletion.  Will return false if node is not a valid task. */
+bool IsTaskConsideredForCompletion(const FStateTreeEditorNode& EditorNode);
+
+/** Set the ConsideredForCompletion flag on the specified task. */
+void SetTaskConsideredForCompletion(FStateTreeEditorNode& EditorNode, bool bIsConsidered);
+
+/** @return Whether the ConsideredForCompletion flag can be edited on the specified task. */
+bool CanEditTaskConsideredForCompletion(const FStateTreeEditorNode& EditorNode);
+
 /**
  * Execute the provided function within a Transaction 
  * @param Description Description to associate to the transaction
diff --git a/Engine/Plugins/Runtime/StateTree/Source/StateTreeEditorModule/Private/Debugger/StateTreeDebuggerUIExtensions.cpp b/Engine/Plugins/Runtime/StateTree/Source/StateTreeEditorModule/Private/Debugger/StateTreeDebuggerUIExtensions.cpp
--- a/Engine/Plugins/Runtime/StateTree/Source/StateTreeEditorModule/Private/Debugger/StateTreeDebuggerUIExtensions.cpp
+++ b/Engine/Plugins/Runtime/StateTree/Source/StateTreeEditorModule/Private/Debugger/StateTreeDebuggerUIExtensions.cpp
@@ -540,9 +540,9 @@ TSharedRef<SWidget> CreateEditorNodeWidget(const TSharedPtr<IPropertyHandle>& St
 				SNew(SBorder)
 				.Padding(FMargin(6, 2))
 				.BorderImage(FStateTreeEditorStyle::Get().GetBrush("StateTree.Node.Label"))
-				.Visibility_Lambda([StructPropertyHandle]
+				.Visibility_Lambda([bIsTask, StructPropertyHandle]
 					{
-						return IsTaskDisabled(StructPropertyHandle) ? EVisibility::Visible : EVisibility::Collapsed;
+						return bIsTask && IsTaskDisabled(StructPropertyHandle) ? EVisibility::Visible : EVisibility::Collapsed;
 					})
 				[
 					SNew(STextBlock)
diff --git a/Engine/Plugins/Runtime/StateTree/Source/StateTreeEditorModule/Private/StateTreeCompiler.cpp b/Engine/Plugins/Runtime/StateTree/Source/StateTreeEditorModule/Private/StateTreeCompiler.cpp
--- a/Engine/Plugins/Runtime/StateTree/Source/StateTreeEditorModule/Private/StateTreeCompiler.cpp
+++ b/Engine/Plugins/Runtime/StateTree/Source/StateTreeEditorModule/Private/StateTreeCompiler.cpp
@@ -16,6 +16,7 @@
 #include "StateTreePropertyRefHelpers.h"
 #include "StateTreePropertyHelpers.h"
 #include "StateTreeDelegate.h"
+#include "Customizations/StateTreeEditorNodeUtils.h"
 
 
 namespace UE::StateTree::Compiler
@@ -243,9 +244,9 @@ namespace UE::StateTree::Compiler
 			for (int32 Index = NumberOfTasks-1; Index >= 0; --Index)
 			{
 				const int32 TaskIndex = ValidTasks[Index];
-				const FStateTreeTaskBase& Task = AllTasks[TaskIndex].Node.Get<FStateTreeTaskBase>();
 				Mask <<= 1;
-				if (Task.bConsideredForCompletion)
+				if (UE::StateTreeEditor::EditorNodeUtils::IsTaskEnabled(AllTasks[TaskIndex])
+					&& UE::StateTreeEditor::EditorNodeUtils::IsTaskConsideredForCompletion(AllTasks[TaskIndex]))
 				{
 					Mask |= 1;
 				}
@@ -756,7 +757,7 @@ bool FStateTreeCompiler::CreateGlobalTasks()
 	const UE::StateTree::Compiler::FCompletionTasksMaskResult MaskResult = UE::StateTree::Compiler::MakeCompletionTasksMask(CompletionGlobalTaskStartBitIndex, EditorData->GlobalTasks, ValidTaskNodeIndex);
 	StateTree->CompletionGlobalTasksMask = MaskResult.Mask;
 	GlobalTaskEndBit = MaskResult.FullMaskEndTaskBitOffset;
-	StateTree->CompletionGlobalTasksControl = EditorData->GlobalTasksCompletion;
+	StateTree->CompletionGlobalTasksControl = (Schema && Schema->AllowTasksCompletion()) ? EditorData->GlobalTasksCompletion : EStateTreeTaskCompletionType::Any;
 
 	if (MaskResult.MaskFirstTaskBitOffset != 0)
 	{
@@ -1039,7 +1040,7 @@ bool FStateTreeCompiler::CreateStateTasksAndParameters()
 			}
 
 			CompactState.CompletionTasksMask = MaskResult.Mask;
-			CompactState.CompletionTasksControl = State->TasksCompletion;
+			CompactState.CompletionTasksControl = (Schema && Schema->AllowTasksCompletion()) ? State->TasksCompletion : EStateTreeTaskCompletionType::Any;
 			CompactState.CompletionTasksMaskBufferIndex = static_cast<uint8>(CompletionTasksMaskBufferIndex);
 			CompactState.CompletionTasksMaskBitsOffset = static_cast<uint8>(CompletionTasksMaskBitsOffset);
 			StateInfos[StateIndex].NextBitIndexForCompletionMask = MaskResult.FullMaskEndTaskBitOffset;
@@ -2357,7 +2358,7 @@ bool FStateTreeCompiler::GetAndValidateBindings(const FStateTreeBindableStructDe
 				&& !UE::StateTree::AcceptTaskInstanceData(TargetStruct.DataSource))
 			{
 				Log.Reportf(EMessageSeverity::Error, TargetStruct,
-					TEXT("Property at %s cannot be bound to %s, because the binding source %s is an instance data that is possibly not instantiated before %s in the tree."),
+					TEXT("Property at %s cannot be bound to %s, because the binding source %s is a task instance data that is possibly not instantiated before %s in the tree."),
 					* UE::StateTree::GetDescAndPathAsString(*SourceStruct, Binding.GetSourcePath()),
 					* UE::StateTree::GetDescAndPathAsString(TargetStruct, Binding.GetTargetPath()),
 					* SourceStruct->ToString(), * TargetStruct.ToString());
diff --git a/Engine/Plugins/Runtime/StateTree/Source/StateTreeModule/Private/Blueprint/StateTreeTaskBlueprintBase.cpp b/Engine/Plugins/Runtime/StateTree/Source/StateTreeModule/Private/Blueprint/StateTreeTaskBlueprintBase.cpp
--- a/Engine/Plugins/Runtime/StateTree/Source/StateTreeModule/Private/Blueprint/StateTreeTaskBlueprintBase.cpp
+++ b/Engine/Plugins/Runtime/StateTree/Source/StateTreeModule/Private/Blueprint/StateTreeTaskBlueprintBase.cpp
@@ -19,6 +19,10 @@ UStateTreeTaskBlueprintBase::UStateTreeTaskBlueprintBase(const FObjectInitialize
 	, bShouldCallTickOnlyOnEvents(false)
 	, bShouldCopyBoundPropertiesOnTick(true)
 	, bShouldCopyBoundPropertiesOnExitState(true)
+#if WITH_EDITORONLY_DATA
+	, bConsideredForCompletion(true)
+	, bCanEditConsideredForCompletion(true)
+#endif
 	, bIsProcessingEnterStateOrTick(false)
 {
 	bHasExitState = BlueprintNodeHelpers::HasBlueprintFunction(TEXT("ReceiveExitState"), *this, *StaticClass());
diff --git a/Engine/Plugins/Runtime/StateTree/Source/StateTreeModule/Public/Blueprint/StateTreeTaskBlueprintBase.h b/Engine/Plugins/Runtime/StateTree/Source/StateTreeModule/Public/Blueprint/StateTreeTaskBlueprintBase.h
--- a/Engine/Plugins/Runtime/StateTree/Source/StateTreeModule/Public/Blueprint/StateTreeTaskBlueprintBase.h
+++ b/Engine/Plugins/Runtime/StateTree/Source/StateTreeModule/Public/Blueprint/StateTreeTaskBlueprintBase.h
@@ -21,6 +21,7 @@ UCLASS(MinimalAPI, Abstract, Blueprintable)
 class UStateTreeTaskBlueprintBase : public UStateTreeNodeBlueprintBase
 {
 	GENERATED_BODY()
+
 public:
 	UE_API UStateTreeTaskBlueprintBase(const FObjectInitializer& ObjectInitializer);
 
@@ -75,7 +76,6 @@ public:
 	UE_API EStateTreeRunStatus ReceiveTick(const float DeltaTime);
 
 protected:
-
 	UE_API virtual EStateTreeRunStatus EnterState(FStateTreeExecutionContext& Context, const FStateTreeTransitionResult& Transition);
 	UE_API virtual void ExitState(FStateTreeExecutionContext& Context, const FStateTreeTransitionResult& Transition);
 
@@ -137,7 +137,22 @@ protected:
 	/** If set to true, copy the values of bound properties before calling ExitState(). Default true. */
 	UPROPERTY(EditDefaultsOnly, Category="Default")
 	uint8 bShouldCopyBoundPropertiesOnExitState : 1;
-	
+
+public:
+#if WITH_EDITORONLY_DATA
+	/**
+	 * True if the task is considered for completion.
+	 * False if the task runs in the background without affecting the state completion.
+	 */
+	UPROPERTY(EditDefaultsOnly, Category = "Default")
+	uint8 bConsideredForCompletion : 1;
+
+	/** True if the user can edit bConsideredForCompletion in the editor. */
+	UPROPERTY(EditDefaultsOnly, Category = "Default")
+	uint8 bCanEditConsideredForCompletion : 1;
+#endif
+
+protected:
 	uint8 bHasExitState : 1;
 	uint8 bHasStateCompleted : 1;
 	uint8 bHasLatentEnterState : 1;
@@ -160,7 +175,10 @@ struct FStateTreeBlueprintTaskWrapper : public FStateTreeTaskBase
 {
 	GENERATED_BODY()
 
-	virtual const UStruct* GetInstanceDataType() const override { return TaskClass; };
+	virtual const UStruct* GetInstanceDataType() const override
+	{
+		return TaskClass;
+	};
 	UE_API virtual EDataValidationResult Compile(FStateTreeDataView InstanceDataView, TArray<FText>& ValidationMessages) override;
 	UE_API virtual bool Link(FStateTreeLinker& Linker) override;
 	UE_API virtual EStateTreeRunStatus EnterState(FStateTreeExecutionContext& Context, const FStateTreeTransitionResult& Transition) const override;
@@ -178,6 +196,13 @@ struct FStateTreeBlueprintTaskWrapper : public FStateTreeTaskBase
 
 	UPROPERTY()
 	uint8 TaskFlags = 0;
+
+private:
+#if WITH_EDITORONLY_DATA
+	// The node will use the blueprint data instead.
+	using FStateTreeTaskBase::bConsideredForCompletion;
+	using FStateTreeTaskBase::bCanEditConsideredForCompletion;
+#endif
 };
 
 #undef UE_API
diff --git a/Engine/Plugins/Runtime/StateTree/Source/StateTreeModule/Public/StateTreeSchema.h b/Engine/Plugins/Runtime/StateTree/Source/StateTreeModule/Public/StateTreeSchema.h
--- a/Engine/Plugins/Runtime/StateTree/Source/StateTreeModule/Public/StateTreeSchema.h
+++ b/Engine/Plugins/Runtime/StateTree/Source/StateTreeModule/Public/StateTreeSchema.h
@@ -113,7 +113,7 @@ public:
 		return true;
 	}
 
-	/** @return True if modifying the tasks completion is allowed. */
+	/** @return True if modifying the tasks completion is allowed. If not allowed, "any" will be used.*/
 	virtual bool AllowTasksCompletion() const
 	{
 		return true;
diff --git a/Engine/Plugins/Runtime/StateTree/Source/StateTreeModule/Public/StateTreeTaskBase.h b/Engine/Plugins/Runtime/StateTree/Source/StateTreeModule/Public/StateTreeTaskBase.h
--- a/Engine/Plugins/Runtime/StateTree/Source/StateTreeModule/Public/StateTreeTaskBase.h
+++ b/Engine/Plugins/Runtime/StateTree/Source/StateTreeModule/Public/StateTreeTaskBase.h
@@ -41,14 +41,19 @@ struct FStateTreeTaskBase : public FStateTreeNodeBase
 	 * Note: The API has been deprecated. ChangeType is moved into FStateTreeTransitionResult.
 	 * You can configure the task to be only called on state changes (that is, never call sustained changes) by setting bShouldStateChangeOnReselect to true.
 	 */
-	virtual EStateTreeRunStatus EnterState(FStateTreeExecutionContext& Context, const EStateTreeStateChangeType ChangeType, const FStateTreeTransitionResult& Transition) const final { return EStateTreeRunStatus::Running; }
+	virtual EStateTreeRunStatus EnterState(FStateTreeExecutionContext& Context, const EStateTreeStateChangeType ChangeType, const FStateTreeTransitionResult& Transition) const final
+	{
+		return EStateTreeRunStatus::Running;
+	}
 
 	UE_DEPRECATED(5.5, "Use ExitState without the EStateTreeStateChangeType.")
 	/**
 	 * Note: The API has been deprecated. ChangeType is moved into FStateTreeTransitionResult.
 	 * You can configure the task to be only called on state changes (that is, never call sustained changes) by setting bShouldStateChangeOnReselect to true.
 	 */
-	virtual void ExitState(FStateTreeExecutionContext& Context, const EStateTreeStateChangeType ChangeType, const FStateTreeTransitionResult& Transition) const final {}
+	virtual void ExitState(FStateTreeExecutionContext& Context, const EStateTreeStateChangeType ChangeType, const FStateTreeTransitionResult& Transition) const final
+	{
+	}
 
 	/**
 	 * Called when a new state is entered and task is part of active states.
@@ -56,14 +61,19 @@ struct FStateTreeTaskBase : public FStateTreeNodeBase
 	 * @param Transition Describes the states involved in the transition
 	 * @return Succeed/Failed will end the state immediately and trigger to select new state, Running will carry on to tick the state.
 	 */
-	virtual EStateTreeRunStatus EnterState(FStateTreeExecutionContext& Context, const FStateTreeTransitionResult& Transition) const { return EStateTreeRunStatus::Running; }
+	virtual EStateTreeRunStatus EnterState(FStateTreeExecutionContext& Context, const FStateTreeTransitionResult& Transition) const
+	{
+		return EStateTreeRunStatus::Running;
+	}
 
 	/**
-	* Called when a current state is exited and task is part of active states.
-	* @param Context Reference to current execution context.
-	* @param Transition Describes the states involved in the transition
-	*/
-	virtual void ExitState(FStateTreeExecutionContext& Context, const FStateTreeTransitionResult& Transition) const {}
+	 * Called when a current state is exited and task is part of active states.
+	 * @param Context Reference to current execution context.
+	 * @param Transition Describes the states involved in the transition
+	 */
+	virtual void ExitState(FStateTreeExecutionContext& Context, const FStateTreeTransitionResult& Transition) const
+	{
+	}
 
 	/**
 	 * Called right after a state has been completed, but before new state has been selected. StateCompleted is called in reverse order to allow to propagate state to other Tasks that
@@ -72,7 +82,9 @@ struct FStateTreeTaskBase : public FStateTreeNodeBase
 	 * @param CompletionStatus Describes the running status of the completed state (Succeeded/Failed).
 	 * @param CompletedActiveStates Active states at the time of completion.
 	 */
-	virtual void StateCompleted(FStateTreeExecutionContext& Context, const EStateTreeRunStatus CompletionStatus, const FStateTreeActiveStates& CompletedActiveStates) const {}
+	virtual void StateCompleted(FStateTreeExecutionContext& Context, const EStateTreeRunStatus CompletionStatus, const FStateTreeActiveStates& CompletedActiveStates) const
+	{
+	}
 
 	/**
 	 * Called during state tree tick when the task is on active state.
@@ -81,14 +93,19 @@ struct FStateTreeTaskBase : public FStateTreeNodeBase
 	 * @param DeltaTime Time since last StateTree tick.
 	 * @return Running status of the state: Running if still in progress, Succeeded if execution is done and succeeded, Failed if execution is done and failed.
 	 */
-	virtual EStateTreeRunStatus Tick(FStateTreeExecutionContext& Context, const float DeltaTime) const { return EStateTreeRunStatus::Running; };
+	virtual EStateTreeRunStatus Tick(FStateTreeExecutionContext& Context, const float DeltaTime) const
+	{
+		return EStateTreeRunStatus::Running;
+	}
 
 	/**
 	 * Called when state tree triggers transitions. This method is called during transition handling, before state's tick and event transitions are handled.
 	 * Note: the method is called only if bShouldAffectTransitions is set.
 	 * @param Context Reference to current execution context.
 	 */
-	virtual void TriggerTransitions(FStateTreeExecutionContext& Context) const {};
+	virtual void TriggerTransitions(FStateTreeExecutionContext& Context) const
+	{
+	};
 
 #if WITH_EDITOR
 	virtual FName GetIconName() const override
@@ -112,7 +129,8 @@ struct FStateTreeTaskBase : public FStateTreeNodeBase
 	 * If set to true, the task will receive EnterState/ExitState even if the state was previously active.
 	 * Generally this should be true for action type tasks, like playing animation,
 	 * and false on state like tasks like claiming a resource that is expected to be acquired on child states.
-	 * Default value is true. */
+	 * Default value is true.
+	 */
 	uint8 bShouldStateChangeOnReselect : 1;
 
 	/** If set to true, Tick() is called. Not ticking implies no property copy. Default true. */
