StateTree: Update the node title to respect clipping. Update the condition parentheses.

diff --git a/Engine/Plugins/Runtime/StateTree/Source/StateTreeEditorModule/Private/Customizations/StateTreeEditorNodeDetails.cpp b/Engine/Plugins/Runtime/StateTree/Source/StateTreeEditorModule/Private/Customizations/StateTreeEditorNodeDetails.cpp
--- a/Engine/Plugins/Runtime/StateTree/Source/StateTreeEditorModule/Private/Customizations/StateTreeEditorNodeDetails.cpp
+++ b/Engine/Plugins/Runtime/StateTree/Source/StateTreeEditorModule/Private/Customizations/StateTreeEditorNodeDetails.cpp
@@ -516,6 +516,7 @@ void FStateTreeEditorNodeDetails::CustomizeHeader(TSharedRef<class IPropertyHand
 				.Padding(FMargin(0, 4))
 				[
 					SAssignNew(DescriptionBox, SHorizontalBox)
+					.Clipping(EWidgetClipping::ClipToBoundsAlways)
 
 					// Icon
 					+ SHorizontalBox::Slot()
@@ -574,17 +575,18 @@ void FStateTreeEditorNodeDetails::CustomizeHeader(TSharedRef<class IPropertyHand
 							.BorderImage(FStyleDefaults::GetNoBrush())
 							.Visibility(this, &FStateTreeEditorNodeDetails::AreFlagsVisible)
 					]
-				]
 
-				// Close parens
-				+ SHorizontalBox::Slot()
-				.AutoWidth()
-				.VAlign(VAlign_Center)
-				[
-					SNew(STextBlock)
-					.TextStyle(FStateTreeEditorStyle::Get(), "StateTree.Node.Parens")
-					.Text(this, &FStateTreeEditorNodeDetails::GetCloseParens)
-					.Visibility(this, &FStateTreeEditorNodeDetails::AreParensVisible)
+					// Close parens
+					+ SHorizontalBox::Slot()
+					.AutoWidth()
+					.HAlign(HAlign_Right)
+					.VAlign(VAlign_Center)
+					[
+						SNew(STextBlock)
+						.TextStyle(FStateTreeEditorStyle::Get(), "StateTree.Node.Parens")
+						.Text(this, &FStateTreeEditorNodeDetails::GetCloseParens)
+						.Visibility(this, &FStateTreeEditorNodeDetails::AreParensVisible)
+					]
 				]
 
 				// Debug and property widgets
@@ -651,6 +653,7 @@ void FStateTreeEditorNodeDetails::CustomizeHeader(TSharedRef<class IPropertyHand
 					+ SHorizontalBox::Slot()
 					.AutoWidth()
 					.VAlign(VAlign_Center)
+					.HAlign(HAlign_Right)
 					.Padding(FMargin(4.f, 0.f, 0.f, 0.f))
 					[
 						SNew(SComboButton)
