StateTree: Revert a fix made for transition with ID. When trying to reselect a linked state, always try to select the existing state instead of creating a new ID. This will be fixed with the transition ID works in progress.

diff --git a/Engine/Plugins/Runtime/StateTree/Source/StateTreeModule/Private/StateTreeExecutionContext.cpp b/Engine/Plugins/Runtime/StateTree/Source/StateTreeModule/Private/StateTreeExecutionContext.cpp
--- a/Engine/Plugins/Runtime/StateTree/Source/StateTreeModule/Private/StateTreeExecutionContext.cpp
+++ b/Engine/Plugins/Runtime/StateTree/Source/StateTreeModule/Private/StateTreeExecutionContext.cpp
@@ -5429,13 +5429,8 @@ bool FStateTreeExecutionContext::SelectStateInternal(
 				FStateTreeExecutionFrame NewFrame;
 				NewFrame.StateTree = CurrentFrame.StateTree;
 				NewFrame.RootState = NextState.LinkedState;
-				NewFrame.FrameID = UE::StateTree::FActiveFrameID(Storage.GenerateUniqueId());
 				NewFrame.ExternalDataBaseIndex = CurrentFrame.ExternalDataBaseIndex;
 
-				const FCompactStateTreeFrame* FrameInfo = NewFrame.StateTree->GetFrameFromHandle(NewFrame.RootState);
-				ensureAlwaysMsgf(FrameInfo, TEXT("The compiled data is invalid. It should contains the information for the new root frame."));
-				NewFrame.ActiveTasksStatus = FrameInfo ? FStateTreeTasksCompletionStatus(*FrameInfo) : FStateTreeTasksCompletionStatus();
-
 				// Check and prevent recursion.
 				const bool bNewFrameAlreadySelected = OutSelectionResult.GetSelectedFrames().ContainsByPredicate([&NewFrame](const FStateTreeExecutionFrame& Frame) {
 					return Frame.IsSameFrame(NewFrame);
@@ -5456,7 +5451,8 @@ bool FStateTreeExecutionContext::SelectStateInternal(
 					});
 				if (ExistingFrame)
 				{
-					// Note. Do not copy the FrameID. This is a new frame and the enter/exit state should be executed.
+					NewFrame.FrameID = ExistingFrame->FrameID;
+					NewFrame.ActiveTasksStatus = ExistingFrame->ActiveTasksStatus;
 					NewFrame.ActiveInstanceIndexBase = ExistingFrame->ActiveInstanceIndexBase;
 					NewFrame.GlobalInstanceIndexBase = ExistingFrame->GlobalInstanceIndexBase;
 					NewFrame.StateParameterDataHandle = ExistingFrame->StateParameterDataHandle;
@@ -5464,7 +5460,13 @@ bool FStateTreeExecutionContext::SelectStateInternal(
 				}
 				else
 				{
-					// Since the StateTree is the same, we can access the global tasks of CurrentFrame, if they are initialized.   
+					NewFrame.FrameID = UE::StateTree::FActiveFrameID(Storage.GenerateUniqueId());
+
+					const FCompactStateTreeFrame* FrameInfo = NewFrame.StateTree->GetFrameFromHandle(NewFrame.RootState);
+					ensureAlwaysMsgf(FrameInfo, TEXT("The compiled data is invalid. It should contains the information for the new root frame."));
+					NewFrame.ActiveTasksStatus = FrameInfo ? FStateTreeTasksCompletionStatus(*FrameInfo) : FStateTreeTasksCompletionStatus();
+
+					// Since the StateTree is the same, we can access the global tasks of CurrentFrame, if they are initialized.
 					NewFrame.GlobalParameterDataHandle = CurrentFrame.GlobalParameterDataHandle;
 					NewFrame.GlobalInstanceIndexBase = CurrentFrame.GlobalInstanceIndexBase;
 					NewFrame.StateParameterDataHandle = NextState.ParameterDataHandle; // Temporary allocated earlier if did not exists.
@@ -5511,15 +5513,10 @@ bool FStateTreeExecutionContext::SelectStateInternal(
 			}
 				
 			FStateTreeExecutionFrame NewFrame;
-			NewFrame.FrameID = UE::StateTree::FActiveFrameID(Storage.GenerateUniqueId());
 			NewFrame.StateTree = NextLinkedStateAsset;
 			NewFrame.RootState = FStateTreeStateHandle::Root;
 			NewFrame.bIsGlobalFrame = true;
 
-			const FCompactStateTreeFrame* FrameInfo = NewFrame.StateTree->GetFrameFromHandle(NewFrame.RootState);
-			ensureAlwaysMsgf(FrameInfo, TEXT("The compiled data is invalid. It should contains the information for the root frame."));
-			NewFrame.ActiveTasksStatus = FrameInfo ? FStateTreeTasksCompletionStatus(*FrameInfo) : FStateTreeTasksCompletionStatus();
-
 			// Check and prevent recursion.
 			const bool bNewFrameAlreadySelected = OutSelectionResult.GetSelectedFrames().ContainsByPredicate([&NewFrame](const FStateTreeExecutionFrame& Frame) {
 				return Frame.IsSameFrame(NewFrame);
@@ -5542,7 +5539,8 @@ bool FStateTreeExecutionContext::SelectStateInternal(
 			bool bStartedTemporaryEvaluatorsAndGlobalTasks = false;
 			if (ExistingFrame)
 			{
-				// Note. Do not copy the FrameID. This is a new frame and the enter/exit state should be executed.
+				NewFrame.FrameID = ExistingFrame->FrameID;
+				NewFrame.ActiveTasksStatus = ExistingFrame->ActiveTasksStatus;
 				NewFrame.ActiveInstanceIndexBase = ExistingFrame->ActiveInstanceIndexBase;
 				NewFrame.GlobalInstanceIndexBase = ExistingFrame->GlobalInstanceIndexBase;
 				NewFrame.StateParameterDataHandle = ExistingFrame->StateParameterDataHandle;
@@ -5551,6 +5549,11 @@ bool FStateTreeExecutionContext::SelectStateInternal(
 			}
 			else
 			{
+				NewFrame.FrameID = UE::StateTree::FActiveFrameID(Storage.GenerateUniqueId());
+				const FCompactStateTreeFrame* FrameInfo = NewFrame.StateTree->GetFrameFromHandle(NewFrame.RootState);
+				ensureMsgf(FrameInfo, TEXT("The compiled data is invalid. It should contains the information for the root frame."));
+				NewFrame.ActiveTasksStatus = FrameInfo ? FStateTreeTasksCompletionStatus(*FrameInfo) : FStateTreeTasksCompletionStatus();
+
 				// Pass the linked state's parameters as global parameters to the linked asset.
 				NewFrame.GlobalParameterDataHandle = NextState.ParameterDataHandle;
 
diff --git a/Engine/Plugins/Runtime/StateTree/Source/StateTreeTestSuite/Private/StateTreeStatePathTest.cpp b/Engine/Plugins/Runtime/StateTree/Source/StateTreeTestSuite/Private/StateTreeStatePathTest.cpp
--- a/Engine/Plugins/Runtime/StateTree/Source/StateTreeTestSuite/Private/StateTreeStatePathTest.cpp
+++ b/Engine/Plugins/Runtime/StateTree/Source/StateTreeTestSuite/Private/StateTreeStatePathTest.cpp
@@ -349,7 +349,8 @@ struct FStateTreeTest_StatePath_LinkStates : FStateTreeTestBase
 			const UE::StateTree::FActiveStatePath ExecPath = InstanceData.GetExecutionState()->GetActiveStatePath();
 			AITEST_TRUE("Has the correct number of path elements", ExecPath.Num() == 9);
 			AITEST_TRUE("Has the correct number of active states", InstanceData.GetExecutionState()->ActiveFrames.Num() == 4);
-			AITEST_TRUE("The last frame changed", TestFrameId(3));
+			//@todo Temporary reverted the frameID fixed until TransitionID
+			//AITEST_TRUE("The last frame changed", TestFrameId(3));
 			Exec.LogClear();
 		}
 		{
@@ -385,7 +386,8 @@ struct FStateTreeTest_StatePath_LinkStates : FStateTreeTestBase
 			UE::StateTree::FActiveStatePath ExecPath = InstanceData.GetExecutionState()->GetActiveStatePath();
 			AITEST_TRUE("Has the correct number of path elements", ExecPath.Num() == 9);
 			AITEST_TRUE("Has the correct number of active states", InstanceData.GetExecutionState()->ActiveFrames.Num() == 4);
-			AITEST_TRUE("All frames are the same", TestFrameId(2));
+			//@todo Temporary reverted the frameID fixed until TransitionID
+			//AITEST_TRUE("All frames are the same", TestFrameId(2));
 			Exec.LogClear();
 		}
 		{
@@ -400,7 +402,8 @@ struct FStateTreeTest_StatePath_LinkStates : FStateTreeTestBase
 			Status = Exec.Tick(1.0f);
 			AITEST_EQUAL("Start should complete with Running", Status, EStateTreeRunStatus::Running);
 			AITEST_TRUE("Should be in the correct state", Exec.ExpectInActiveStates("Tree1Root", "Tree1RootState1", "Tree1State1StateSub1A", "Tree1Sub1", "Tree1Sub1StateLinkTree2B", "Tree2StateRoot", "Tree2RootStateSub1B", "Tree2Sub1", "Tree2Sub1State1")); 
-			AITEST_TRUE("All frames are the same", TestFrameId(3));
+			//@todo Temporary reverted the frameID fixed until TransitionID
+			//AITEST_TRUE("All frames are the same", TestFrameId(3));
 			Exec.LogClear();
 
 			PreTestFrameId();
