PR #13019: Use the correct tree to iterate by active state for GetDebugInfo

diff --git a/Engine/Plugins/Runtime/StateTree/Source/StateTreeModule/Private/StateTreeExecutionContext.cpp b/Engine/Plugins/Runtime/StateTree/Source/StateTreeModule/Private/StateTreeExecutionContext.cpp
--- a/Engine/Plugins/Runtime/StateTree/Source/StateTreeModule/Private/StateTreeExecutionContext.cpp
+++ b/Engine/Plugins/Runtime/StateTree/Source/StateTreeModule/Private/StateTreeExecutionContext.cpp
@@ -518,7 +518,7 @@ FString FStateTreeReadOnlyExecutionContext::GetDebugInfoString() const
 				FStateTreeStateHandle Handle = CurrentFrame.ActiveStates[Index];
 				if (Handle.IsValid())
 				{
-					const FCompactStateTreeState& State = RootStateTree.States[Handle.Index];
+					const FCompactStateTreeState& State = CurrentStateTree->States[Handle.Index];
 					DebugString << TEXT('[');
 					DebugString << State.Name;
 					DebugString << TEXT("]\n");
@@ -528,7 +528,7 @@ FString FStateTreeReadOnlyExecutionContext::GetDebugInfoString() const
 						DebugString += TEXT("\nTasks:\n");
 						for (int32 TaskIndex = State.TasksBegin; TaskIndex < (State.TasksBegin + State.TasksNum); TaskIndex++)
 						{
-							const FStateTreeTaskBase& Task = RootStateTree.Nodes[TaskIndex].Get<const FStateTreeTaskBase>();
+							const FStateTreeTaskBase& Task = CurrentStateTree->Nodes[TaskIndex].Get<const FStateTreeTaskBase>();
 							if (Task.bTaskEnabled)
 							{
 								DebugString << Task.GetDebugInfo(*this);
