PR #13154: Adding support for editing the ConsumeEventOnSelected property of the FStateTreeEventDesc struct
#rnx

diff --git a/Engine/Plugins/Runtime/StateTree/Source/StateTreeEditorModule/Private/Customizations/StateTreeEventDescDetails.cpp b/Engine/Plugins/Runtime/StateTree/Source/StateTreeEditorModule/Private/Customizations/StateTreeEventDescDetails.cpp
--- a/Engine/Plugins/Runtime/StateTree/Source/StateTreeEditorModule/Private/Customizations/StateTreeEventDescDetails.cpp
+++ b/Engine/Plugins/Runtime/StateTree/Source/StateTreeEditorModule/Private/Customizations/StateTreeEventDescDetails.cpp
@@ -230,6 +230,10 @@ void FStateTreeEventDescDetails::CustomizeChildren(TSharedRef<IPropertyHandle> P
 	{
 		PropUtils->RequestForceRefresh(); 
 	}));
+
+	TSharedPtr<IPropertyHandle> ConsumeEventOnSelectProperty = PropertyHandle->GetChildHandle(GET_MEMBER_NAME_CHECKED(FStateTreeEventDesc, bConsumeEventOnSelect));
+	check(ConsumeEventOnSelectProperty);
+	ChildBuilder.AddProperty(ConsumeEventOnSelectProperty.ToSharedRef());
 }
 
 #undef LOCTEXT_NAMESPACE
