StateTree: Do not update the bound properties, for the delay task, on Tick and ExitState. They are not needed

diff --git a/Engine/Plugins/Runtime/StateTree/Source/StateTreeModule/Private/Tasks/StateTreeDelayTask.cpp b/Engine/Plugins/Runtime/StateTree/Source/StateTreeModule/Private/Tasks/StateTreeDelayTask.cpp
--- a/Engine/Plugins/Runtime/StateTree/Source/StateTreeModule/Private/Tasks/StateTreeDelayTask.cpp
+++ b/Engine/Plugins/Runtime/StateTree/Source/StateTreeModule/Private/Tasks/StateTreeDelayTask.cpp
@@ -10,6 +10,8 @@
 FStateTreeDelayTask::FStateTreeDelayTask()
 {
 	bConsideredForScheduling = false;
+	bShouldCopyBoundPropertiesOnTick = false;
+	bShouldCopyBoundPropertiesOnExitState = false;
 }
 
 EStateTreeRunStatus FStateTreeDelayTask::EnterState(FStateTreeExecutionContext& Context, const FStateTreeTransitionResult& Transition) const
