[StateTreeDebugger] Instance descriptor is now sharable to allow tracks to rely on a shared pointer to access instance lifetime. This also improves performance and facilitates integration in RewindDebugger.

diff --git a/Engine/Plugins/Runtime/StateTree/Source/StateTreeModule/Private/Debugger/StateTreeDebugger.cpp b/Engine/Plugins/Runtime/StateTree/Source/StateTreeModule/Private/Debugger/StateTreeDebugger.cpp
--- a/Engine/Plugins/Runtime/StateTree/Source/StateTreeModule/Private/Debugger/StateTreeDebugger.cpp
+++ b/Engine/Plugins/Runtime/StateTree/Source/StateTreeModule/Private/Debugger/StateTreeDebugger.cpp
@@ -118,8 +118,6 @@ void FStateTreeDebugger::Tick(const float DeltaTime)
 			LastLiveSessionId = INDEX_NONE;
 		}
 	}
-	
-	UpdateInstances();
 
 	if (bSessionAnalysisPaused == false && StateTreeAsset.IsValid())
 	{
@@ -153,13 +151,17 @@ void FStateTreeDebugger::SyncToCurrentSessionDuration()
 
 const UE::StateTreeDebugger::FInstanceDescriptor* FStateTreeDebugger::GetInstanceDescriptor(const FStateTreeInstanceDebugId InstanceId) const
 {
-	const UE::StateTreeDebugger::FInstanceDescriptor* FoundDescriptor = InstanceDescs.FindByPredicate(
-		[InstanceId](const UE::StateTreeDebugger::FInstanceDescriptor& Descriptor)
+	if (const TraceServices::IAnalysisSession* Session = GetAnalysisSession())
+	{
+		TraceServices::FAnalysisSessionReadScope SessionReadScope(*Session);
+
+		if (const IStateTreeTraceProvider* Provider = Session->ReadProvider<IStateTreeTraceProvider>(FStateTreeTraceProvider::ProviderName))
 		{
-			return Descriptor.Id == InstanceId;
-		});
+			return Provider->GetInstanceDescriptor(InstanceId).Get();
+		}
+	}
 
-	return FoundDescriptor;
+	return nullptr;
 }
 
 FText FStateTreeDebugger::GetInstanceName(const FStateTreeInstanceDebugId InstanceId) const
@@ -193,8 +195,10 @@ void FStateTreeDebugger::SelectInstance(const FStateTreeInstanceDebugId Instance
 	}
 }
 
+// Deprecated
 void FStateTreeDebugger::GetSessionInstances(TArray<UE::StateTreeDebugger::FInstanceDescriptor>& OutInstances) const
 {
+PRAGMA_DISABLE_DEPRECATION_WARNINGS
 	if (const TraceServices::IAnalysisSession* Session = GetAnalysisSession())
 	{
 		TraceServices::FAnalysisSessionReadScope SessionReadScope(*Session);
@@ -204,9 +208,10 @@ void FStateTreeDebugger::GetSessionInstances(TArray<UE::StateTreeDebugger::FInst
 			Provider->GetInstances(OutInstances);
 		}
 	}
+PRAGMA_ENABLE_DEPRECATION_WARNINGS
 }
 
-void FStateTreeDebugger::UpdateInstances()
+void FStateTreeDebugger::GetSessionInstanceDescriptors(TArray<const TSharedRef<const UE::StateTreeDebugger::FInstanceDescriptor>>& OutInstances) const
 {
 	if (const TraceServices::IAnalysisSession* Session = GetAnalysisSession())
 	{
@@ -214,7 +219,7 @@ void FStateTreeDebugger::UpdateInstances()
 
 		if (const IStateTreeTraceProvider* Provider = Session->ReadProvider<IStateTreeTraceProvider>(FStateTreeTraceProvider::ProviderName))
 		{
-			Provider->GetInstances(InstanceDescs);
+			Provider->GetInstances(OutInstances);
 		}
 	}
 }
@@ -478,11 +483,14 @@ void FStateTreeDebugger::SetScrubTime(const double ScrubTime)
 
 bool FStateTreeDebugger::IsActiveInstance(const double Time, const FStateTreeInstanceDebugId InstanceId) const
 {
-	for (const UE::StateTreeDebugger::FInstanceDescriptor& Desc : InstanceDescs)
+	if (const TraceServices::IAnalysisSession* Session = GetAnalysisSession())
 	{
-		if (Desc.Id == InstanceId && Desc.Lifetime.Contains(Time))
+		TraceServices::FAnalysisSessionReadScope SessionReadScope(*Session);
+
+		if (const IStateTreeTraceProvider* Provider = Session->ReadProvider<IStateTreeTraceProvider>(FStateTreeTraceProvider::ProviderName))
 		{
-			return true;
+			const TSharedPtr<const UE::StateTreeDebugger::FInstanceDescriptor> Descriptor = Provider->GetInstanceDescriptor(InstanceId);
+			return Descriptor.IsValid() && Descriptor->Lifetime.Contains(Time);
 		}
 	}
 	return false;
diff --git a/Engine/Plugins/Runtime/StateTree/Source/StateTreeModule/Private/Debugger/StateTreeTraceProvider.cpp b/Engine/Plugins/Runtime/StateTree/Source/StateTreeModule/Private/Debugger/StateTreeTraceProvider.cpp
--- a/Engine/Plugins/Runtime/StateTree/Source/StateTreeModule/Private/Debugger/StateTreeTraceProvider.cpp
+++ b/Engine/Plugins/Runtime/StateTree/Source/StateTreeModule/Private/Debugger/StateTreeTraceProvider.cpp
@@ -53,9 +53,9 @@ bool FStateTreeTraceProvider::ReadTimelines(const UStateTree& StateTree, TFuncti
 		check(EventsTimelines.IsValidIndex(It.Value()));
 		check(Descriptors.Num() == EventsTimelines.Num());
 
-		if (Descriptors[It.Value()].StateTree == &StateTree)
+		if (Descriptors[It.Value()].Get().StateTree == &StateTree)
 		{
-			Callback(Descriptors[It.Value()].Id, *EventsTimelines[It.Value()]);
+			Callback(Descriptors[It.Value()].Get().Id, *EventsTimelines[It.Value()]);
 		}
 	}
 
@@ -92,7 +92,11 @@ void FStateTreeTraceProvider::AppendInstanceEvent(
 		{
 			if (const UStateTree* StateTree = WeakStateTree.Get())
 			{
-				Descriptors.Emplace(StateTree, InInstanceId, InInstanceName, TRange<double>(InWorldRecordingTime, std::numeric_limits<double>::max()));
+				Descriptors.Emplace(MakeShared<UE::StateTreeDebugger::FInstanceDescriptor>(
+					StateTree
+					, InInstanceId
+					, InInstanceName
+					, TRange<double>(InWorldRecordingTime, std::numeric_limits<double>::max())));
 
 				check(InstanceIdToDebuggerEntryTimelines.Find(InInstanceId) == nullptr);
 				InstanceIdToDebuggerEntryTimelines.Add(InInstanceId, EventsTimelines.Num());
@@ -115,7 +119,7 @@ void FStateTreeTraceProvider::AppendInstanceEvent(
 		if (const uint32* Index = InstanceIdToDebuggerEntryTimelines.Find(InInstanceId))
 		{
 			check(Descriptors.IsValidIndex(*Index));
-			Descriptors[*Index].Lifetime.SetUpperBound(InWorldRecordingTime);
+			Descriptors[*Index].Get().Lifetime.SetUpperBound(InWorldRecordingTime);
 		}
 	}
 
@@ -150,16 +154,30 @@ bool FStateTreeTraceProvider::GetAssetFromInstanceId(const FStateTreeInstanceDeb
 	if (const uint32* IndexPtr = InstanceIdToDebuggerEntryTimelines.Find(InstanceId))
 	{
 		check(Descriptors.Num() == EventsTimelines.Num());
-		WeakStateTree = Descriptors[*IndexPtr].StateTree;
+		WeakStateTree = Descriptors[*IndexPtr].Get().StateTree;
 		return true;
 	}
 
 	return false;
 }
 
-void FStateTreeTraceProvider::GetInstances(TArray<UE::StateTreeDebugger::FInstanceDescriptor>& OutInstances) const
+TSharedPtr<const UE::StateTreeDebugger::FInstanceDescriptor> FStateTreeTraceProvider::GetInstanceDescriptor(const FStateTreeInstanceDebugId InstanceId) const
 {
-	OutInstances = Descriptors;
+	const TSharedRef<UE::StateTreeDebugger::FInstanceDescriptor>* Descriptor = Descriptors.FindByPredicate([InstanceId](const TSharedRef<const UE::StateTreeDebugger::FInstanceDescriptor>& Descriptor)
+	{
+		return Descriptor.Get().Id == InstanceId;
+	});
+
+	return Descriptor ? Descriptor->ToSharedPtr() : TSharedPtr<const UE::StateTreeDebugger::FInstanceDescriptor>{};
+}
+
+void FStateTreeTraceProvider::GetInstances(TArray<const TSharedRef<const UE::StateTreeDebugger::FInstanceDescriptor>>& OutInstances) const
+{
+	OutInstances.Reserve(Descriptors.Num());
+	Algo::Transform(Descriptors, OutInstances, [](const TSharedRef<UE::StateTreeDebugger::FInstanceDescriptor>& Descriptor)
+	{
+		return Descriptor;
+	});
 }
 
 #undef LOCTEXT_NAMESPACE
diff --git a/Engine/Plugins/Runtime/StateTree/Source/StateTreeModule/Public/Debugger/IStateTreeTraceProvider.h b/Engine/Plugins/Runtime/StateTree/Source/StateTreeModule/Public/Debugger/IStateTreeTraceProvider.h
--- a/Engine/Plugins/Runtime/StateTree/Source/StateTreeModule/Public/Debugger/IStateTreeTraceProvider.h
+++ b/Engine/Plugins/Runtime/StateTree/Source/StateTreeModule/Public/Debugger/IStateTreeTraceProvider.h
@@ -17,11 +17,23 @@ class IStateTreeTraceProvider : public TraceServices::IProvider
 public:
 	typedef TraceServices::ITimeline<FStateTreeTraceEventVariantType> FEventsTimeline;
 
+	/**
+	 * Return instance descriptor associated to the given instance id.
+	 * @param InstanceId Id of a specific instance to get the descriptor for.
+	 * @return Shared pointer to the descriptor if the specified instance was found.
+	 */
+	virtual TSharedPtr<const UE::StateTreeDebugger::FInstanceDescriptor> GetInstanceDescriptor(const FStateTreeInstanceDebugId InstanceId) const = 0;
+
 	/**
 	 * Return all instances with events in the traces.
 	 * @param OutInstances List of all instances with events in the traces regardless if they are active or not.
 	 */
-	virtual void GetInstances(TArray<UE::StateTreeDebugger::FInstanceDescriptor>& OutInstances) const = 0;
+	virtual void GetInstances(TArray<const TSharedRef<const UE::StateTreeDebugger::FInstanceDescriptor>>& OutInstances) const = 0;
+
+	UE_DEPRECATED(5.7, "Use the version with TSharedRef array instead")
+	virtual void GetInstances(TArray<UE::StateTreeDebugger::FInstanceDescriptor>& OutInstances) const final
+	{
+	}
 
 	/**
 	 * Execute given function receiving an event timeline for a given instance or all timelines if instance not specified.  
diff --git a/Engine/Plugins/Runtime/StateTree/Source/StateTreeModule/Public/Debugger/StateTreeDebugger.h b/Engine/Plugins/Runtime/StateTree/Source/StateTreeModule/Public/Debugger/StateTreeDebugger.h
--- a/Engine/Plugins/Runtime/StateTree/Source/StateTreeModule/Public/Debugger/StateTreeDebugger.h
+++ b/Engine/Plugins/Runtime/StateTree/Source/StateTreeModule/Public/Debugger/StateTreeDebugger.h
@@ -14,6 +14,8 @@
 #include "TraceServices/Model/Diagnostics.h"
 #include "TraceServices/Model/Frames.h"
 
+#define UE_API STATETREEMODULE_API
+
 enum class EStateTreeBreakpointType : uint8;
 
 namespace UE::Trace
@@ -37,11 +39,24 @@ struct FStateTreeDebugger : FTickableGameObject
 	struct FTraceDescriptor
 	{
 		FTraceDescriptor() = default;
-		FTraceDescriptor(const FString& Name, const uint32 Id) : Name(Name), TraceId(Id) {}
+		FTraceDescriptor(const FString& Name, const uint32 Id) : Name(Name), TraceId(Id)
+		{
+		}
 		
-		bool operator==(const FTraceDescriptor& Other) const { return Other.TraceId == TraceId; }
-		bool operator!=(const FTraceDescriptor& Other) const { return !(Other == *this); }
-		bool IsValid() const { return TraceId != INDEX_NONE; }
+		bool operator==(const FTraceDescriptor& Other) const
+		{
+			return Other.TraceId == TraceId;
+		}
+
+		bool operator!=(const FTraceDescriptor& Other) const
+		{
+			return !(Other == *this);
+		}
+
+		bool IsValid() const
+		{
+			return TraceId != INDEX_NONE;
+		}
 
 		FString Name;
 		uint32 TraceId = INDEX_NONE;
@@ -51,7 +66,10 @@ struct FStateTreeDebugger : FTickableGameObject
 
 	struct FHitBreakpoint
 	{
-		bool IsSet() const { return Index != INDEX_NONE; }
+		bool IsSet() const
+		{
+			return Index != INDEX_NONE;
+		}
 		void Reset()
 		{
 			InstanceId = FStateTreeInstanceDebugId::Invalid;
@@ -72,49 +90,78 @@ struct FStateTreeDebugger : FTickableGameObject
 		int32 Index = INDEX_NONE;
 	};
 
-	STATETREEMODULE_API FStateTreeDebugger();
-	STATETREEMODULE_API virtual ~FStateTreeDebugger() override;
+	UE_API FStateTreeDebugger();
+	UE_API virtual ~FStateTreeDebugger() override;
+
+	const UStateTree* GetAsset() const
+	{
+		return StateTreeAsset.Get();
+	}
 
-	const UStateTree* GetAsset() const { return StateTreeAsset.Get(); }
-	void SetAsset(const UStateTree* InStateTreeAsset) { StateTreeAsset = InStateTreeAsset; }
+	void SetAsset(const UStateTree* InStateTreeAsset)
+	{
+		StateTreeAsset = InStateTreeAsset;
+	}
 	
 	/** Forces a single refresh to latest state. Useful when simulation is paused. */
-	STATETREEMODULE_API void SyncToCurrentSessionDuration();
+	UE_API void SyncToCurrentSessionDuration();
 	
-	STATETREEMODULE_API bool CanStepBackToPreviousStateWithEvents() const;
-	STATETREEMODULE_API void StepBackToPreviousStateWithEvents();
-
-	STATETREEMODULE_API bool CanStepForwardToNextStateWithEvents() const;
-	STATETREEMODULE_API void StepForwardToNextStateWithEvents();
-
-	STATETREEMODULE_API bool CanStepBackToPreviousStateChange() const;
-	STATETREEMODULE_API void StepBackToPreviousStateChange();
-
-	STATETREEMODULE_API bool CanStepForwardToNextStateChange() const;
-	STATETREEMODULE_API void StepForwardToNextStateChange();
-
-	STATETREEMODULE_API bool IsActiveInstance(double Time, FStateTreeInstanceDebugId InstanceId) const;
-	STATETREEMODULE_API FText GetInstanceName(FStateTreeInstanceDebugId InstanceId) const;
-	STATETREEMODULE_API FText GetInstanceDescription(FStateTreeInstanceDebugId InstanceId) const;
-	STATETREEMODULE_API void SelectInstance(const FStateTreeInstanceDebugId InstanceId);
-	void ClearSelection() { SelectInstance({}); }
-	FStateTreeInstanceDebugId GetSelectedInstanceId() const { return SelectedInstanceId; }
-	STATETREEMODULE_API const UE::StateTreeDebugger::FInstanceDescriptor* GetInstanceDescriptor(const FStateTreeInstanceDebugId InstanceId) const ;
-	const UE::StateTreeDebugger::FInstanceDescriptor* GetSelectedInstanceDescriptor() const { return GetInstanceDescriptor(SelectedInstanceId); }
-
-	STATETREEMODULE_API bool HasStateBreakpoint(FStateTreeStateHandle StateHandle, EStateTreeBreakpointType BreakpointType) const;
-	STATETREEMODULE_API bool HasTaskBreakpoint(FStateTreeIndex16 Index, EStateTreeBreakpointType BreakpointType) const;
-	STATETREEMODULE_API bool HasTransitionBreakpoint(FStateTreeIndex16 Index, EStateTreeBreakpointType BreakpointType) const;
-	STATETREEMODULE_API void SetStateBreakpoint(FStateTreeStateHandle StateHandle, EStateTreeBreakpointType BreakpointType);
-	STATETREEMODULE_API void SetTransitionBreakpoint(FStateTreeIndex16 SubIndex, EStateTreeBreakpointType BreakpointType);
-	STATETREEMODULE_API void SetTaskBreakpoint(FStateTreeIndex16 NodeIndex, EStateTreeBreakpointType BreakpointType);
-	STATETREEMODULE_API void ClearBreakpoint(FStateTreeIndex16 NodeIndex, EStateTreeBreakpointType BreakpointType);
-	STATETREEMODULE_API void ClearAllBreakpoints();
-	int32 NumBreakpoints() const { return Breakpoints.Num(); }
-	bool HasHitBreakpoint() const { return HitBreakpoint.IsSet(); }
-
-	static STATETREEMODULE_API FText DescribeTrace(const FTraceDescriptor& TraceDescriptor);
-	static STATETREEMODULE_API FText DescribeInstance(const UE::StateTreeDebugger::FInstanceDescriptor& StateTreeInstanceDesc);
+	UE_API bool CanStepBackToPreviousStateWithEvents() const;
+	UE_API void StepBackToPreviousStateWithEvents();
+
+	UE_API bool CanStepForwardToNextStateWithEvents() const;
+	UE_API void StepForwardToNextStateWithEvents();
+
+	UE_API bool CanStepBackToPreviousStateChange() const;
+	UE_API void StepBackToPreviousStateChange();
+
+	UE_API bool CanStepForwardToNextStateChange() const;
+	UE_API void StepForwardToNextStateChange();
+
+	UE_API bool IsActiveInstance(double Time, FStateTreeInstanceDebugId InstanceId) const;
+
+	UE_API FText GetInstanceName(FStateTreeInstanceDebugId InstanceId) const;
+	UE_API FText GetInstanceDescription(FStateTreeInstanceDebugId InstanceId) const;
+
+	UE_API void SelectInstance(const FStateTreeInstanceDebugId InstanceId);
+
+	void ClearSelection()
+	{
+		SelectInstance({});
+	}
+
+	FStateTreeInstanceDebugId GetSelectedInstanceId() const
+	{
+		return SelectedInstanceId;
+	}
+
+	UE_API const UE::StateTreeDebugger::FInstanceDescriptor* GetInstanceDescriptor(const FStateTreeInstanceDebugId InstanceId) const ;
+	const UE::StateTreeDebugger::FInstanceDescriptor* GetSelectedInstanceDescriptor() const
+	{
+		return GetInstanceDescriptor(SelectedInstanceId);
+	}
+
+	UE_API bool HasStateBreakpoint(FStateTreeStateHandle StateHandle, EStateTreeBreakpointType BreakpointType) const;
+	UE_API bool HasTaskBreakpoint(FStateTreeIndex16 Index, EStateTreeBreakpointType BreakpointType) const;
+	UE_API bool HasTransitionBreakpoint(FStateTreeIndex16 Index, EStateTreeBreakpointType BreakpointType) const;
+	UE_API void SetStateBreakpoint(FStateTreeStateHandle StateHandle, EStateTreeBreakpointType BreakpointType);
+	UE_API void SetTransitionBreakpoint(FStateTreeIndex16 SubIndex, EStateTreeBreakpointType BreakpointType);
+	UE_API void SetTaskBreakpoint(FStateTreeIndex16 NodeIndex, EStateTreeBreakpointType BreakpointType);
+	UE_API void ClearBreakpoint(FStateTreeIndex16 NodeIndex, EStateTreeBreakpointType BreakpointType);
+	UE_API void ClearAllBreakpoints();
+
+	int32 NumBreakpoints() const
+	{
+		return Breakpoints.Num();
+	}
+
+	bool HasHitBreakpoint() const
+	{
+		return HitBreakpoint.IsSet();
+	}
+
+	static UE_API FText DescribeTrace(const FTraceDescriptor& TraceDescriptor);
+	static UE_API FText DescribeInstance(const UE::StateTreeDebugger::FInstanceDescriptor& StateTreeInstanceDesc);
 
 	/**
 	 * Finds and returns the event collection associated to a given instance Id.
@@ -122,28 +169,38 @@ struct FStateTreeDebugger : FTickableGameObject
 	 * @param InstanceId Id of the instance for which the event collection is returned. 
 	 * @return Event collection associated to the provided Id or an invalid one if not found.
 	 */
-	STATETREEMODULE_API const UE::StateTreeDebugger::FInstanceEventCollection& GetEventCollection(FStateTreeInstanceDebugId InstanceId) const;
+	UE_API const UE::StateTreeDebugger::FInstanceEventCollection& GetEventCollection(FStateTreeInstanceDebugId InstanceId) const;
 
 	/** Clears events from all instances. */
-	STATETREEMODULE_API void ResetEventCollections();
+	UE_API void ResetEventCollections();
 
 	/** Returns the recording duration in world recorded time. */
-	double GetRecordingDuration() const { return RecordingDuration; }
+	double GetRecordingDuration() const
+	{
+		return RecordingDuration;
+	}
 
 	/** Returns the duration of the analysis session. This is not related to the world simulation time. */
-	double GetAnalysisDuration() const { return AnalysisDuration; }
+	double GetAnalysisDuration() const
+	{
+		return AnalysisDuration;
+	}
 
 	/** Returns the time (based on the recording duration) associated to the selected frame. */
-	double GetScrubTime() const	{ return ScrubState.GetScrubTime(); }
-	STATETREEMODULE_API void SetScrubTime(double ScrubTime);
+	double GetScrubTime() const
+	{
+		return ScrubState.GetScrubTime();
+	}
 
-	STATETREEMODULE_API void GetLiveTraces(TArray<FTraceDescriptor>& OutTraceDescriptors) const;
+	UE_API void SetScrubTime(double ScrubTime);
+
+	UE_API void GetLiveTraces(TArray<FTraceDescriptor>& OutTraceDescriptors) const;
 
 	/**
 	 * Queue a request to auto start an analysis session on the next available live trace.
 	 * @return True if connection was successfully requested or was able to use active trace, false otherwise.
 	 */
-	STATETREEMODULE_API bool RequestAnalysisOfEditorSession();
+	UE_API bool RequestAnalysisOfEditorSession();
 	bool IsAnalyzingEditorSession() const
 	{
 		return AnalysisTransitionType == EAnalysisTransitionType::NoneToEditor
@@ -157,9 +214,17 @@ struct FStateTreeDebugger : FTickableGameObject
 			|| AnalysisTransitionType == EAnalysisTransitionType::EditorToEditor;
 	}
 
-	bool IsAnalysisSessionActive() const { return GetAnalysisSession() != nullptr; }
-	bool IsAnalysisSessionPaused() const { return bSessionAnalysisPaused; }
-	STATETREEMODULE_API const TraceServices::IAnalysisSession* GetAnalysisSession() const;
+	bool IsAnalysisSessionActive() const
+	{
+		return GetAnalysisSession() != nullptr;
+	}
+
+	bool IsAnalysisSessionPaused() const
+	{
+		return bSessionAnalysisPaused;
+	}
+
+	UE_API const TraceServices::IAnalysisSession* GetAnalysisSession() const;
 
 	/**
 	 * Tries to start an analysis for a given trace descriptor.
@@ -167,18 +232,31 @@ struct FStateTreeDebugger : FTickableGameObject
 	 * @param TraceDescriptor Descriptor of the trace that needs to be analyzed
 	 * @return True if analysis was successfully started, false otherwise.
 	 */
-	STATETREEMODULE_API bool RequestSessionAnalysis(const FTraceDescriptor& TraceDescriptor);
-	void PauseSessionAnalysis()  { bSessionAnalysisPaused = true; }
+	UE_API bool RequestSessionAnalysis(const FTraceDescriptor& TraceDescriptor);
+
+	void PauseSessionAnalysis()
+	{
+		bSessionAnalysisPaused = true;
+	}
+
 	void ResumeSessionAnalysis()
 	{
 		bSessionAnalysisPaused = false;
 		HitBreakpoint.Reset();
 	}
-	STATETREEMODULE_API void StopSessionAnalysis();
-	FTraceDescriptor GetSelectedTraceDescriptor() const { return ActiveSessionTraceDescriptor; }
-	STATETREEMODULE_API FText GetSelectedTraceDescription() const;
 
-	STATETREEMODULE_API void GetSessionInstances(TArray<UE::StateTreeDebugger::FInstanceDescriptor>& OutInstances) const;
+	UE_API void StopSessionAnalysis();
+
+	FTraceDescriptor GetSelectedTraceDescriptor() const
+	{
+		return ActiveSessionTraceDescriptor;
+	}
+
+	UE_API FText GetSelectedTraceDescription() const;
+
+	UE_DEPRECATED(5.7, "Use GetSessionInstanceDescriptors instead")
+	UE_API void GetSessionInstances(TArray<UE::StateTreeDebugger::FInstanceDescriptor>& OutInstances) const;
+	UE_API void GetSessionInstanceDescriptors(TArray<const TSharedRef<const UE::StateTreeDebugger::FInstanceDescriptor>>& OutInstances) const;
 
 	FOnStateTreeDebuggerNewSession OnNewSession;
 	FOnStateTreeDebuggerNewInstance OnNewInstance;
@@ -188,11 +266,27 @@ struct FStateTreeDebugger : FTickableGameObject
 	FOnStateTreeDebuggerActiveStatesChanges OnActiveStatesChanged;
 
 protected:
-	STATETREEMODULE_API virtual void Tick(float DeltaTime) override;
-	virtual bool IsTickable() const override { return true; }
-	virtual bool IsTickableWhenPaused() const override { return true; }
-	virtual bool IsTickableInEditor() const override { return true; }
-	virtual TStatId GetStatId() const override { RETURN_QUICK_DECLARE_CYCLE_STAT(FStateTreeDebugger, STATGROUP_Tickables); }
+	virtual void Tick(float DeltaTime) override;
+
+	virtual bool IsTickable() const override
+	{
+		return true;
+	}
+
+	virtual bool IsTickableWhenPaused() const override
+	{
+		return true;
+	}
+
+	virtual bool IsTickableInEditor() const override
+	{
+		return true;
+	}
+
+	virtual TStatId GetStatId() const override
+	{
+		RETURN_QUICK_DECLARE_CYCLE_STAT(FStateTreeDebugger, STATGROUP_Tickables);
+	}
 	
 private:
 	enum class EAnalysisSourceType : uint8
@@ -215,11 +309,11 @@ private:
 		SelectedToEditor,
 	};
 
-	STATETREEMODULE_API void UpdateAnalysisTransitionType(EAnalysisSourceType SourceType);
+	void UpdateAnalysisTransitionType(EAnalysisSourceType SourceType);
 	
-	STATETREEMODULE_API void ReadTrace(double ScrubTime);
-	STATETREEMODULE_API void ReadTrace(uint64 FrameIndex);
-	STATETREEMODULE_API void ReadTrace(
+	void ReadTrace(double ScrubTime);
+	void ReadTrace(uint64 FrameIndex);
+	void ReadTrace(
 		const TraceServices::IAnalysisSession& Session,
 		const TraceServices::IFrameProvider& FrameProvider,
 		const TraceServices::FFrame& Frame
@@ -231,17 +325,17 @@ private:
 	 * @param Event The event received from the instance.
 	 * @return True if a breakpoint has been it, false otherwise.
 	 */
-	STATETREEMODULE_API bool EvaluateBreakpoints(FStateTreeInstanceDebugId InstanceId, const FStateTreeTraceEventVariantType& Event);
+	bool EvaluateBreakpoints(FStateTreeInstanceDebugId InstanceId, const FStateTreeTraceEventVariantType& Event);
 
-	STATETREEMODULE_API void SendNotifications();
+	void SendNotifications();
 
-	STATETREEMODULE_API void SetActiveStates(const FStateTreeTraceActiveStates& NewActiveStates);
+	void SetActiveStates(const FStateTreeTraceActiveStates& NewActiveStates);
 
 	/**
 	 * Request an analysis session on the latest next available live trace.
 	 * This will replace the current analysis session if any.
 	 */
-	STATETREEMODULE_API void RequestAnalysisOfLatestTrace();
+	void RequestAnalysisOfLatestTrace();
 
 	/**
 	 * Looks for a new live traces to start an analysis session.
@@ -249,25 +343,23 @@ private:
 	 * @param RetryPollingDuration - On failure, how many seconds to retry connecting during FStateTreeDebugger::Tick
 	 * @return True if the analysis was successfully started; false otherwise.
 	 */
-	STATETREEMODULE_API bool TryStartNewLiveSessionAnalysis(float RetryPollingDuration);
+	bool TryStartNewLiveSessionAnalysis(float RetryPollingDuration);
 
-	STATETREEMODULE_API bool StartSessionAnalysis(const FTraceDescriptor& TraceDescriptor);
+	bool StartSessionAnalysis(const FTraceDescriptor& TraceDescriptor);
 
-	STATETREEMODULE_API void SetScrubStateCollectionIndex(int32 EventCollectionIndex);
+	void SetScrubStateCollectionIndex(int32 EventCollectionIndex);
 
 	/**
 	 * Recompute index of the span that contains the active states change event and update the active states.
 	 * This method handles unselected instances in which case it will reset the active states and set the span index to INDEX_NONE
 	 * */
-	STATETREEMODULE_API void RefreshActiveStates();
-
-	STATETREEMODULE_API UE::Trace::FStoreClient* GetStoreClient() const;
+	void RefreshActiveStates();
 
-	STATETREEMODULE_API void UpdateInstances();
+	UE::Trace::FStoreClient* GetStoreClient() const;
 
-	STATETREEMODULE_API bool ProcessEvent(const FStateTreeInstanceDebugId InstanceId, const TraceServices::FFrame& Frame, const FStateTreeTraceEventVariantType& Event);
-	STATETREEMODULE_API void AddEvents(double StartTime, double EndTime, const TraceServices::IFrameProvider& FrameProvider, const IStateTreeTraceProvider& StateTreeTraceProvider);
-	STATETREEMODULE_API void UpdateMetadata(FTraceDescriptor& TraceDescriptor) const;
+	bool ProcessEvent(const FStateTreeInstanceDebugId InstanceId, const TraceServices::FFrame& Frame, const FStateTreeTraceEventVariantType& Event);
+	void AddEvents(double StartTime, double EndTime, const TraceServices::IFrameProvider& FrameProvider, const IStateTreeTraceProvider& StateTreeTraceProvider);
+	void UpdateMetadata(FTraceDescriptor& TraceDescriptor) const;
 
 	/** Module used to access the store client and analysis sessions .*/
 	IStateTreeModule& StateTreeModule;
@@ -281,9 +373,6 @@ private:
 	/** Descriptor of the currently selected session */
 	FTraceDescriptor ActiveSessionTraceDescriptor;
 
-	/** Descriptors for all instances of the StateTree asset that have traces in the analysis session and still active. */
-	TArray<UE::StateTreeDebugger::FInstanceDescriptor> InstanceDescs;
-
 	/** Processed events for each instance. */
 	TArray<UE::StateTreeDebugger::FInstanceEventCollection> EventCollections;
 
@@ -348,4 +437,6 @@ private:
 	FDelegateHandle TracingTimelineScrubbedHandle;
 };
 
+#undef UE_API
+
 #endif // WITH_STATETREE_TRACE_DEBUGGER
diff --git a/Engine/Plugins/Runtime/StateTree/Source/StateTreeModule/Public/Debugger/StateTreeDebuggerTypes.h b/Engine/Plugins/Runtime/StateTree/Source/StateTreeModule/Public/Debugger/StateTreeDebuggerTypes.h
--- a/Engine/Plugins/Runtime/StateTree/Source/StateTreeModule/Public/Debugger/StateTreeDebuggerTypes.h
+++ b/Engine/Plugins/Runtime/StateTree/Source/StateTreeModule/Public/Debugger/StateTreeDebuggerTypes.h
@@ -60,7 +60,7 @@ struct FFrameSpan
 /**
  * Struct describing a state tree instance for a given StateTree asset
  */
-struct FInstanceDescriptor
+struct FInstanceDescriptor : TSharedFromThis<FInstanceDescriptor>
 {
 	FInstanceDescriptor() = default;
 	STATETREEMODULE_API FInstanceDescriptor(const UStateTree* InStateTree, const FStateTreeInstanceDebugId InId, const FString& InName, const TRange<double> InLifetime);
diff --git a/Engine/Plugins/Runtime/StateTree/Source/StateTreeModule/Public/Debugger/StateTreeTraceProvider.h b/Engine/Plugins/Runtime/StateTree/Source/StateTreeModule/Public/Debugger/StateTreeTraceProvider.h
--- a/Engine/Plugins/Runtime/StateTree/Source/StateTreeModule/Public/Debugger/StateTreeTraceProvider.h
+++ b/Engine/Plugins/Runtime/StateTree/Source/StateTreeModule/Public/Debugger/StateTreeTraceProvider.h
@@ -34,8 +34,9 @@ public:
 	bool GetAssetFromInstanceId(const FStateTreeInstanceDebugId InstanceId, TWeakObjectPtr<const UStateTree>& WeakStateTree) const;
 
 protected:
-	/** IStateTreeDebuggerProvider interface */
-	virtual void GetInstances(TArray<UE::StateTreeDebugger::FInstanceDescriptor>& OutInstances) const override;
+	//~ IStateTreeDebuggerProvider interface
+	virtual TSharedPtr<const UE::StateTreeDebugger::FInstanceDescriptor> GetInstanceDescriptor(const FStateTreeInstanceDebugId InstanceId) const override;
+	virtual void GetInstances(TArray<const TSharedRef<const UE::StateTreeDebugger::FInstanceDescriptor>>& OutInstances) const override;
 	virtual bool ReadTimelines(const FStateTreeInstanceDebugId InstanceId, TFunctionRef<void(const FStateTreeInstanceDebugId ProcessedInstanceId, const FEventsTimeline&)> Callback) const override;
 	virtual bool ReadTimelines(const UStateTree& StateTree, TFunctionRef<void(const FStateTreeInstanceDebugId ProcessedInstanceId, const FEventsTimeline&)> Callback) const override;
 
@@ -43,7 +44,7 @@ private:
 	TraceServices::IAnalysisSession& Session;
 
 	TMap<FStateTreeInstanceDebugId, uint32> InstanceIdToDebuggerEntryTimelines;
-	TArray<UE::StateTreeDebugger::FInstanceDescriptor> Descriptors;
+	TArray<TSharedRef<UE::StateTreeDebugger::FInstanceDescriptor>> Descriptors;
 	TArray<TSharedRef<TraceServices::TPointTimeline<FStateTreeTraceEventVariantType>>> EventsTimelines;
 
 	struct FStateTreeDebugIdPair
