StateTree: Update the delete key input for mac.

diff --git a/Engine/Plugins/Runtime/StateTree/Source/StateTreeEditorModule/Private/StateTreeEditorCommands.cpp b/Engine/Plugins/Runtime/StateTree/Source/StateTreeEditorModule/Private/StateTreeEditorCommands.cpp
--- a/Engine/Plugins/Runtime/StateTree/Source/StateTreeEditorModule/Private/StateTreeEditorCommands.cpp
+++ b/Engine/Plugins/Runtime/StateTree/Source/StateTreeEditorModule/Private/StateTreeEditorCommands.cpp
@@ -30,7 +30,7 @@ void FStateTreeEditorCommands::RegisterCommands()
 	UI_COMMAND(PasteStatesAsSiblings, "Paste As Siblings", "Paste States as Siblings", EUserInterfaceActionType::Button, FInputChord(EModifierKey::Control, EKeys::V));
 	UI_COMMAND(PasteStatesAsChildren, "Paste As Children", "Paste States as Children", EUserInterfaceActionType::Button, FInputChord(EModifierKey::Control | EModifierKey::Shift, EKeys::V));
 	UI_COMMAND(DuplicateStates, "Duplicate", "Duplicate Selected States", EUserInterfaceActionType::Button, FInputChord(EModifierKey::Control, EKeys::D));
-	UI_COMMAND(DeleteStates, "Delete", "Delete Selected States", EUserInterfaceActionType::Button, FInputChord(EKeys::Delete));
+	UI_COMMAND(DeleteStates, "Delete", "Delete Selected States", EUserInterfaceActionType::Button, FInputChord(EKeys::Platform_Delete));
 	UI_COMMAND(RenameState, "Rename", "Rename Selected State", EUserInterfaceActionType::Button, FInputChord(EKeys::F2));
 	UI_COMMAND(EnableStates, "State Enabled", "Enables Selected States", EUserInterfaceActionType::Check, FInputChord());
 
