StateTree: Update the icons for task completion.

#rnx

diff --git a/Engine/Plugins/Runtime/StateTree/Source/StateTreeEditorModule/Private/Customizations/StateTreeEditorNodeDetails.cpp b/Engine/Plugins/Runtime/StateTree/Source/StateTreeEditorModule/Private/Customizations/StateTreeEditorNodeDetails.cpp
--- a/Engine/Plugins/Runtime/StateTree/Source/StateTreeEditorModule/Private/Customizations/StateTreeEditorNodeDetails.cpp
+++ b/Engine/Plugins/Runtime/StateTree/Source/StateTreeEditorModule/Private/Customizations/StateTreeEditorNodeDetails.cpp
@@ -415,7 +415,8 @@ void FStateTreeEditorNodeDetails::CustomizeHeader(TSharedRef<class IPropertyHand
 			// Border to capture mouse clicks on the row (used for right click menu).
 			SAssignNew(RowBorder, SBorder)
 			.BorderImage(FStyleDefaults::GetNoBrush())
-			.Padding(0)
+			.Padding(0.0f)
+			.ForegroundColor(this, &FStateTreeEditorNodeDetails::GetContentRowColor)
 			.OnMouseButtonDown(this, &FStateTreeEditorNodeDetails::OnRowMouseDown)
 			.OnMouseButtonUp(this, &FStateTreeEditorNodeDetails::OnRowMouseUp)
 			[
@@ -475,7 +476,7 @@ void FStateTreeEditorNodeDetails::CustomizeHeader(TSharedRef<class IPropertyHand
 				[
 					SNew(SBox)
 					.WidthOverride(30.0f)
-					.Padding(FMargin(2, 4, 2, 3))
+					.Padding(FMargin(2.0f, 4.0f, 2.0f, 3.0f))
 					.VAlign(VAlign_Center)
 					.Visibility(this, &FStateTreeEditorNodeDetails::IsOperandVisible)
 					[
@@ -513,7 +514,7 @@ void FStateTreeEditorNodeDetails::CustomizeHeader(TSharedRef<class IPropertyHand
 				+ SHorizontalBox::Slot()
 				.FillContentWidth(0.0f, 1.0f) // no growing, allow shrink
 				.VAlign(VAlign_Center)
-				.Padding(FMargin(0, 4))
+				.Padding(FMargin(0.0f, 4.0f))
 				[
 					SAssignNew(DescriptionBox, SHorizontalBox)
 					.Clipping(EWidgetClipping::ClipToBoundsAlways)
@@ -522,7 +523,7 @@ void FStateTreeEditorNodeDetails::CustomizeHeader(TSharedRef<class IPropertyHand
 					+ SHorizontalBox::Slot()
 					.AutoWidth()
 					.VAlign(VAlign_Center)
-					.Padding(0, 0, 0, 0)
+					.Padding(0.0f, 0.0f, 0.0f, 0.0f)
 					[
 						SNew(SImage)
 						.Image(this, &FStateTreeEditorNodeDetails::GetIcon)
@@ -534,14 +535,14 @@ void FStateTreeEditorNodeDetails::CustomizeHeader(TSharedRef<class IPropertyHand
 					+ SHorizontalBox::Slot()
 					.VAlign(VAlign_Center)
 					.AutoWidth()
-					.Padding(0, 0, 0, 0)
+					.Padding(0.0f, 0.0f, 0.0f, 0.0f)
 					[
 						SAssignNew(NameSwitcher, SWidgetSwitcher)
 						.WidgetIndex(0)
 						+ SWidgetSwitcher::Slot()
 						[
 							SNew(SBox)
-							.Padding(FMargin(1,0, 1, 1))
+							.Padding(FMargin(1.0f,0.0f, 1.0f, 1.0f))
 							[
 								SNew(SRichTextBlock)
 								.Text(this, &FStateTreeEditorNodeDetails::GetNodeDescription)
@@ -619,7 +620,7 @@ void FStateTreeEditorNodeDetails::CustomizeHeader(TSharedRef<class IPropertyHand
 							.ButtonStyle(FAppStyle::Get(), "SimpleButton")
 							.OnClicked(this, &FStateTreeEditorNodeDetails::OnBrowseToNodeBlueprint)
 							.ToolTipText(LOCTEXT("BrowseToCurrentNodeBP", "Browse to the current node blueprint in Content Browser"))
-							.ContentPadding(.0f)
+							.ContentPadding(0.0f)
 							[
 								SNew(SImage)
 								.Image(FAppStyle::GetBrush("Icons.BrowseContent"))
@@ -640,7 +641,7 @@ void FStateTreeEditorNodeDetails::CustomizeHeader(TSharedRef<class IPropertyHand
 							.ButtonStyle(FAppStyle::Get(), "SimpleButton")
 							.OnClicked(this, &FStateTreeEditorNodeDetails::OnEditNodeBlueprint)
 							.ToolTipText(LOCTEXT("EditCurrentNodeBP", "Edit the current node blueprint in Editor"))
-							.ContentPadding(.0f)
+							.ContentPadding(0.0f)
 							[
 								SNew(SImage)
 								.Image(FAppStyle::GetBrush("Icons.Edit"))
@@ -661,7 +662,7 @@ void FStateTreeEditorNodeDetails::CustomizeHeader(TSharedRef<class IPropertyHand
 						.OnGetMenuContent(this, &FStateTreeEditorNodeDetails::GenerateOptionsMenu)
 						.ToolTipText(LOCTEXT("ItemActions", "Item actions"))
 						.HasDownArrow(false)
-						.ContentPadding(FMargin(4.f, 2.f))
+						.ContentPadding(FMargin(4.0f, 2.0f))
 						.ButtonContent()
 						[
 							SNew(SImage)
@@ -684,19 +685,20 @@ void FStateTreeEditorNodeDetails::CustomizeHeader(TSharedRef<class IPropertyHand
 			DescriptionBox->InsertSlot(0)
 				.AutoWidth()
 				.VAlign(VAlign_Center)
-				.Padding(0.f, 0.f, 0.f, 0.f)
+				.Padding(0.0f, 0.0f, 0.0f, 0.0f)
 				[
 					// Create the toggle favorites button
 					SNew(SButton)
 					.ButtonStyle(FAppStyle::Get(), "SimpleButton")
 					.OnClicked(this, &FStateTreeEditorNodeDetails::HandleToggleCompletionTaskClicked)
-					.ToolTipText(LOCTEXT("ToggleTaskCompletion", "Toggle Completion.\nA task can be considered for state completion or run in the background without affecting the state completion."))
+					.ToolTipText(this, &FStateTreeEditorNodeDetails::GetToggleCompletionTaskTooltip)
 					[
 						SNew(SImage)
 						.ColorAndOpacity(this, &FStateTreeEditorNodeDetails::GetToggleCompletionTaskColor)
-						.Image(FAppStyle::Get().GetBrush("Icons.Star"))
+						.Image(this, &FStateTreeEditorNodeDetails::GetToggleCompletionTaskIcon)
 					]
 					.IsEnabled(UE::StateTreeEditor::EditorNodeUtils::CanEditTaskConsideredForCompletion(*Node))
+					.Visibility(this, &FStateTreeEditorNodeDetails::GetToggleCompletionTaskVisibility)
 				];
 		}
 	}
@@ -1172,7 +1174,7 @@ int32 FStateTreeEditorNodeDetails::GetNextIndent() const
 	{
 		return 0;
 	}
-	
+
 	const int32 NextIndex = StructProperty->GetIndexInArray() + 1;
 	if (NextIndex >= (int32)NumElements)
 	{
@@ -1184,13 +1186,13 @@ int32 FStateTreeEditorNodeDetails::GetNextIndent() const
 	{
 		return 0;
 	}
-	
+
 	TSharedPtr<IPropertyHandle> NextIndentProperty = NextStructProperty->GetChildHandle(GET_MEMBER_NAME_CHECKED(FStateTreeEditorNode, ExpressionIndent));
 	if (!NextIndentProperty.IsValid())
 	{
 		return 0;
 	}
-	
+
 	uint8 Indent = 0;
 	NextIndentProperty->GetValue(Indent);
 
@@ -1209,10 +1211,10 @@ FText FStateTreeEditorNodeDetails::GetOpenParens() const
 	static_assert(UE::StateTree::MaxExpressionIndent == 4);
 	switch (OpenParens)
 	{
-		case 1: return FText::FromString(TEXT("("));
-		case 2: return FText::FromString(TEXT("(("));
-		case 3: return FText::FromString(TEXT("((("));
-		case 4: return FText::FromString(TEXT("(((("));
+	case 1: return FText::FromString(TEXT("("));
+	case 2: return FText::FromString(TEXT("(("));
+	case 3: return FText::FromString(TEXT("((("));
+	case 4: return FText::FromString(TEXT("(((("));
 	}
 	return FText::GetEmpty();
 }
@@ -1237,6 +1239,13 @@ FText FStateTreeEditorNodeDetails::GetCloseParens() const
 	return FText::GetEmpty();
 }
 
+FSlateColor FStateTreeEditorNodeDetails::GetContentRowColor() const
+{
+	return UE::StateTreeEditor::DebuggerExtensions::IsEditorNodeEnabled(StructProperty)
+		? FSlateColor::UseForeground()
+		: FSlateColor::UseSubduedForeground();
+}
+
 FText FStateTreeEditorNodeDetails::GetOperandText() const
 {
 	check(OperandProperty);
@@ -1625,22 +1634,70 @@ FReply FStateTreeEditorNodeDetails::HandleToggleCompletionTaskClicked()
 {
 	if (FStateTreeEditorNode* Node = UE::StateTreeEditor::EditorNodeUtils::GetMutableCommonNode(StructProperty))
 	{
-		const bool bCurrentValue = UE::StateTreeEditor::EditorNodeUtils::IsTaskConsideredForCompletion(*Node);
-		UE::StateTreeEditor::EditorNodeUtils::SetTaskConsideredForCompletion(*Node, !bCurrentValue);
+		if (UE::StateTreeEditor::EditorNodeUtils::IsTaskEnabled(*Node))
+		{
+			const bool bCurrentValue = UE::StateTreeEditor::EditorNodeUtils::IsTaskConsideredForCompletion(*Node);
+			UE::StateTreeEditor::EditorNodeUtils::SetTaskConsideredForCompletion(*Node, !bCurrentValue);
+		}
 	}
 	return FReply::Handled();
 }
 
+FText FStateTreeEditorNodeDetails::GetToggleCompletionTaskTooltip() const
+{
+	if (const FStateTreeEditorNode* Node = UE::StateTreeEditor::EditorNodeUtils::GetCommonNode(StructProperty))
+	{
+		if (UE::StateTreeEditor::EditorNodeUtils::IsTaskConsideredForCompletion(*Node))
+		{
+			return LOCTEXT("ToggleTaskCompletionEnabled", "Toggle Completion.\n"
+				"The task is considered for state completion.\n"
+				"When the task completes, it will stop ticking, and the state can be considered for transition.");
+		}
+		else
+		{
+			return LOCTEXT("ToggleTaskCompletionDisabled", "Toggle Completion.\n"
+				"The task doesn't affect the state completion.\n"
+				"When the task completes, it will stop ticking.");
+		}
+	}
+	return FText::GetEmpty();
+}
+
 FSlateColor FStateTreeEditorNodeDetails::GetToggleCompletionTaskColor() const
 {
 	if (const FStateTreeEditorNode* Node = UE::StateTreeEditor::EditorNodeUtils::GetCommonNode(StructProperty))
 	{
 		if (UE::StateTreeEditor::EditorNodeUtils::IsTaskConsideredForCompletion(*Node))
 		{
-			return FSlateColor(EStyleColor::Foreground);
+			return UE::StateTree::Colors::Cyan;
+		}
+	}
+	return FSlateColor(EStyleColor::Foreground);
+}
+
+const FSlateBrush* FStateTreeEditorNodeDetails::GetToggleCompletionTaskIcon() const
+{
+	if (const FStateTreeEditorNode* Node = UE::StateTreeEditor::EditorNodeUtils::GetCommonNode(StructProperty))
+	{
+		if (UE::StateTreeEditor::EditorNodeUtils::IsTaskConsideredForCompletion(*Node))
+		{
+			return FStateTreeEditorStyle::Get().GetBrush("StateTreeEditor.TasksCompletion.Enabled");
+		}
+		else
+		{
+			return FStateTreeEditorStyle::Get().GetBrush("StateTreeEditor.TasksCompletion.Disabled");
 		}
 	}
-	return FSlateColor(EStyleColor::Recessed);
+	return nullptr;
+}
+
+EVisibility FStateTreeEditorNodeDetails::GetToggleCompletionTaskVisibility() const
+{
+	if (const FStateTreeEditorNode* Node = UE::StateTreeEditor::EditorNodeUtils::GetCommonNode(StructProperty))
+	{
+		return UE::StateTreeEditor::EditorNodeUtils::IsTaskEnabled(*Node) ? EVisibility::Visible : EVisibility::Hidden;
+	}
+	return EVisibility::Collapsed;
 }
 
 FText FStateTreeEditorNodeDetails::GetNodePickerTooltip() const
diff --git a/Engine/Plugins/Runtime/StateTree/Source/StateTreeEditorModule/Private/Customizations/StateTreeEditorNodeDetails.h b/Engine/Plugins/Runtime/StateTree/Source/StateTreeEditorModule/Private/Customizations/StateTreeEditorNodeDetails.h
--- a/Engine/Plugins/Runtime/StateTree/Source/StateTreeEditorModule/Private/Customizations/StateTreeEditorNodeDetails.h
+++ b/Engine/Plugins/Runtime/StateTree/Source/StateTreeEditorModule/Private/Customizations/StateTreeEditorNodeDetails.h
@@ -61,6 +61,8 @@ private:
 	void SetIndent(const int32 Indent) const;
 	bool IsIndent(const int32 Indent) const;
 	
+	FSlateColor GetContentRowColor() const;
+
 	FText GetOperandText() const;
 	FSlateColor GetOperandColor() const;
 	TSharedRef<SWidget> OnGetOperandContent() const;
@@ -101,7 +103,10 @@ private:
 	void HandleNameCommitted(const FText& NewLabel, ETextCommit::Type CommitType) const;
 
 	FReply HandleToggleCompletionTaskClicked();
+	FText GetToggleCompletionTaskTooltip() const;
 	FSlateColor GetToggleCompletionTaskColor() const;
+	const FSlateBrush* GetToggleCompletionTaskIcon() const;
+	EVisibility GetToggleCompletionTaskVisibility() const;
 
 	FText GetNodePickerTooltip() const;
 	void OnNodePicked(const UStruct* InStruct) const;
diff --git a/Engine/Plugins/Runtime/StateTree/Source/StateTreeEditorModule/Private/Customizations/StateTreeTransitionDetails.cpp b/Engine/Plugins/Runtime/StateTree/Source/StateTreeEditorModule/Private/Customizations/StateTreeTransitionDetails.cpp
--- a/Engine/Plugins/Runtime/StateTree/Source/StateTreeEditorModule/Private/Customizations/StateTreeTransitionDetails.cpp
+++ b/Engine/Plugins/Runtime/StateTree/Source/StateTreeEditorModule/Private/Customizations/StateTreeTransitionDetails.cpp
@@ -73,7 +73,8 @@ void FStateTreeTransitionDetails::CustomizeHeader(TSharedRef<IPropertyHandle> St
 			// Border to capture mouse clicks on the row (used for right click menu).
 			SAssignNew(RowBorder, SBorder)
 			.BorderImage(FStyleDefaults::GetNoBrush())
-			.Padding(0)
+			.Padding(0.f)
+			.ForegroundColor(this, &FStateTreeTransitionDetails::GetContentRowColor)
 			.OnMouseButtonDown(this, &FStateTreeTransitionDetails::OnRowMouseDown)
 			.OnMouseButtonUp(this, &FStateTreeTransitionDetails::OnRowMouseUp)
 			[
@@ -84,7 +85,7 @@ void FStateTreeTransitionDetails::CustomizeHeader(TSharedRef<IPropertyHandle> St
 				+ SHorizontalBox::Slot()
 				.AutoWidth()
 				.VAlign(VAlign_Center)
-				.Padding(FMargin(0, 0, 4, 0))
+				.Padding(FMargin(0.f, 0.f, 4.f, 0.f))
 				[
 					SNew(SImage)
 					.ColorAndOpacity(FSlateColor::UseSubduedForeground())
@@ -95,7 +96,7 @@ void FStateTreeTransitionDetails::CustomizeHeader(TSharedRef<IPropertyHandle> St
 				+ SHorizontalBox::Slot()
 				.AutoWidth()
 				.VAlign(VAlign_Center)
-				.Padding(FMargin(0, 1, 0, 0))
+				.Padding(FMargin(0.f, 1.f, 0.f, 0.f))
 				[
 					SNew(SRichTextBlock)
 					.Text(this, &FStateTreeTransitionDetails::GetDescription)
@@ -111,7 +112,7 @@ void FStateTreeTransitionDetails::CustomizeHeader(TSharedRef<IPropertyHandle> St
 				.FillContentWidth(1.0f, 0.0f) // grow, no shrinking
 				.VAlign(VAlign_Center)
 				.HAlign(HAlign_Right)
-				.Padding(FMargin(8, 0, 2, 0))
+				.Padding(FMargin(8.f, 0.f, 2.f, 0.f))
 				[
 					SNew(SHorizontalBox)
 					// Debugger labels
@@ -257,6 +258,13 @@ void FStateTreeTransitionDetails::CustomizeChildren(TSharedRef<class IPropertyHa
 		];
 }
 
+FSlateColor FStateTreeTransitionDetails::GetContentRowColor() const
+{
+	return UE::StateTreeEditor::DebuggerExtensions::IsTransitionEnabled(StructProperty)
+		? FSlateColor::UseForeground()
+		: FSlateColor::UseSubduedForeground();
+}
+
 FReply FStateTreeTransitionDetails::OnRowMouseDown(const FGeometry& MyGeometry, const FPointerEvent& MouseEvent)
 {
 	if (MouseEvent.GetEffectingButton() == EKeys::RightMouseButton)
diff --git a/Engine/Plugins/Runtime/StateTree/Source/StateTreeEditorModule/Private/Customizations/StateTreeTransitionDetails.h b/Engine/Plugins/Runtime/StateTree/Source/StateTreeEditorModule/Private/Customizations/StateTreeTransitionDetails.h
--- a/Engine/Plugins/Runtime/StateTree/Source/StateTreeEditorModule/Private/Customizations/StateTreeTransitionDetails.h
+++ b/Engine/Plugins/Runtime/StateTree/Source/StateTreeEditorModule/Private/Customizations/StateTreeTransitionDetails.h
@@ -40,6 +40,8 @@ private:
 	void OnCopyTransition() const;
 	void OnPasteTransition() const;
 
+	FSlateColor GetContentRowColor() const;
+
 	FReply OnRowMouseDown(const FGeometry& MyGeometry, const FPointerEvent& MouseEvent);
 	FReply OnRowMouseUp(const FGeometry& MyGeometry, const FPointerEvent& MouseEvent);
 	TSharedRef<SWidget> GenerateOptionsMenu();
diff --git a/Engine/Plugins/Runtime/StateTree/Source/StateTreeEditorModule/Private/Debugger/StateTreeDebuggerUIExtensions.cpp b/Engine/Plugins/Runtime/StateTree/Source/StateTreeEditorModule/Private/Debugger/StateTreeDebuggerUIExtensions.cpp
--- a/Engine/Plugins/Runtime/StateTree/Source/StateTreeEditorModule/Private/Debugger/StateTreeDebuggerUIExtensions.cpp
+++ b/Engine/Plugins/Runtime/StateTree/Source/StateTreeEditorModule/Private/Debugger/StateTreeDebuggerUIExtensions.cpp
@@ -538,11 +538,11 @@ TSharedRef<SWidget> CreateEditorNodeWidget(const TSharedPtr<IPropertyHandle>& St
 			.Padding(4.0f, 0.0f)
 			[
 				SNew(SBorder)
-				.Padding(FMargin(6, 2))
+				.Padding(FMargin(6.0f, 2.0f))
 				.BorderImage(FStateTreeEditorStyle::Get().GetBrush("StateTree.Node.Label"))
-				.Visibility_Lambda([bIsTask, StructPropertyHandle]
+				.Visibility_Lambda([StructPropertyHandle]
 					{
-						return bIsTask && IsTaskDisabled(StructPropertyHandle) ? EVisibility::Visible : EVisibility::Collapsed;
+						return IsEditorNodeEnabled(StructPropertyHandle) ? EVisibility::Collapsed : EVisibility::Visible;
 					})
 				[
 					SNew(STextBlock)
@@ -561,7 +561,7 @@ TSharedRef<SWidget> CreateEditorNodeWidget(const TSharedPtr<IPropertyHandle>& St
 			.Padding(4.0f, 0.0f)
 			[
 				SNew(SBorder)
-				.Padding(FMargin(6, 2))
+				.Padding(FMargin(6.0f, 2.0f))
 				.BorderImage(FStateTreeEditorStyle::Get().GetBrush("StateTree.Node.Label"))
 				.Visibility_Lambda([StructPropertyHandle, WeakTreeData]
 					{
@@ -584,7 +584,7 @@ TSharedRef<SWidget> CreateEditorNodeWidget(const TSharedPtr<IPropertyHandle>& St
 			.Padding(4.0f, 0.0f)
 			[
 				SNew(SBorder)
-				.Padding(FMargin(6, 2))
+				.Padding(FMargin(6.0f, 2.0f))
 				.BorderImage(FStateTreeEditorStyle::Get().GetBrush("StateTree.Node.Label"))
 				.Visibility_Lambda([StructPropertyHandle, WeakTreeData]
 					{
@@ -607,7 +607,7 @@ TSharedRef<SWidget> CreateEditorNodeWidget(const TSharedPtr<IPropertyHandle>& St
 			.Padding(4.0f, 0.0f)
 			[
 				SNew(SBorder)
-				.Padding(FMargin(6, 2))
+				.Padding(FMargin(6.0f, 2.0f))
 				.BorderImage(FStateTreeEditorStyle::Get().GetBrush("StateTree.Node.Label"))
 				.BorderBackgroundColor_Lambda([StructPropertyHandle]
 					{
@@ -806,6 +806,21 @@ void AppendEditorNodeMenuItems(FMenuBuilder& InMenuBuilder, const TSharedPtr<IPr
 	}
 }
 
+bool IsEditorNodeEnabled(const TSharedPtr<IPropertyHandle>& StructPropertyHandle)
+{
+	using namespace EditorNodeUtils;
+
+	if (const FStateTreeEditorNode* Node = GetCommonNode(StructPropertyHandle))
+	{
+		if (const FStateTreeTaskBase* TaskBase = Node->Node.GetPtr<FStateTreeTaskBase>())
+		{
+			return TaskBase->bTaskEnabled;
+		}
+	}
+
+	return true;
+}
+
 TSharedRef<SWidget> CreateTransitionWidget(const TSharedPtr<IPropertyHandle>& StructPropertyHandle, UStateTreeEditorData* TreeData)
 {
 	using namespace EditorNodeUtils;
@@ -837,7 +852,7 @@ TSharedRef<SWidget> CreateTransitionWidget(const TSharedPtr<IPropertyHandle>& St
 					})
 				.Visibility_Lambda([StructPropertyHandle]
 					{
-						return GetTransitionEnabledCheckState(StructPropertyHandle) == ECheckBoxState::Checked ? EVisibility::Collapsed : EVisibility::Visible;
+						return IsTransitionEnabled(StructPropertyHandle) ? EVisibility::Collapsed : EVisibility::Visible;
 					})
 				[
 					SNew(STextBlock)
@@ -921,6 +936,11 @@ void AppendTransitionMenuItems(FMenuBuilder& InMenuBuilder, const TSharedPtr<IPr
 }
 
 
+bool IsTransitionEnabled(const TSharedPtr<IPropertyHandle>& StructPropertyHandle)
+{
+	return GetTransitionEnabledCheckState(StructPropertyHandle) == ECheckBoxState::Checked;
+}
+
 } // namespace UE::StateTreeEditor::DebuggerExtensions
 
 #undef LOCTEXT_NAMESPACE
diff --git a/Engine/Plugins/Runtime/StateTree/Source/StateTreeEditorModule/Private/Debugger/StateTreeDebuggerUIExtensions.h b/Engine/Plugins/Runtime/StateTree/Source/StateTreeEditorModule/Private/Debugger/StateTreeDebuggerUIExtensions.h
--- a/Engine/Plugins/Runtime/StateTree/Source/StateTreeEditorModule/Private/Debugger/StateTreeDebuggerUIExtensions.h
+++ b/Engine/Plugins/Runtime/StateTree/Source/StateTreeEditorModule/Private/Debugger/StateTreeDebuggerUIExtensions.h
@@ -23,8 +23,10 @@ void AppendStateMenuItems(FMenuBuilder& InMenuBuilder, TSharedPtr<IPropertyHandl
 
 TSharedRef<SWidget> CreateEditorNodeWidget(const TSharedPtr<IPropertyHandle>& StructPropertyHandle, UStateTreeEditorData* TreeData);
 void AppendEditorNodeMenuItems(FMenuBuilder& InMenuBuilder, const TSharedPtr<IPropertyHandle>& StructPropertyHandle, UStateTreeEditorData* TreeData);
+bool IsEditorNodeEnabled(const TSharedPtr<IPropertyHandle>& StructPropertyHandle);
 
 TSharedRef<SWidget> CreateTransitionWidget(const TSharedPtr<IPropertyHandle>& StructPropertyHandle, UStateTreeEditorData* TreeData);
 void AppendTransitionMenuItems(FMenuBuilder& InMenuBuilder, const TSharedPtr<IPropertyHandle>& StructPropertyHandle, UStateTreeEditorData* TreeData);
+bool IsTransitionEnabled(const TSharedPtr<IPropertyHandle>& StructPropertyHandle);
 
 }; // UE::StateTreeEditor::DebuggerExtensions
diff --git a/Engine/Plugins/Runtime/StateTree/Source/StateTreeEditorModule/Private/StateTreeEditorStyle.cpp b/Engine/Plugins/Runtime/StateTree/Source/StateTreeEditorModule/Private/StateTreeEditorStyle.cpp
--- a/Engine/Plugins/Runtime/StateTree/Source/StateTreeEditorModule/Private/StateTreeEditorStyle.cpp
+++ b/Engine/Plugins/Runtime/StateTree/Source/StateTreeEditorModule/Private/StateTreeEditorStyle.cpp
@@ -395,6 +395,9 @@ FStateTreeEditorStyle::FStateTreeEditorStyle()
 		Set("StateTreeEditor.Tasks.Large", new IMAGE_BRUSH_SVG("Icons/Tasks", CoreStyleConstants::Icon24x24));
 		Set("StateTreeEditor.Transitions", new IMAGE_BRUSH_SVG("Icons/Transitions", CoreStyleConstants::Icon16x16));
 
+		Set("StateTreeEditor.TasksCompletion.Enabled", new IMAGE_BRUSH_SVG("Icons/ConsiderTask", CoreStyleConstants::Icon16x16));
+		Set("StateTreeEditor.TasksCompletion.Disabled", new IMAGE_BRUSH_SVG("Icons/NotConsiderTask", CoreStyleConstants::Icon16x16));
+
 		Set("StateTreeEditor.StateSubtree", new IMAGE_BRUSH_SVG("Icons/State_Subtree", CoreStyleConstants::Icon16x16));
 		Set("StateTreeEditor.StateLinked", new IMAGE_BRUSH_SVG("Icons/State_Linked", CoreStyleConstants::Icon16x16));
 
