StateTree: Clean up unit test string to use wide char and remove FString allocation when not needed.
#rnx

diff --git a/Engine/Plugins/Runtime/StateTree/Source/StateTreeTestSuite/Private/StateTreeBindingTest.cpp b/Engine/Plugins/Runtime/StateTree/Source/StateTreeTestSuite/Private/StateTreeBindingTest.cpp
--- a/Engine/Plugins/Runtime/StateTree/Source/StateTreeTestSuite/Private/StateTreeBindingTest.cpp
+++ b/Engine/Plugins/Runtime/StateTree/Source/StateTreeTestSuite/Private/StateTreeBindingTest.cpp
@@ -25,7 +25,7 @@ struct FStateTreeTest_BindingsCompiler : FStateTreeTestBase
 		FStateTreePropertyBindingCompiler BindingCompiler;
 
 		const bool bInitResult = BindingCompiler.Init(Bindings, Log);
-		AITEST_TRUE("Expect init to succeed", bInitResult);
+		AITEST_TRUE(TEXT("Expect init to succeed"), bInitResult);
 
 		FStateTreeBindableStructDesc SourceADesc;
 		SourceADesc.Name = FName(TEXT("SourceA"));
@@ -65,13 +65,13 @@ struct FStateTreeTest_BindingsCompiler : FStateTreeTestBase
 
 		int32 CopyBatchIndex = INDEX_NONE;
 		const bool bCompileBatchResult = BindingCompiler.CompileBatch(TargetDesc, PropertyBindings, FStateTreeIndex16::Invalid, FStateTreeIndex16::Invalid, CopyBatchIndex);
-		AITEST_TRUE("CompileBatch should succeed", bCompileBatchResult);
-		AITEST_NOT_EQUAL("CopyBatchIndex should not be INDEX_NONE", CopyBatchIndex, (int32)INDEX_NONE);
+		AITEST_TRUE(TEXT("CompileBatch should succeed"), bCompileBatchResult);
+		AITEST_NOT_EQUAL(TEXT("CopyBatchIndex should not be INDEX_NONE"), CopyBatchIndex, (int32)INDEX_NONE);
 
 		BindingCompiler.Finalize();
 
 		const bool bResolveResult = Bindings.ResolvePaths();
-		AITEST_TRUE("ResolvePaths should succeed", bResolveResult);
+		AITEST_TRUE(TEXT("ResolvePaths should succeed"), bResolveResult);
 
 		FStateTreeTest_PropertyCopy SourceA;
 		SourceA.Item.B = 123;
@@ -94,8 +94,8 @@ struct FStateTreeTest_BindingsCompiler : FStateTreeTestBase
 		FStateTreeTest_PropertyCopy Target;
 		Target.FixedArray.SetNum(FixedArraySize, EAllowShrinking::No);
 
-		AITEST_TRUE("SourceAIndex should be less than max number of source structs.", SourceAIndex < Bindings.GetNumBindableStructDescriptors());
-		AITEST_TRUE("SourceBIndex should be less than max number of source structs.", SourceBIndex < Bindings.GetNumBindableStructDescriptors());
+		AITEST_TRUE(TEXT("SourceAIndex should be less than max number of source structs."), SourceAIndex < Bindings.GetNumBindableStructDescriptors());
+		AITEST_TRUE(TEXT("SourceBIndex should be less than max number of source structs."), SourceBIndex < Bindings.GetNumBindableStructDescriptors());
 
 		TArray<FStateTreeDataView> SourceViews;
 		SourceViews.SetNum(Bindings.GetNumBindableStructDescriptors());
@@ -108,7 +108,7 @@ struct FStateTreeTest_BindingsCompiler : FStateTreeTestBase
 		{
 			bCopyResult &= Bindings.Super::CopyProperty(Copy, SourceViews[Copy.SourceDataHandle.Get<FStateTreeDataHandle>().GetIndex()], TargetView);
 		}
-		AITEST_TRUE("CopyTo should succeed", bCopyResult);
+		AITEST_TRUE(TEXT("CopyTo should succeed"), bCopyResult);
 
 		// Due to binding sorting, we expect them to execute in this order (sorted based on target access, earliest to latest)
 		// SourceA.CArray -> Target.CArray
@@ -121,30 +121,30 @@ struct FStateTreeTest_BindingsCompiler : FStateTreeTestBase
 		// SourceB.Item -> Target.Array[1]
 		// SourceA.Item.B -> Target.Array[1].B
 
-		AITEST_EQUAL("Expect TargetArray to be copied from SourceA", Target.Array.Num(), SourceA.Array.Num());
-		AITEST_EQUAL("Expect Target.Array[0].A copied from SourceA.Array[0].A", Target.Array[0].A, SourceA.Array[0].A);
-		AITEST_EQUAL("Expect Target.Array[0].B copied from SourceA.Array[0].B", Target.Array[0].B, SourceA.Array[0].B);
-		AITEST_EQUAL("Expect Target.Array[1].A copied from SourceB.Item.A", Target.Array[1].A, SourceB.Item.A);
-		AITEST_EQUAL("Expect Target.Array[1].B copied from SourceA.Item.B", Target.Array[1].B, SourceA.Item.B);
-
-		AITEST_EQUAL("Expect TargetArray to be copied from SourceA", Target.FixedArray.Num(), SourceA.FixedArray.Num());
-		AITEST_EQUAL("Expect Target.FixedArray[0].A copied from SourceA.FixedArray[0].A", Target.FixedArray[0].A, SourceA.FixedArray[0].A);
-		AITEST_EQUAL("Expect Target.FixedArray[0].B copied from SourceA.FixedArray[0].B", Target.FixedArray[0].B, SourceA.FixedArray[0].B);
-		AITEST_EQUAL("Expect Target.FixedArray[1].A copied from SourceB.Item.A", Target.FixedArray[1].A, SourceB.Item.A);
-		AITEST_EQUAL("Expect Target.FixedArray[1].B copied from SourceA.Item.B", Target.FixedArray[1].B, SourceA.Item.B);
-		AITEST_EQUAL("Expect Target.FixedArray to not have changed size", Target.FixedArray.Num(), FixedArraySize);
-
-		AITEST_EQUAL("Expect Target.CArray[0].A copied from SourceA.CArray[0].A", Target.CArray[0].A, SourceA.CArray[0].A);
-		AITEST_EQUAL("Expect Target.CArray[0].B copied from SourceA.CArray[0].B", Target.CArray[0].B, SourceA.CArray[0].B);
-		AITEST_EQUAL("Expect Target.CArray[1].A copied from SourceB.Item.A", Target.CArray[1].A, SourceB.Item.A);
-		AITEST_EQUAL("Expect Target.CArray[1].B copied from SourceA.Item.B", Target.CArray[1].B, SourceA.Item.B);
+		AITEST_EQUAL(TEXT("Expect TargetArray to be copied from SourceA"), Target.Array.Num(), SourceA.Array.Num());
+		AITEST_EQUAL(TEXT("Expect Target.Array[0].A copied from SourceA.Array[0].A"), Target.Array[0].A, SourceA.Array[0].A);
+		AITEST_EQUAL(TEXT("Expect Target.Array[0].B copied from SourceA.Array[0].B"), Target.Array[0].B, SourceA.Array[0].B);
+		AITEST_EQUAL(TEXT("Expect Target.Array[1].A copied from SourceB.Item.A"), Target.Array[1].A, SourceB.Item.A);
+		AITEST_EQUAL(TEXT("Expect Target.Array[1].B copied from SourceA.Item.B"), Target.Array[1].B, SourceA.Item.B);
+
+		AITEST_EQUAL(TEXT("Expect TargetArray to be copied from SourceA"), Target.FixedArray.Num(), SourceA.FixedArray.Num());
+		AITEST_EQUAL(TEXT("Expect Target.FixedArray[0].A copied from SourceA.FixedArray[0].A"), Target.FixedArray[0].A, SourceA.FixedArray[0].A);
+		AITEST_EQUAL(TEXT("Expect Target.FixedArray[0].B copied from SourceA.FixedArray[0].B"), Target.FixedArray[0].B, SourceA.FixedArray[0].B);
+		AITEST_EQUAL(TEXT("Expect Target.FixedArray[1].A copied from SourceB.Item.A"), Target.FixedArray[1].A, SourceB.Item.A);
+		AITEST_EQUAL(TEXT("Expect Target.FixedArray[1].B copied from SourceA.Item.B"), Target.FixedArray[1].B, SourceA.Item.B);
+		AITEST_EQUAL(TEXT("Expect Target.FixedArray to not have changed size"), Target.FixedArray.Num(), FixedArraySize);
+
+		AITEST_EQUAL(TEXT("Expect Target.CArray[0].A copied from SourceA.CArray[0].A"), Target.CArray[0].A, SourceA.CArray[0].A);
+		AITEST_EQUAL(TEXT("Expect Target.CArray[0].B copied from SourceA.CArray[0].B"), Target.CArray[0].B, SourceA.CArray[0].B);
+		AITEST_EQUAL(TEXT("Expect Target.CArray[1].A copied from SourceB.Item.A"), Target.CArray[1].A, SourceB.Item.A);
+		AITEST_EQUAL(TEXT("Expect Target.CArray[1].B copied from SourceA.Item.B"), Target.CArray[1].B, SourceA.Item.B);
 		
 		const int32 NumAllocated_FStateTreeTest_PropertyStructB_BeforeReset = FStateTreeTest_PropertyStructB::NumConstructed;
 		bool bResetResult = Bindings.FPropertyBindingBindingCollection::ResetObjects(FPropertyBindingIndex16(CopyBatchIndex), TargetView);
-		AITEST_TRUE("ResetObjects should succeed", bResetResult);
-		AITEST_EQUAL("Expect Target dynamic array to be empty", Target.Array.Num(), 0);
-		AITEST_EQUAL("Expect Target fixed size Array to not have changed size.", Target.FixedArray.Num(), FixedArraySize);
-		AITEST_NOT_EQUAL("Expect the count of constructed FStateTreeTest_PropertyStructB to be smaller after calling ResetObjects", FStateTreeTest_PropertyStructB::NumConstructed, NumAllocated_FStateTreeTest_PropertyStructB_BeforeReset);
+		AITEST_TRUE(TEXT("ResetObjects should succeed"), bResetResult);
+		AITEST_EQUAL(TEXT("Expect Target dynamic array to be empty"), Target.Array.Num(), 0);
+		AITEST_EQUAL(TEXT("Expect Target fixed size Array to not have changed size."), Target.FixedArray.Num(), FixedArraySize);
+		AITEST_NOT_EQUAL(TEXT("Expect the count of constructed FStateTreeTest_PropertyStructB to be smaller after calling ResetObjects"), FStateTreeTest_PropertyStructB::NumConstructed, NumAllocated_FStateTreeTest_PropertyStructB_BeforeReset);
 
 		return true;
 	}
@@ -201,12 +201,12 @@ struct FStateTreeTest_PropertyFunctions : FStateTreeTestBase
 		FStateTreeCompilerLog Log;
 		FStateTreeCompiler Compiler(Log);
 		const bool bResult = Compiler.Compile(StateTree);
-		AITEST_TRUE("StateTree should get compiled", bResult);
+		AITEST_TRUE(TEXT("StateTree should get compiled"), bResult);
 
 		FStateTreeInstanceData InstanceData;
 		FTestStateTreeExecutionContext Exec(StateTree, StateTree, InstanceData);
 		const bool bInitSucceeded = Exec.IsValid();
-		AITEST_TRUE("StateTree should init", bInitSucceeded);
+		AITEST_TRUE(TEXT("StateTree should init"), bInitSucceeded);
 
 		Exec.Start();
 		AITEST_TRUE(*FString::Printf(TEXT("StateTree TaskA should enter state with value %d"), TaskAPropertyFunctionsAmount), Exec.Expect(TaskA.GetName(), *FString::Printf(TEXT("EnterState%d"), TaskAPropertyFunctionsAmount)));
@@ -239,7 +239,7 @@ struct FStateTreeTest_CopyObjects : FStateTreeTestBase
 		FStateTreePropertyBindingCompiler BindingCompiler;
 
 		const bool bInitResult = BindingCompiler.Init(Bindings, Log);
-		AITEST_TRUE("Expect init to succeed", bInitResult);
+		AITEST_TRUE(TEXT("Expect init to succeed"), bInitResult);
 
 		FStateTreeBindableStructDesc SourceDesc;
 		SourceDesc.Name = FName(TEXT("Source"));
@@ -277,18 +277,18 @@ struct FStateTreeTest_CopyObjects : FStateTreeTestBase
 		
 		int32 TargetACopyBatchIndex = INDEX_NONE;
 		const bool bCompileBatchResultA = BindingCompiler.CompileBatch(TargetADesc, PropertyBindings, FStateTreeIndex16::Invalid, FStateTreeIndex16::Invalid, TargetACopyBatchIndex);
-		AITEST_TRUE("CompileBatchResultA should succeed", bCompileBatchResultA);
-		AITEST_NOT_EQUAL("TargetACopyBatchIndex should not be INDEX_NONE", TargetACopyBatchIndex, (int32)INDEX_NONE);
+		AITEST_TRUE(TEXT("CompileBatchResultA should succeed"), bCompileBatchResultA);
+		AITEST_NOT_EQUAL(TEXT("TargetACopyBatchIndex should not be INDEX_NONE"), TargetACopyBatchIndex, (int32)INDEX_NONE);
 
 		int32 TargetBCopyBatchIndex = INDEX_NONE;
 		const bool bCompileBatchResultB = BindingCompiler.CompileBatch(TargetBDesc, PropertyBindings, FStateTreeIndex16::Invalid, FStateTreeIndex16::Invalid, TargetBCopyBatchIndex);
-		AITEST_TRUE("CompileBatchResultB should succeed", bCompileBatchResultB);
-		AITEST_NOT_EQUAL("TargetBCopyBatchIndex should not be INDEX_NONE", TargetBCopyBatchIndex, (int32)INDEX_NONE);
+		AITEST_TRUE(TEXT("CompileBatchResultB should succeed"), bCompileBatchResultB);
+		AITEST_NOT_EQUAL(TEXT("TargetBCopyBatchIndex should not be INDEX_NONE"), TargetBCopyBatchIndex, (int32)INDEX_NONE);
 
 		BindingCompiler.Finalize();
 
 		const bool bResolveResult = Bindings.ResolvePaths();
-		AITEST_TRUE("ResolvePaths should succeed", bResolveResult);
+		AITEST_TRUE(TEXT("ResolvePaths should succeed"), bResolveResult);
 
 		UStateTreeTest_PropertyObject* ObjectA = NewObject<UStateTreeTest_PropertyObject>();
 		UStateTreeTest_PropertyObject2* ObjectB = NewObject<UStateTreeTest_PropertyObject2>();
@@ -299,7 +299,7 @@ struct FStateTreeTest_CopyObjects : FStateTreeTestBase
 		Source.Class = UStateTreeTest_PropertyObject::StaticClass();
 		Source.SoftClass = UStateTreeTest_PropertyObject::StaticClass();
 
-		AITEST_TRUE("SourceIndex should be less than max number of source structs.", SourceIndex < Bindings.GetNumBindableStructDescriptors());
+		AITEST_TRUE(TEXT("SourceIndex should be less than max number of source structs."), SourceIndex < Bindings.GetNumBindableStructDescriptors());
 
 		TArray<FStateTreeDataView> SourceViews;
 		SourceViews.SetNum(Bindings.GetNumBindableStructDescriptors());
@@ -311,12 +311,12 @@ struct FStateTreeTest_CopyObjects : FStateTreeTestBase
 		{
 			bCopyResultA &= Bindings.Super::CopyProperty(Copy, SourceViews[Copy.SourceDataHandle.Get<FStateTreeDataHandle>().GetIndex()], FStructView::Make(TargetA));
 		}
-		AITEST_TRUE("CopyTo should succeed", bCopyResultA);
+		AITEST_TRUE(TEXT("CopyTo should succeed"), bCopyResultA);
 
-		AITEST_TRUE("Expect TargetA.Object == Source.Object", TargetA.Object == Source.Object);
-		AITEST_TRUE("Expect TargetA.SoftObject == Source.SoftObject", TargetA.SoftObject == Source.SoftObject);
-		AITEST_TRUE("Expect TargetA.Class == Source.Class", TargetA.Class == Source.Class);
-		AITEST_TRUE("Expect TargetA.SoftClass == Source.SoftClass", TargetA.SoftClass == Source.SoftClass);
+		AITEST_TRUE(TEXT("Expect TargetA.Object == Source.Object"), TargetA.Object == Source.Object);
+		AITEST_TRUE(TEXT("Expect TargetA.SoftObject == Source.SoftObject"), TargetA.SoftObject == Source.SoftObject);
+		AITEST_TRUE(TEXT("Expect TargetA.Class == Source.Class"), TargetA.Class == Source.Class);
+		AITEST_TRUE(TEXT("Expect TargetA.SoftClass == Source.SoftClass"), TargetA.SoftClass == Source.SoftClass);
 
 		// Copying to TargetB should not affect TargetA
 		TargetA.Object = nullptr;
@@ -327,14 +327,14 @@ struct FStateTreeTest_CopyObjects : FStateTreeTestBase
 		{
 			bCopyResultB &= Bindings.Super::CopyProperty(Copy, SourceViews[Copy.SourceDataHandle.Get<FStateTreeDataHandle>().GetIndex()], FStructView::Make(TargetB));
 		}
-		AITEST_TRUE("CopyTo should succeed", bCopyResultB);
+		AITEST_TRUE(TEXT("CopyTo should succeed"), bCopyResultB);
 
-		AITEST_TRUE("Expect TargetB.Object == Source.SoftObject", TSoftObjectPtr<UObject>(TargetB.Object) == Source.SoftObject);
-		AITEST_TRUE("Expect TargetB.SoftObject == Source.Object", TargetB.SoftObject == TSoftObjectPtr<UObject>(Source.Object));
-		AITEST_TRUE("Expect TargetB.Class == Source.SoftClass", TSoftClassPtr<UObject>(TargetB.Class) == Source.SoftClass);
-		AITEST_TRUE("Expect TargetB.SoftClass == Source.Class", TargetB.SoftClass == TSoftClassPtr<UObject>(Source.Class));
+		AITEST_TRUE(TEXT("Expect TargetB.Object == Source.SoftObject"), TSoftObjectPtr<UObject>(TargetB.Object) == Source.SoftObject);
+		AITEST_TRUE(TEXT("Expect TargetB.SoftObject == Source.Object"), TargetB.SoftObject == TSoftObjectPtr<UObject>(Source.Object));
+		AITEST_TRUE(TEXT("Expect TargetB.Class == Source.SoftClass"), TSoftClassPtr<UObject>(TargetB.Class) == Source.SoftClass);
+		AITEST_TRUE(TEXT("Expect TargetB.SoftClass == Source.Class"), TargetB.SoftClass == TSoftClassPtr<UObject>(Source.Class));
 
-		AITEST_TRUE("Expect TargetA.Object == nullptr after copy of TargetB", TargetA.Object == nullptr);
+		AITEST_TRUE(TEXT("Expect TargetA.Object == nullptr after copy of TargetB"), TargetA.Object == nullptr);
 
 		// Collect ObjectA and ObjectB, soft object paths should still copy ok.
 		ObjectA = nullptr;
@@ -350,8 +350,8 @@ struct FStateTreeTest_CopyObjects : FStateTreeTestBase
 		}
 
 		
-		AITEST_TRUE("CopyTo should succeed", bCopyResultC);
-		AITEST_TRUE("Expect TargetC.SoftObject == Source.SoftObject after GC", TargetC.SoftObject == Source.SoftObject);
+		AITEST_TRUE(TEXT("CopyTo should succeed"), bCopyResultC);
+		AITEST_TRUE(TEXT("Expect TargetC.SoftObject == Source.SoftObject after GC"), TargetC.SoftObject == Source.SoftObject);
 
 		return true;
 	}
@@ -368,7 +368,7 @@ struct FStateTreeTest_References : FStateTreeTestBase
 		FStateTreePropertyBindingCompiler BindingCompiler;
 
 		const bool bInitResult = BindingCompiler.Init(Bindings, Log);
-		AITEST_TRUE("Expect init to succeed", bInitResult);
+		AITEST_TRUE(TEXT("Expect init to succeed"), bInitResult);
 
 		FStateTreeBindableStructDesc SourceDesc;
 		SourceDesc.Name = FName(TEXT("Source"));
@@ -400,35 +400,35 @@ struct FStateTreeTest_References : FStateTreeTestBase
 		IDToStructValue.Emplace(TargetDesc.ID, TargetView);
 
 		const bool bCompileReferencesResult = BindingCompiler.CompileReferences(TargetDesc, PropertyBindings, TargetView, IDToStructValue);
-		AITEST_TRUE("CompileReferences should succeed", bCompileReferencesResult);
+		AITEST_TRUE(TEXT("CompileReferences should succeed"), bCompileReferencesResult);
 
 		BindingCompiler.Finalize();
 
 		const bool bResolveResult = Bindings.ResolvePaths();
-		AITEST_TRUE("ResolvePaths should succeed", bResolveResult);
+		AITEST_TRUE(TEXT("ResolvePaths should succeed"), bResolveResult);
 
 		{
 			const FStateTreePropertyAccess* PropertyAccess = Bindings.GetPropertyAccess(Target.RefToStruct);
-			AITEST_NOT_NULL("GetPropertyAccess should succeed", PropertyAccess);
+			AITEST_NOT_NULL(TEXT("GetPropertyAccess should succeed"), PropertyAccess);
 
 			FStateTreeTest_PropertyStruct* Reference = Bindings.GetMutablePropertyPtr<FStateTreeTest_PropertyStruct>(SourceView, *PropertyAccess);
-			AITEST_EQUAL("Expect RefToStruct to point to SourceA.Item", Reference, &Source.Item);
+			AITEST_EQUAL(TEXT("Expect RefToStruct to point to SourceA.Item"), Reference, &Source.Item);
 		}
 
 		{
 			const FStateTreePropertyAccess* PropertyAccess = Bindings.GetPropertyAccess(Target.RefToInt);
-			AITEST_NOT_NULL("GetPropertyAccess should succeed", PropertyAccess);
+			AITEST_NOT_NULL(TEXT("GetPropertyAccess should succeed"), PropertyAccess);
 
 			int32* Reference = Bindings.GetMutablePropertyPtr<int32>(SourceView, *PropertyAccess);
-			AITEST_EQUAL("Expect RefToInt to point to SourceA.Item.A", Reference, &Source.Item);
+			AITEST_EQUAL(TEXT("Expect RefToInt to point to SourceA.Item.A"), Reference, &Source.Item);
 		}
 
 		{
 			const FStateTreePropertyAccess* PropertyAccess = Bindings.GetPropertyAccess(Target.RefToStructArray);
-			AITEST_NOT_NULL("GetPropertyAccess should succeed", PropertyAccess);
+			AITEST_NOT_NULL(TEXT("GetPropertyAccess should succeed"), PropertyAccess);
 
 			TArray<FStateTreeTest_PropertyStruct>* Reference = Bindings.GetMutablePropertyPtr<TArray<FStateTreeTest_PropertyStruct>>(SourceView, *PropertyAccess);
-			AITEST_EQUAL("Expect RefToStructArray to point to SourceA.Array", Reference, &Source.Array);
+			AITEST_EQUAL(TEXT("Expect RefToStructArray to point to SourceA.Array"), Reference, &Source.Array);
 		}
 
 		return true;
@@ -445,7 +445,7 @@ struct FStateTreeTest_ReferencesConstness : FStateTreeTestBase
 		FStateTreePropertyBindingCompiler BindingCompiler;
 
 		const bool bInitResult = BindingCompiler.Init(Bindings, Log);
-		AITEST_TRUE("Expect init to succeed", bInitResult);
+		AITEST_TRUE(TEXT("Expect init to succeed"), bInitResult);
 
 		FStateTreeBindableStructDesc SourceAsTaskDesc;
 		SourceAsTaskDesc.Name = FName(TEXT("SourceTask"));
@@ -491,22 +491,22 @@ struct FStateTreeTest_ReferencesConstness : FStateTreeTestBase
 
 		{
 			const bool bCompileReferenceResult = BindingCompiler.CompileReferences(TargetDesc, { TaskPropertyBinding }, TargetView, IDToStructValue);
-			AITEST_FALSE("CompileReferences should fail", bCompileReferenceResult);
+			AITEST_FALSE(TEXT("CompileReferences should fail"), bCompileReferenceResult);
 		}
 
 		{
 			const bool bCompileReferenceResult = BindingCompiler.CompileReferences(TargetDesc, { TaskOutputPropertyBinding }, TargetView, IDToStructValue);
-			AITEST_TRUE("CompileReferences should succeed", bCompileReferenceResult);
+			AITEST_TRUE(TEXT("CompileReferences should succeed"), bCompileReferenceResult);
 		}
 
 		{
 			const bool bCompileReferenceResult = BindingCompiler.CompileReferences(TargetDesc, { ContextPropertyBinding }, TargetView, IDToStructValue);
-			AITEST_FALSE("CompileReferences should fail", bCompileReferenceResult);
+			AITEST_FALSE(TEXT("CompileReferences should fail"), bCompileReferenceResult);
 		}
 
 		{
 			const bool bCompileReferenceResult = BindingCompiler.CompileReferences(TargetDesc, { ContextOutputPropertyBinding }, TargetView, IDToStructValue);
-			AITEST_FALSE("CompileReferences should fail", bCompileReferenceResult);
+			AITEST_FALSE(TEXT("CompileReferences should fail"), bCompileReferenceResult);
 		}
 
 		return true;
@@ -590,14 +590,14 @@ struct FStateTreeTest_MutableArray : FStateTreeTestBase
 		{
 			FStateTreeCompiler Compiler(Log);
 			const bool bResult = Compiler.Compile(StateTree);
-			AITEST_TRUE("StateTree2 should get compiled", bResult);
+			AITEST_TRUE(TEXT("StateTree2 should get compiled"), bResult);
 		}
 		{
 			FStateTreeInstanceData InstanceData;
 			FTestStateTreeExecutionContext Exec(StateTree, StateTree, InstanceData);
 
 			const bool bInitSucceeded = Exec.IsValid();
-			AITEST_TRUE("StateTree should init", bInitSucceeded);
+			AITEST_TRUE(TEXT("StateTree should init"), bInitSucceeded);
 
 			EStateTreeRunStatus Status = EStateTreeRunStatus::Unset;
 			FInstancedPropertyBag GlobalParameters = StateTree.GetDefaultParameters();
@@ -609,8 +609,8 @@ struct FStateTreeTest_MutableArray : FStateTreeTestBase
 				ValueArrayRef.SetValueInt32(1, 922);
 				Status = Exec.Start(&GlobalParameters);
 			}
-			AITEST_EQUAL("Start should complete with Running", Status, EStateTreeRunStatus::Running);
-			AITEST_TRUE("Start should enter Global tasks", Exec.Expect("Tree1GlobalTaskA", TEXT("EnterState11"))
+			AITEST_EQUAL(TEXT("Start should complete with Running"), Status, EStateTreeRunStatus::Running);
+			AITEST_TRUE(TEXT("Start should enter Global tasks"), Exec.Expect("Tree1GlobalTaskA", TEXT("EnterState11"))
 				.Then("Tree1GlobalTaskA", TEXT("EnterState:{911, 922}")) // should copy the full array
 				.Then("Tree1StateATaskA", TEXT("EnterState11"))
 				.Then("Tree1StateATaskA", TEXT("EnterState:{-1, -2, -3, -4}")) // should not copy anything since [3] is out of scope
@@ -698,14 +698,14 @@ struct FStateTreeTest_TransitionTaskWithBinding : FStateTreeTestBase
 		{
 			FStateTreeCompiler Compiler(Log);
 			const bool bResult = Compiler.Compile(StateTree);
-			AITEST_TRUE("StateTree2 should get compiled", bResult);
+			AITEST_TRUE(TEXT("StateTree2 should get compiled"), bResult);
 		}
 		{
 			FStateTreeInstanceData InstanceData;
 			FTestStateTreeExecutionContext Exec(StateTree, StateTree, InstanceData);
 
 			const bool bInitSucceeded = Exec.IsValid();
-			AITEST_TRUE("StateTree should init", bInitSucceeded);
+			AITEST_TRUE(TEXT("StateTree should init"), bInitSucceeded);
 
 			EStateTreeRunStatus Status = EStateTreeRunStatus::Unset;
 			FInstancedPropertyBag GlobalParameters = StateTree.GetDefaultParameters();
@@ -714,8 +714,8 @@ struct FStateTreeTest_TransitionTaskWithBinding : FStateTreeTestBase
 				GlobalParameters.SetValueInt32("Value", 99);
 				Status = Exec.Start(&GlobalParameters);
 			}
-			AITEST_EQUAL("Start should complete with Running", Status, EStateTreeRunStatus::Running);
-			AITEST_TRUE("Start should enter Global tasks", Exec.Expect("Tree1GlobalTaskA", TEXT("EnterState99"))
+			AITEST_EQUAL(TEXT("Start should complete with Running"), Status, EStateTreeRunStatus::Running);
+			AITEST_TRUE(TEXT("Start should enter Global tasks"), Exec.Expect("Tree1GlobalTaskA", TEXT("EnterState99"))
 				.Then("Tree1GlobalTaskB", TEXT("EnterState99"))
 				.Then("Tree1GlobalTaskC", TEXT("EnterState99"))
 				.Then("Tree1StateRootTaskA", TEXT("EnterState99"))
@@ -731,8 +731,8 @@ struct FStateTreeTest_TransitionTaskWithBinding : FStateTreeTestBase
 			InstanceData.GetMutableStorage().SetGlobalParameters(GlobalParameters);
 
 			Status = Exec.Tick(1.0f);
-			AITEST_EQUAL("2nd Tick should complete with Running", Status, EStateTreeRunStatus::Running);
-			AITEST_TRUE("2nd Tick should tick tasks", Exec.Expect("Tree1GlobalTaskA", TEXT("Tick88"))
+			AITEST_EQUAL(TEXT("2nd Tick should complete with Running"), Status, EStateTreeRunStatus::Running);
+			AITEST_TRUE(TEXT("2nd Tick should tick tasks"), Exec.Expect("Tree1GlobalTaskA", TEXT("Tick88"))
 				.Then("Tree1GlobalTaskB", TEXT("Tick88"))
 				.Then("Tree1StateRootTaskA", TEXT("Tick88"))
 				.Then("Tree1StateRootTaskB", TEXT("Tick88"))
diff --git a/Engine/Plugins/Runtime/StateTree/Source/StateTreeTestSuite/Private/StateTreeDelegateTest.cpp b/Engine/Plugins/Runtime/StateTree/Source/StateTreeTestSuite/Private/StateTreeDelegateTest.cpp
--- a/Engine/Plugins/Runtime/StateTree/Source/StateTreeTestSuite/Private/StateTreeDelegateTest.cpp
+++ b/Engine/Plugins/Runtime/StateTree/Source/StateTreeTestSuite/Private/StateTreeDelegateTest.cpp
@@ -95,7 +95,7 @@ struct FStateTreeTest_Delegate_MutuallyExclusiveListeners : FStateTreeTestBase
 		AITEST_TRUE("StateTree should init", bInitSucceeded);
 
 		Exec.Start();
-		AITEST_TRUE("StateTree Active States should be in Root/A", Exec.ExpectInActiveStates(Root.Name, StateA.Name));
+		AITEST_TRUE(TEXT("StateTree Active States should be in Root/A"), Exec.ExpectInActiveStates(Root.Name, StateA.Name));
 		Exec.LogClear();
 
 		Exec.Tick(0.1f);
@@ -109,7 +109,7 @@ struct FStateTreeTest_Delegate_MutuallyExclusiveListeners : FStateTreeTestBase
 		AITEST_TRUE(TEXT("StateTree ListenerTaskB should be triggered once"), Exec.Expect(ListenerTaskB.GetName(), *FString::Printf(TEXT("OnDelegate%d"), 1)));
 		AITEST_TRUE(TEXT("StateTree ListenerTaskA0 shouldn't be triggered."), !Exec.Expect(ListenerTaskA0.GetName(), *FString::Printf(TEXT("OnDelegate%d"), 1)));
 		AITEST_TRUE(TEXT("StateTree ListenerTaskA1 shouldn't be triggered."), !Exec.Expect(ListenerTaskA1.GetName(), *FString::Printf(TEXT("OnDelegate%d"), 1)));
-		AITEST_TRUE("StateTree Active States should be in Root/A", Exec.ExpectInActiveStates(Root.Name, StateA.Name));
+		AITEST_TRUE(TEXT("StateTree Active States should be in Root/A"), Exec.ExpectInActiveStates(Root.Name, StateA.Name));
 		Exec.LogClear();
 
 		Exec.Stop();
@@ -150,15 +150,15 @@ struct FStateTreeTest_Delegate_Transitions : FStateTreeTestBase
 		AITEST_TRUE("StateTree should init", bInitSucceeded);
 
 		Exec.Start();
-		AITEST_TRUE("StateTree Active States should be in Root/A", Exec.ExpectInActiveStates(Root.Name, StateA.Name));
+		AITEST_TRUE(TEXT("StateTree Active States should be in Root/A"), Exec.ExpectInActiveStates(Root.Name, StateA.Name));
 		Exec.LogClear();
 
 		Exec.Tick(0.1f);
-		AITEST_TRUE("StateTree Active States should be in Root/B", Exec.ExpectInActiveStates(Root.Name, StateB.Name));
+		AITEST_TRUE(TEXT("StateTree Active States should be in Root/B"), Exec.ExpectInActiveStates(Root.Name, StateB.Name));
 		Exec.LogClear();
 
 		Exec.Tick(0.1f);
-		AITEST_TRUE("StateTree Active States should be in Root/A", Exec.ExpectInActiveStates(Root.Name, StateA.Name));
+		AITEST_TRUE(TEXT("StateTree Active States should be in Root/A"), Exec.ExpectInActiveStates(Root.Name, StateA.Name));
 		Exec.LogClear();
 
 		Exec.Stop();
@@ -186,12 +186,12 @@ struct FStateTreeTest_Delegate_Rebroadcasting : FStateTreeTestBase
 		FStateTreeCompilerLog Log;
 		FStateTreeCompiler Compiler(Log);
 		const bool bResult = Compiler.Compile(StateTree);
-		AITEST_TRUE("StateTree should get compiled", bResult);
+		AITEST_TRUE(TEXT("StateTree should get compiled"), bResult);
 
 		FStateTreeInstanceData InstanceData;
 		FTestStateTreeExecutionContext Exec(StateTree, StateTree, InstanceData);
 		const bool bInitSucceeded = Exec.IsValid();
-		AITEST_TRUE("StateTree should init", bInitSucceeded);
+		AITEST_TRUE(TEXT("StateTree should init"), bInitSucceeded);
 
 		Exec.Start();
 		Exec.LogClear();
@@ -234,12 +234,12 @@ struct FStateTreeTest_Delegate_SelfRemoval : FStateTreeTestBase
 		FStateTreeCompilerLog Log;
 		FStateTreeCompiler Compiler(Log);
 		const bool bResult = Compiler.Compile(StateTree);
-		AITEST_TRUE("StateTree should get compiled", bResult);
+		AITEST_TRUE(TEXT("StateTree should get compiled"), bResult);
 
 		FStateTreeInstanceData InstanceData;
 		FTestStateTreeExecutionContext Exec(StateTree, StateTree, InstanceData);
 		const bool bInitSucceeded = Exec.IsValid();
-		AITEST_TRUE("StateTree should init", bInitSucceeded);
+		AITEST_TRUE(TEXT("StateTree should init"), bInitSucceeded);
 
 		Exec.Start();
 		Exec.LogClear();
@@ -280,25 +280,25 @@ struct FStateTreeTest_Delegate_WithoutRemoval : FStateTreeTestBase
 		FStateTreeCompilerLog Log;
 		FStateTreeCompiler Compiler(Log);
 		const bool bResult = Compiler.Compile(StateTree);
-		AITEST_TRUE("StateTree should get compiled", bResult);
+		AITEST_TRUE(TEXT("StateTree should get compiled"), bResult);
 
 		FStateTreeInstanceData InstanceData;
 		FTestStateTreeExecutionContext Exec(StateTree, StateTree, InstanceData);
 		const bool bInitSucceeded = Exec.IsValid();
-		AITEST_TRUE("StateTree should init", bInitSucceeded);
+		AITEST_TRUE(TEXT("StateTree should init"), bInitSucceeded);
 
 		Exec.Start();
 		Exec.LogClear();
-		AITEST_TRUE("StateTree Active States should be in Root/A", Exec.ExpectInActiveStates(Root.Name, StateA.Name));
+		AITEST_TRUE(TEXT("StateTree Active States should be in Root/A"), Exec.ExpectInActiveStates(Root.Name, StateA.Name));
 
 		Exec.Tick(0.1f);
 		AITEST_TRUE(TEXT("StateTree Delegate should be triggered once."), Exec.Expect(ListenerTask.GetName(), *FString::Printf(TEXT("OnDelegate%d"), 1)));
-		AITEST_TRUE("StateTree Active States should be in Root/B", Exec.ExpectInActiveStates(Root.Name, StateB.Name));
+		AITEST_TRUE(TEXT("StateTree Active States should be in Root/B"), Exec.ExpectInActiveStates(Root.Name, StateB.Name));
 		Exec.LogClear();
 
 		Exec.Tick(0.1f);
 		AITEST_FALSE(TEXT("StateTree Delegate shouldn't be triggered again."), Exec.Expect(ListenerTask.GetName()));
-		AITEST_TRUE("StateTree Active States should be in Root/B", Exec.ExpectInActiveStates(Root.Name, StateB.Name));
+		AITEST_TRUE(TEXT("StateTree Active States should be in Root/B"), Exec.ExpectInActiveStates(Root.Name, StateB.Name));
 		Exec.LogClear();
 
 		Exec.Stop();
@@ -327,12 +327,12 @@ struct FStateTreeTest_Delegate_GlobalDispatcherAndListener : FStateTreeTestBase
 		FStateTreeCompilerLog Log;
 		FStateTreeCompiler Compiler(Log);
 		const bool bResult = Compiler.Compile(StateTree);
-		AITEST_TRUE("StateTree should get compiled", bResult);
+		AITEST_TRUE(TEXT("StateTree should get compiled"), bResult);
 
 		FStateTreeInstanceData InstanceData;
 		FTestStateTreeExecutionContext Exec(StateTree, StateTree, InstanceData);
 		const bool bInitSucceeded = Exec.IsValid();
-		AITEST_TRUE("StateTree should init", bInitSucceeded);
+		AITEST_TRUE(TEXT("StateTree should init"), bInitSucceeded);
 
 		Exec.Start();
 		Exec.LogClear();
@@ -369,12 +369,12 @@ struct FStateTreeTest_Delegate_ListeningToDelegateOnExit : FStateTreeTestBase
 		FStateTreeCompilerLog Log;
 		FStateTreeCompiler Compiler(Log);
 		const bool bResult = Compiler.Compile(StateTree);
-		AITEST_TRUE("StateTree should get compiled", bResult);
+		AITEST_TRUE(TEXT("StateTree should get compiled"), bResult);
 
 		FStateTreeInstanceData InstanceData;
 		FTestStateTreeExecutionContext Exec(StateTree, StateTree, InstanceData);
 		const bool bInitSucceeded = Exec.IsValid();
-		AITEST_TRUE("StateTree should init", bInitSucceeded);
+		AITEST_TRUE(TEXT("StateTree should init"), bInitSucceeded);
 
 		Exec.Start();
 		AITEST_TRUE("Expected Root to be active.", Exec.ExpectInActiveStates(Root.Name));
diff --git a/Engine/Plugins/Runtime/StateTree/Source/StateTreeTestSuite/Private/StateTreeLinkedStateTest.cpp b/Engine/Plugins/Runtime/StateTree/Source/StateTreeTestSuite/Private/StateTreeLinkedStateTest.cpp
--- a/Engine/Plugins/Runtime/StateTree/Source/StateTreeTestSuite/Private/StateTreeLinkedStateTest.cpp
+++ b/Engine/Plugins/Runtime/StateTree/Source/StateTreeTestSuite/Private/StateTreeLinkedStateTest.cpp
@@ -67,7 +67,7 @@ struct FStateTreeTest_FailEnterLinkedAsset : FStateTreeTestBase
 
 		FStateTreeCompiler Compiler2(Log);
 		const bool bResult2 = Compiler2.Compile(StateTree2);
-		AITEST_TRUE("StateTree2 should get compiled", bResult2);
+		AITEST_TRUE(TEXT("StateTree2 should get compiled"), bResult2);
 
 		// Main asset
 		UStateTree& StateTree = NewStateTree();
@@ -82,24 +82,24 @@ struct FStateTreeTest_FailEnterLinkedAsset : FStateTreeTestBase
 
 		FStateTreeCompiler Compiler(Log);
 		const bool bResult = Compiler.Compile(StateTree);
-		AITEST_TRUE("StateTree should get compiled", bResult);
+		AITEST_TRUE(TEXT("StateTree should get compiled"), bResult);
 
-		const FString EnterStateStr(TEXT("EnterState"));
-		const FString ExitStateStr(TEXT("ExitState"));
+		const TCHAR* EnterStateStr(TEXT("EnterState"));
+		const TCHAR* ExitStateStr(TEXT("ExitState"));
 
 		{
 			EStateTreeRunStatus Status = EStateTreeRunStatus::Unset;
 			FStateTreeInstanceData InstanceData;
 			FTestStateTreeExecutionContext Exec(StateTree, StateTree, InstanceData);
 			const bool bInitSucceeded = Exec.IsValid();
-			AITEST_TRUE("StateTree should init", bInitSucceeded);
+			AITEST_TRUE(TEXT("StateTree should init"), bInitSucceeded);
 
 			Status = Exec.Start();
-			AITEST_EQUAL("Start should complete with Running", Status, EStateTreeRunStatus::Running);
-			AITEST_TRUE("StateTree should enter GlobalTask2", Exec.Expect(GlobalTask2.GetName(), EnterStateStr));
-			AITEST_TRUE("StateTree should exit GlobalTask2", Exec.Expect(GlobalTask2.GetName(), ExitStateStr));
-			AITEST_FALSE("StateTree should not enter Task2", Exec.Expect(Task2.GetName(), EnterStateStr));
-			AITEST_TRUE("StateTree should enter Task1", Exec.Expect(Task1.GetName(), EnterStateStr));
+			AITEST_EQUAL(TEXT("Start should complete with Running"), Status, EStateTreeRunStatus::Running);
+			AITEST_TRUE(TEXT("StateTree should enter GlobalTask2"), Exec.Expect(GlobalTask2.GetName(), EnterStateStr));
+			AITEST_TRUE(TEXT("StateTree should exit GlobalTask2"), Exec.Expect(GlobalTask2.GetName(), ExitStateStr));
+			AITEST_FALSE(TEXT("StateTree should not enter Task2"), Exec.Expect(Task2.GetName(), EnterStateStr));
+			AITEST_TRUE(TEXT("StateTree should enter Task1"), Exec.Expect(Task1.GetName(), EnterStateStr));
 
 			Exec.Stop();
 		}
@@ -125,7 +125,7 @@ struct FStateTreeTest_EnterAndExitLinkedAsset : FStateTreeTestBase
 
 		FStateTreeCompiler Compiler2(Log);
 		const bool bResult2 = Compiler2.Compile(StateTree2);
-		AITEST_TRUE("StateTree2 should get compiled", bResult2);
+		AITEST_TRUE(TEXT("StateTree2 should get compiled"), bResult2);
 
 		// Main asset
 		UStateTree& StateTree = NewStateTree();
@@ -141,34 +141,34 @@ struct FStateTreeTest_EnterAndExitLinkedAsset : FStateTreeTestBase
 
 		FStateTreeCompiler Compiler(Log);
 		const bool bResult = Compiler.Compile(StateTree);
-		AITEST_TRUE("StateTree should get compiled", bResult);
+		AITEST_TRUE(TEXT("StateTree should get compiled"), bResult);
 
-		const FString EnterStateStr(TEXT("EnterState"));
-		const FString ExitStateStr(TEXT("ExitState"));
+		const TCHAR* EnterStateStr(TEXT("EnterState"));
+		const TCHAR* ExitStateStr(TEXT("ExitState"));
 
 		{
 			EStateTreeRunStatus Status = EStateTreeRunStatus::Unset;
 			FStateTreeInstanceData InstanceData;
 			FTestStateTreeExecutionContext Exec(StateTree, StateTree, InstanceData);
 			const bool bInitSucceeded = Exec.IsValid();
-			AITEST_TRUE("StateTree should init", bInitSucceeded);
+			AITEST_TRUE(TEXT("StateTree should init"), bInitSucceeded);
 
 			Status = Exec.Start();
-			AITEST_EQUAL("Start should complete with Running", Status, EStateTreeRunStatus::Running);
-			AITEST_TRUE("StateTree should enter GlobalTask2", Exec.Expect(GlobalTask2.GetName(), EnterStateStr));
-			AITEST_FALSE("StateTree should not exit GlobalTask2", Exec.Expect(GlobalTask2.GetName(), ExitStateStr));
-			AITEST_TRUE("StateTree should enter Task2", Exec.Expect(Task2.GetName(), EnterStateStr));
-			AITEST_FALSE("StateTree should not exit Task2", Exec.Expect(Task2.GetName(), ExitStateStr));
-			AITEST_FALSE("StateTree should not enter Task1", Exec.Expect(Task1.GetName(), EnterStateStr));
+			AITEST_EQUAL(TEXT("Start should complete with Running"), Status, EStateTreeRunStatus::Running);
+			AITEST_TRUE(TEXT("StateTree should enter GlobalTask2"), Exec.Expect(GlobalTask2.GetName(), EnterStateStr));
+			AITEST_FALSE(TEXT("StateTree should not exit GlobalTask2"), Exec.Expect(GlobalTask2.GetName(), ExitStateStr));
+			AITEST_TRUE(TEXT("StateTree should enter Task2"), Exec.Expect(Task2.GetName(), EnterStateStr));
+			AITEST_FALSE(TEXT("StateTree should not exit Task2"), Exec.Expect(Task2.GetName(), ExitStateStr));
+			AITEST_FALSE(TEXT("StateTree should not enter Task1"), Exec.Expect(Task1.GetName(), EnterStateStr));
 			Exec.LogClear();
 
 			Status = Exec.Tick(0.1f);
-			AITEST_EQUAL("Tick should complete with Running", Status, EStateTreeRunStatus::Running);
-			AITEST_FALSE("StateTree should not enter GlobalTask2", Exec.Expect(GlobalTask2.GetName(), EnterStateStr));
-			AITEST_TRUE("StateTree should exit GlobalTask2", Exec.Expect(GlobalTask2.GetName(), ExitStateStr));
-			AITEST_FALSE("StateTree should not enter Task2", Exec.Expect(Task2.GetName(), EnterStateStr));
-			AITEST_TRUE("StateTree should exit Task2", Exec.Expect(Task2.GetName(), ExitStateStr));
-			AITEST_TRUE("StateTree should enter Task1", Exec.Expect(Task1.GetName(), EnterStateStr));
+			AITEST_EQUAL(TEXT("Tick should complete with Running"), Status, EStateTreeRunStatus::Running);
+			AITEST_FALSE(TEXT("StateTree should not enter GlobalTask2"), Exec.Expect(GlobalTask2.GetName(), EnterStateStr));
+			AITEST_TRUE(TEXT("StateTree should exit GlobalTask2"), Exec.Expect(GlobalTask2.GetName(), ExitStateStr));
+			AITEST_FALSE(TEXT("StateTree should not enter Task2"), Exec.Expect(Task2.GetName(), EnterStateStr));
+			AITEST_TRUE(TEXT("StateTree should exit Task2"), Exec.Expect(Task2.GetName(), ExitStateStr));
+			AITEST_TRUE(TEXT("StateTree should enter Task1"), Exec.Expect(Task1.GetName(), EnterStateStr));
 			Exec.LogClear();
 
 			Exec.Stop();
@@ -234,7 +234,7 @@ struct FStateTreeTest_MultipleSameLinkedAsset : FStateTreeTestBase
 
 			FStateTreeCompiler Compiler2(Log);
 			const bool bResult2 = Compiler2.Compile(StateTree2);
-			AITEST_TRUE("StateTree2 should get compiled", bResult2);
+			AITEST_TRUE(TEXT("StateTree2 should get compiled"), bResult2);
 		}
 
 		// Main asset
@@ -299,7 +299,7 @@ struct FStateTreeTest_MultipleSameLinkedAsset : FStateTreeTestBase
 
 			FStateTreeCompiler Compiler(Log);
 			const bool bResult = Compiler.Compile(StateTree1);
-			AITEST_TRUE("StateTree should get compiled", bResult);
+			AITEST_TRUE(TEXT("StateTree should get compiled"), bResult);
 		}
 
 		Private::FScopedCVarBool CVarTickGlobalNodesWithHierarchy = TEXT("StateTree.TickGlobalNodesFollowingTreeHierarchy");
@@ -312,164 +312,164 @@ struct FStateTreeTest_MultipleSameLinkedAsset : FStateTreeTestBase
 			FTestStateTreeExecutionContext Exec(StateTree1, StateTree1, InstanceData);
 			{
 				const bool bInitSucceeded = Exec.IsValid();
-				AITEST_TRUE("StateTree should init", bInitSucceeded);
+				AITEST_TRUE(TEXT("StateTree should init"), bInitSucceeded);
 			}
 			{
 				EStateTreeRunStatus Status = Exec.Start();
-				AITEST_EQUAL("Start should complete with Running", Status, EStateTreeRunStatus::Running);
+				AITEST_EQUAL(TEXT("Start should complete with Running"), Status, EStateTreeRunStatus::Running);
 				FTestStateTreeExecutionContext::FLogOrder LogOrder = Exec.Expect("Tree1GlobalTaskA", TEXT("EnterState99"));
-				AITEST_TRUE("Start should enter Tree1GlobalTaskA", LogOrder);
+				AITEST_TRUE(TEXT("Start should enter Tree1GlobalTaskA"), LogOrder);
 				LogOrder = LogOrder.Then("Tree1StateRootTaskA", TEXT("EnterState88"));
-				AITEST_TRUE("Start should enter Tree1StateRootTaskA", LogOrder);
+				AITEST_TRUE(TEXT("Start should enter Tree1StateRootTaskA"), LogOrder);
 				LogOrder = LogOrder.Then("Tree1StateATaskA", TEXT("EnterState1"));
-				AITEST_TRUE("Start should enter Tree1StateATaskA", LogOrder);
+				AITEST_TRUE(TEXT("Start should enter Tree1StateATaskA"), LogOrder);
 				Exec.LogClear();
 			}
 			{
 				EStateTreeRunStatus Status = Exec.Tick(1.5f); // over tick, should trigger
-				AITEST_EQUAL("1st Tick should complete with Running", Status, EStateTreeRunStatus::Running);
+				AITEST_EQUAL(TEXT("1st Tick should complete with Running"), Status, EStateTreeRunStatus::Running);
 				FTestStateTreeExecutionContext::FLogOrder LogOrder = Exec.Expect("Tree1GlobalTaskA", TEXT("Tick99"));
-				AITEST_TRUE("1st should tick tasks Tree1GlobalTaskA", LogOrder);
+				AITEST_TRUE(TEXT("1st should tick tasks Tree1GlobalTaskA"), LogOrder);
 				LogOrder = LogOrder.Then("Tree1StateRootTaskA", TEXT("Tick88"));
-				AITEST_TRUE("1st should tick tasks Tree1StateRootTaskA", LogOrder);
+				AITEST_TRUE(TEXT("1st should tick tasks Tree1StateRootTaskA"), LogOrder);
 				LogOrder = LogOrder.Then("Tree1StateATaskA", TEXT("Tick1"));
-				AITEST_TRUE("1st should tick tasks Tree1StateATaskA", LogOrder);
+				AITEST_TRUE(TEXT("1st should tick tasks Tree1StateATaskA"), LogOrder);
 				Exec.LogClear();
 			}
 			{
 				EStateTreeRunStatus Status = Exec.Tick(1.0f);
-				AITEST_EQUAL("2nd Tick should complete with Running", Status, EStateTreeRunStatus::Running);
+				AITEST_EQUAL(TEXT("2nd Tick should complete with Running"), Status, EStateTreeRunStatus::Running);
 				FTestStateTreeExecutionContext::FLogOrder LogOrder = Exec.Expect("Tree1GlobalTaskA", TEXT("Tick99"));
-				AITEST_TRUE("2nd Tick should tick Tree1GlobalTaskA", LogOrder);
+				AITEST_TRUE(TEXT("2nd Tick should tick Tree1GlobalTaskA"), LogOrder);
 				LogOrder = LogOrder.Then("Tree1StateRootTaskA", TEXT("Tick88"));
-				AITEST_TRUE("2nd Tick should tick Tree1StateRootTaskA", LogOrder);
+				AITEST_TRUE(TEXT("2nd Tick should tick Tree1StateRootTaskA"), LogOrder);
 				LogOrder = LogOrder.Then("Tree1StateATaskA", TEXT("Tick1"));
-				AITEST_TRUE("2nd Tick should tick Tree1StateATaskA", LogOrder);
+				AITEST_TRUE(TEXT("2nd Tick should tick Tree1StateATaskA"), LogOrder);
 				LogOrder = LogOrder.Then("Tree1StateATaskA", TEXT("ExitState1"));
-				AITEST_TRUE("2nd Tick should exit Tree1StateATaskA", LogOrder);
+				AITEST_TRUE(TEXT("2nd Tick should exit Tree1StateATaskA"), LogOrder);
 				LogOrder = LogOrder.Then("Tree1StateBTaskA", TEXT("EnterState2"));
-				AITEST_TRUE("2nd Tick should enter Tree1StateBTaskA", LogOrder);
+				AITEST_TRUE(TEXT("2nd Tick should enter Tree1StateBTaskA"), LogOrder);
 				Exec.LogClear();
 			}
 			{
 				EStateTreeRunStatus Status = Exec.Tick(1.0f);
-				AITEST_EQUAL("3rd Tick should complete with Running", Status, EStateTreeRunStatus::Running);
+				AITEST_EQUAL(TEXT("3rd Tick should complete with Running"), Status, EStateTreeRunStatus::Running);
 				FTestStateTreeExecutionContext::FLogOrder LogOrder = Exec.Expect("Tree1GlobalTaskA", TEXT("Tick99"));
-				AITEST_TRUE("3rd should tick Tree1GlobalTaskA", LogOrder);
+				AITEST_TRUE(TEXT("3rd should tick Tree1GlobalTaskA"), LogOrder);
 				LogOrder = LogOrder.Then("Tree1StateRootTaskA", TEXT("Tick88"));
-				AITEST_TRUE("3rd should tick Tree1StateRootTaskA", LogOrder);
+				AITEST_TRUE(TEXT("3rd should tick Tree1StateRootTaskA"), LogOrder);
 				LogOrder = LogOrder.Then("Tree1StateBTaskA", TEXT("Tick2"));
-				AITEST_TRUE("3rd should tick Tree1StateBTaskA", LogOrder);
+				AITEST_TRUE(TEXT("3rd should tick Tree1StateBTaskA"), LogOrder);
 				Exec.LogClear();
 			}
 			{
 				EStateTreeRunStatus Status = Exec.Tick(1.0f);
-				AITEST_EQUAL("4th Tick should complete with Running", Status, EStateTreeRunStatus::Running);
+				AITEST_EQUAL(TEXT("4th Tick should complete with Running"), Status, EStateTreeRunStatus::Running);
 				FTestStateTreeExecutionContext::FLogOrder LogOrder = Exec.Expect("Tree1GlobalTaskA", TEXT("Tick99"));
-				AITEST_TRUE("4th Tick should tick Tree1GlobalTaskA", LogOrder);
+				AITEST_TRUE(TEXT("4th Tick should tick Tree1GlobalTaskA"), LogOrder);
 				LogOrder = LogOrder.Then("Tree1StateRootTaskA", TEXT("Tick88"));
-				AITEST_TRUE("4th Tick should tick Tree1StateRootTaskA", LogOrder);
+				AITEST_TRUE(TEXT("4th Tick should tick Tree1StateRootTaskA"), LogOrder);
 				LogOrder = LogOrder.Then("Tree1StateBTaskA", TEXT("Tick2"));
-				AITEST_TRUE("4th Tick should tick Tree1StateBTaskA", LogOrder);
+				AITEST_TRUE(TEXT("4th Tick should tick Tree1StateBTaskA"), LogOrder);
 				LogOrder = LogOrder.Then("Tree1StateBTaskA", TEXT("ExitState2"));
-				AITEST_TRUE("4th Tick should exit Tree1StateBTaskA", LogOrder);
+				AITEST_TRUE(TEXT("4th Tick should exit Tree1StateBTaskA"), LogOrder);
 				//@todo: bug: order is wrong
 				//LogOrder = LogOrder.Then("Tree2GlobalTaskA", TEXT("EnterState111"))
-				//AITEST_TRUE("4th Tick should enter Tree2GlobalTaskA", LogOrder);
+				//AITEST_TRUE(TEXT("4th Tick should enter Tree2GlobalTaskA"), LogOrder);
 				LogOrder = LogOrder.Then("Tree2StateRootTaskA", TEXT("EnterState111"));
-				AITEST_TRUE("4th Tick should enter Tree2StateRootTaskA", LogOrder);
+				AITEST_TRUE(TEXT("4th Tick should enter Tree2StateRootTaskA"), LogOrder);
 				LogOrder = LogOrder.Then("Tree2StateATaskA", TEXT("EnterState111"));
-				AITEST_TRUE("4th Tick should enter Tree2StateATaskA", LogOrder);
-				AITEST_TRUE("4th Tick should tick tasks", Exec.Expect("Tree2GlobalTaskA", TEXT("EnterState111")));
+				AITEST_TRUE(TEXT("4th Tick should enter Tree2StateATaskA"), LogOrder);
+				AITEST_TRUE(TEXT("4th Tick should tick tasks"), Exec.Expect("Tree2GlobalTaskA", TEXT("EnterState111")));
 				Exec.LogClear();
 			}
 			{
 				EStateTreeRunStatus Status = Exec.Tick(0.001f);
-				AITEST_EQUAL("5th Tick should complete with Running", Status, EStateTreeRunStatus::Running);
+				AITEST_EQUAL(TEXT("5th Tick should complete with Running"), Status, EStateTreeRunStatus::Running);
 				if (bTickGlobalNodesWithHierarchy)
 				{
 					FTestStateTreeExecutionContext::FLogOrder LogOrder = Exec.Expect("Tree1GlobalTaskA", TEXT("Tick99"));
-					AITEST_TRUE("5h Tick should tick Tree1GlobalTaskA", LogOrder);
+					AITEST_TRUE(TEXT("5h Tick should tick Tree1GlobalTaskA"), LogOrder);
 					LogOrder = LogOrder.Then("Tree1StateRootTaskA", TEXT("Tick88"));
-					AITEST_TRUE("5h Tick should tick Tree1StateRootTaskA", LogOrder);
+					AITEST_TRUE(TEXT("5h Tick should tick Tree1StateRootTaskA"), LogOrder);
 					LogOrder = LogOrder.Then("Tree2GlobalTaskA", TEXT("Tick111"));
-					AITEST_TRUE("5h Tick should tick Tree2GlobalTaskA", LogOrder);
+					AITEST_TRUE(TEXT("5h Tick should tick Tree2GlobalTaskA"), LogOrder);
 					LogOrder = LogOrder.Then("Tree2StateRootTaskA", TEXT("Tick111"));
-					AITEST_TRUE("5h Tick should tick Tree2StateRootTaskA", LogOrder);
+					AITEST_TRUE(TEXT("5h Tick should tick Tree2StateRootTaskA"), LogOrder);
 					LogOrder = LogOrder.Then("Tree2StateATaskA", TEXT("Tick111"));
-					AITEST_TRUE("5h Tick should tick Tree2StateATaskA", LogOrder);
+					AITEST_TRUE(TEXT("5h Tick should tick Tree2StateATaskA"), LogOrder);
 				}
 				else
 				{
 					FTestStateTreeExecutionContext::FLogOrder LogOrder = Exec.Expect("Tree1GlobalTaskA", TEXT("Tick99"));
-					AITEST_TRUE("5h Tick should tick Tree1GlobalTaskA", LogOrder);
+					AITEST_TRUE(TEXT("5h Tick should tick Tree1GlobalTaskA"), LogOrder);
 					LogOrder = LogOrder.Then("Tree2GlobalTaskA", TEXT("Tick111"));
-					AITEST_TRUE("5h Tick should tick Tree2GlobalTaskA", LogOrder);
+					AITEST_TRUE(TEXT("5h Tick should tick Tree2GlobalTaskA"), LogOrder);
 					LogOrder = LogOrder.Then("Tree1StateRootTaskA", TEXT("Tick88"));
-					AITEST_TRUE("5h Tick should tick Tree1StateRootTaskA", LogOrder);
+					AITEST_TRUE(TEXT("5h Tick should tick Tree1StateRootTaskA"), LogOrder);
 					LogOrder = LogOrder.Then("Tree2StateRootTaskA", TEXT("Tick111"));
-					AITEST_TRUE("5h Tick should tick Tree2StateRootTaskA", LogOrder);
+					AITEST_TRUE(TEXT("5h Tick should tick Tree2StateRootTaskA"), LogOrder);
 					LogOrder = LogOrder.Then("Tree2StateATaskA", TEXT("Tick111"));
-					AITEST_TRUE("5h Tick should tick Tree2StateATaskA", LogOrder);
+					AITEST_TRUE(TEXT("5h Tick should tick Tree2StateATaskA"), LogOrder);
 				}
 				Exec.LogClear();
 			}
 			{
 				EStateTreeRunStatus Status = Exec.Tick(1.0f);
-				AITEST_EQUAL("6th Tick should complete with Running", Status, EStateTreeRunStatus::Running);
+				AITEST_EQUAL(TEXT("6th Tick should complete with Running"), Status, EStateTreeRunStatus::Running);
 				if (bTickGlobalNodesWithHierarchy)
 				{
 					FTestStateTreeExecutionContext::FLogOrder LogOrder = Exec.Expect("Tree1GlobalTaskA", TEXT("Tick99"));
-					AITEST_TRUE("6th Tick should tick Tree1GlobalTaskA", LogOrder);
+					AITEST_TRUE(TEXT("6th Tick should tick Tree1GlobalTaskA"), LogOrder);
 					LogOrder = LogOrder.Then("Tree1StateRootTaskA", TEXT("Tick88"));
-					AITEST_TRUE("6th Tick should tick Tree1StateRootTaskA", LogOrder);
+					AITEST_TRUE(TEXT("6th Tick should tick Tree1StateRootTaskA"), LogOrder);
 					LogOrder = LogOrder.Then("Tree2GlobalTaskA", TEXT("Tick111"));
-					AITEST_TRUE("6th Tick should tick Tree2GlobalTaskA", LogOrder);
+					AITEST_TRUE(TEXT("6th Tick should tick Tree2GlobalTaskA"), LogOrder);
 					LogOrder = LogOrder.Then("Tree2StateRootTaskA", TEXT("Tick111"));
-					AITEST_TRUE("6th Tick should tick Tree2StateRootTaskA", LogOrder);
+					AITEST_TRUE(TEXT("6th Tick should tick Tree2StateRootTaskA"), LogOrder);
 					LogOrder = LogOrder.Then("Tree2StateATaskA", TEXT("Tick111"));
-					AITEST_TRUE("6th Tick should tick Tree2StateATaskA", LogOrder);
+					AITEST_TRUE(TEXT("6th Tick should tick Tree2StateATaskA"), LogOrder);
 					LogOrder = LogOrder.Then("Tree2StateATaskA", TEXT("ExitState111"));
-					AITEST_TRUE("6th Tick should exit Tree2StateATaskA", LogOrder);
+					AITEST_TRUE(TEXT("6th Tick should exit Tree2StateATaskA"), LogOrder);
 					LogOrder = LogOrder.Then("Tree2StateRootTaskA", TEXT("ExitState111"));
-					AITEST_TRUE("6th Tick should exit Tree2StateRootTaskA", LogOrder);
+					AITEST_TRUE(TEXT("6th Tick should exit Tree2StateRootTaskA"), LogOrder);
 					//@todo: bug: The exit state is not called on the global, it's a new frame.
 					//LogOrder = LogOrder.Then("Tree2GlobalTaskA", TEXT("ExitState111"));
-					//AITEST_TRUE("6th Tick should exit Tree2GlobalTaskA", LogOrder);
+					//AITEST_TRUE(TEXT("6th Tick should exit Tree2GlobalTaskA"), LogOrder);
 					//@todo: bug: The enter state is not called on the global, it's a new frame.
 					//LogOrder = LogOrder.Then("Tree2GlobalTaskA", TEXT("EnterState222"));
-					//AITEST_TRUE("6th Tick should enter Tree2GlobalTaskA", LogOrder);
+					//AITEST_TRUE(TEXT("6th Tick should enter Tree2GlobalTaskA"), LogOrder);
 					LogOrder = LogOrder.Then("Tree2StateRootTaskA", TEXT("EnterState222"));
-					AITEST_TRUE("6th Tick should enter Tree2StateRootTaskA", LogOrder);
+					AITEST_TRUE(TEXT("6th Tick should enter Tree2StateRootTaskA"), LogOrder);
 					LogOrder = LogOrder.Then("Tree2StateATaskA", TEXT("EnterState222"));
-					AITEST_TRUE("6th Tick should enter Tree2StateATaskA", LogOrder);
+					AITEST_TRUE(TEXT("6th Tick should enter Tree2StateATaskA"), LogOrder);
 				}
 				else
 				{
 					FTestStateTreeExecutionContext::FLogOrder LogOrder = Exec.Expect("Tree1GlobalTaskA", TEXT("Tick99"));
-					AITEST_TRUE("6th Tick should tick Tree1GlobalTaskA", LogOrder);
+					AITEST_TRUE(TEXT("6th Tick should tick Tree1GlobalTaskA"), LogOrder);
 					LogOrder = LogOrder.Then("Tree2GlobalTaskA", TEXT("Tick111"));
-					AITEST_TRUE("6th Tick should tick Tree2GlobalTaskA", LogOrder);
+					AITEST_TRUE(TEXT("6th Tick should tick Tree2GlobalTaskA"), LogOrder);
 					LogOrder = LogOrder.Then("Tree1StateRootTaskA", TEXT("Tick88"));
-					AITEST_TRUE("6th Tick should tick Tree1StateRootTaskA", LogOrder);
+					AITEST_TRUE(TEXT("6th Tick should tick Tree1StateRootTaskA"), LogOrder);
 					LogOrder = LogOrder.Then("Tree2StateRootTaskA", TEXT("Tick111"));
-					AITEST_TRUE("6th Tick should tick Tree2StateRootTaskA", LogOrder);
+					AITEST_TRUE(TEXT("6th Tick should tick Tree2StateRootTaskA"), LogOrder);
 					LogOrder = LogOrder.Then("Tree2StateATaskA", TEXT("Tick111"));
-					AITEST_TRUE("6th Tick should tick Tree2StateATaskA", LogOrder);
+					AITEST_TRUE(TEXT("6th Tick should tick Tree2StateATaskA"), LogOrder);
 					LogOrder = LogOrder.Then("Tree2StateATaskA", TEXT("ExitState111"));
-					AITEST_TRUE("6th Tick should exit Tree2StateATaskA", LogOrder);
+					AITEST_TRUE(TEXT("6th Tick should exit Tree2StateATaskA"), LogOrder);
 					LogOrder = LogOrder.Then("Tree2StateRootTaskA", TEXT("ExitState111"));
-					AITEST_TRUE("6th Tick should exit Tree2StateRootTaskA", LogOrder);
+					AITEST_TRUE(TEXT("6th Tick should exit Tree2StateRootTaskA"), LogOrder);
 					//@todo: bug: The exit state is not called on the global, it's a new frame.
 					//LogOrder = LogOrder.Then("Tree2GlobalTaskA", TEXT("ExitState111"));
-					//AITEST_TRUE("6th Tick should exit Tree2GlobalTaskA", LogOrder);
+					//AITEST_TRUE(TEXT("6th Tick should exit Tree2GlobalTaskA"), LogOrder);
 					//@todo: bug: The enter state is not called on the global, it's a new frame.
 					//LogOrder = LogOrder.Then("Tree2GlobalTaskA", TEXT("EnterState222"));
-					//AITEST_TRUE("6th Tick should enter Tree2GlobalTaskA", LogOrder);
+					//AITEST_TRUE(TEXT("6th Tick should enter Tree2GlobalTaskA"), LogOrder);
 					LogOrder = LogOrder.Then("Tree2StateRootTaskA", TEXT("EnterState222"));
-					AITEST_TRUE("6th Tick should enter Tree2StateRootTaskA", LogOrder);
+					AITEST_TRUE(TEXT("6th Tick should enter Tree2StateRootTaskA"), LogOrder);
 					LogOrder = LogOrder.Then("Tree2StateATaskA", TEXT("EnterState222"));
-					AITEST_TRUE("6th Tick should enter Tree2StateATaskA", LogOrder);
+					AITEST_TRUE(TEXT("6th Tick should enter Tree2StateATaskA"), LogOrder);
 				}
 				Exec.LogClear();
 			}
@@ -535,7 +535,7 @@ struct FStateTreeTest_EmptyStateWithTickTransitionLinkedAsset : FStateTreeTestBa
 
 			FStateTreeCompiler Compiler2(Log);
 			const bool bResult2 = Compiler2.Compile(StateTree2);
-			AITEST_TRUE("StateTree2 should get compiled", bResult2);
+			AITEST_TRUE(TEXT("StateTree2 should get compiled"), bResult2);
 		}
 
 		// Main asset
@@ -577,7 +577,7 @@ struct FStateTreeTest_EmptyStateWithTickTransitionLinkedAsset : FStateTreeTestBa
 
 			FStateTreeCompiler Compiler(Log);
 			const bool bResult = Compiler.Compile(StateTree1);
-			AITEST_TRUE("StateTree should get compiled", bResult);
+			AITEST_TRUE(TEXT("StateTree should get compiled"), bResult);
 		}
 
 		Private::FScopedCVarBool CVarTickGlobalNodesWithHierarchy = TEXT("StateTree.TickGlobalNodesFollowingTreeHierarchy");
@@ -590,133 +590,133 @@ struct FStateTreeTest_EmptyStateWithTickTransitionLinkedAsset : FStateTreeTestBa
 			FTestStateTreeExecutionContext Exec(StateTree1, StateTree1, InstanceData);
 			{
 				const bool bInitSucceeded = Exec.IsValid();
-				AITEST_TRUE("StateTree should init", bInitSucceeded);
+				AITEST_TRUE(TEXT("StateTree should init"), bInitSucceeded);
 			}
 			{
 				EStateTreeRunStatus Status = Exec.Start();
-				AITEST_EQUAL("Start should complete with Running", Status, EStateTreeRunStatus::Running);
+				AITEST_EQUAL(TEXT("Start should complete with Running"), Status, EStateTreeRunStatus::Running);
 				FTestStateTreeExecutionContext::FLogOrder LogOrder = Exec.Expect("Tree1GlobalTaskA", TEXT("EnterState11"));
-				AITEST_TRUE("Start enters in the correct order Tree1GlobalTaskA", LogOrder);
+				AITEST_TRUE(TEXT("Start enters in the correct order Tree1GlobalTaskA"), LogOrder);
 				LogOrder = LogOrder.Then("Tree1StateRootTaskA", TEXT("EnterState12"));
-				AITEST_TRUE("Start enters in the correct order Tree1StateRootTaskA", LogOrder);
+				AITEST_TRUE(TEXT("Start enters in the correct order Tree1StateRootTaskA"), LogOrder);
 				LogOrder = LogOrder.Then("Tree1StateATaskA", TEXT("EnterState13"));
-				AITEST_TRUE("Start enters in the correct order Tree1StateATaskA", LogOrder);
-				AITEST_TRUE("Start should be in the correct states", Exec.ExpectInActiveStates("Tree1StateRoot", "Tree1StateA"));
+				AITEST_TRUE(TEXT("Start enters in the correct order Tree1StateATaskA"), LogOrder);
+				AITEST_TRUE(TEXT("Start should be in the correct states"), Exec.ExpectInActiveStates("Tree1StateRoot", "Tree1StateA"));
 				Exec.LogClear();
 			}
 			{
 				EStateTreeRunStatus Status = Exec.Tick(1.5f); // over tick, should trigger
-				AITEST_EQUAL("1st Tick should complete with Running", Status, EStateTreeRunStatus::Running);
+				AITEST_EQUAL(TEXT("1st Tick should complete with Running"), Status, EStateTreeRunStatus::Running);
 				FTestStateTreeExecutionContext::FLogOrder LogOrder = Exec.Expect("Tree1GlobalTaskA", TEXT("Tick11"));
-				AITEST_TRUE("1st Tick should tick Tree1GlobalTaskA", LogOrder);
+				AITEST_TRUE(TEXT("1st Tick should tick Tree1GlobalTaskA"), LogOrder);
 				LogOrder = LogOrder.Then("Tree1StateRootTaskA", TEXT("Tick12"));
-				AITEST_TRUE("1st Tick should tick Tree1StateRootTaskA", LogOrder);
+				AITEST_TRUE(TEXT("1st Tick should tick Tree1StateRootTaskA"), LogOrder);
 				LogOrder = LogOrder.Then("Tree1StateATaskA", TEXT("Tick13"));
-				AITEST_TRUE("1st Tick should tick Tree1StateATaskA", LogOrder);
-				AITEST_TRUE("Start should be in the correct states", Exec.ExpectInActiveStates("Tree1StateRoot", "Tree1StateA"));
+				AITEST_TRUE(TEXT("1st Tick should tick Tree1StateATaskA"), LogOrder);
+				AITEST_TRUE(TEXT("Start should be in the correct states"), Exec.ExpectInActiveStates("Tree1StateRoot", "Tree1StateA"));
 				Exec.LogClear();
 			}
 			{
 				EStateTreeRunStatus Status = Exec.Tick(1.0f);
-				AITEST_EQUAL("2nd Tick should complete with Running", Status, EStateTreeRunStatus::Running);
+				AITEST_EQUAL(TEXT("2nd Tick should complete with Running"), Status, EStateTreeRunStatus::Running);
 				FTestStateTreeExecutionContext::FLogOrder LogOrder = Exec.Expect("Tree1GlobalTaskA", TEXT("Tick11"));
-				AITEST_TRUE("2nd Tick should tick Tree1GlobalTaskA", LogOrder);
+				AITEST_TRUE(TEXT("2nd Tick should tick Tree1GlobalTaskA"), LogOrder);
 				LogOrder = LogOrder.Then("Tree1StateRootTaskA", TEXT("Tick12"));
-				AITEST_TRUE("2nd Tick should tick Tree1StateRootTaskA", LogOrder);
+				AITEST_TRUE(TEXT("2nd Tick should tick Tree1StateRootTaskA"), LogOrder);
 				LogOrder = LogOrder.Then("Tree1StateATaskA", TEXT("Tick13"));
-				AITEST_TRUE("2nd Tick should tick Tree1StateATaskA", LogOrder);
+				AITEST_TRUE(TEXT("2nd Tick should tick Tree1StateATaskA"), LogOrder);
 				LogOrder = LogOrder.Then("Tree2GlobalTaskA", TEXT("EnterState21"));
-				AITEST_TRUE("2nd Tick should enter Tree2GlobalTaskA", LogOrder);
+				AITEST_TRUE(TEXT("2nd Tick should enter Tree2GlobalTaskA"), LogOrder);
 				LogOrder = LogOrder.Then("State Tree Test Boolean Condition", TEXT("TestCondition=0"));
-				AITEST_TRUE("2nd Tick should test Bool", LogOrder);
+				AITEST_TRUE(TEXT("2nd Tick should test Bool"), LogOrder);
 				LogOrder = LogOrder.Then("State Tree Test Boolean Condition", TEXT("TestCondition=1"));
-				AITEST_TRUE("2nd Tick should test Bool", LogOrder);
+				AITEST_TRUE(TEXT("2nd Tick should test Bool"), LogOrder);
 				LogOrder = LogOrder.Then("Tree1StateATaskA", TEXT("ExitState13"));
-				AITEST_TRUE("2nd Tick should exit Tree1StateATaskA", LogOrder);
+				AITEST_TRUE(TEXT("2nd Tick should exit Tree1StateATaskA"), LogOrder);
 				LogOrder = LogOrder.Then("Tree2StateRootTaskA", TEXT("EnterState22"));
-				AITEST_TRUE("2nd Tick should enter Tree2StateRootTaskA", LogOrder);
-				AITEST_FALSE("2nd Tick should not enter the fail state.", Exec.Expect("Tree2StateFailTaskA"));
-				AITEST_TRUE("2nd Tick should be in the correct states", Exec.ExpectInActiveStates("Tree1StateRoot", "Tree1StateLinkedTree", "Tree2StateRoot", "Tree2StateB"));
+				AITEST_TRUE(TEXT("2nd Tick should enter Tree2StateRootTaskA"), LogOrder);
+				AITEST_FALSE(TEXT("2nd Tick should not enter the fail state."), Exec.Expect("Tree2StateFailTaskA"));
+				AITEST_TRUE(TEXT("2nd Tick should be in the correct states"), Exec.ExpectInActiveStates("Tree1StateRoot", "Tree1StateLinkedTree", "Tree2StateRoot", "Tree2StateB"));
 				Exec.LogClear();
 			}
 			{
 				EStateTreeRunStatus Status = Exec.Tick(1.0f);
-				AITEST_EQUAL("3rd Tick should complete with Running", Status, EStateTreeRunStatus::Running);
-				AITEST_TRUE("3rd Tick should be in the correct states", Exec.ExpectInActiveStates("Tree1StateRoot", "Tree1StateLinkedTree", "Tree2StateRoot", "Tree2StateB"));
+				AITEST_EQUAL(TEXT("3rd Tick should complete with Running"), Status, EStateTreeRunStatus::Running);
+				AITEST_TRUE(TEXT("3rd Tick should be in the correct states"), Exec.ExpectInActiveStates("Tree1StateRoot", "Tree1StateLinkedTree", "Tree2StateRoot", "Tree2StateB"));
 				if (bTickGlobalNodesWithHierarchy)
 				{
 					FTestStateTreeExecutionContext::FLogOrder LogOrder = Exec.Expect("Tree1GlobalTaskA", TEXT("Tick11"));
-					AITEST_TRUE("3rd Tick should tick Tree1GlobalTaskA", LogOrder);
+					AITEST_TRUE(TEXT("3rd Tick should tick Tree1GlobalTaskA"), LogOrder);
 					LogOrder = LogOrder.Then("Tree1StateRootTaskA", TEXT("Tick12"));
-					AITEST_TRUE("3rd Tick should tick Tree1StateRootTaskA", LogOrder);
+					AITEST_TRUE(TEXT("3rd Tick should tick Tree1StateRootTaskA"), LogOrder);
 					LogOrder = LogOrder.Then("Tree2GlobalTaskA", TEXT("Tick21"));
-					AITEST_TRUE("3rd Tick should tick Tree2GlobalTaskA", LogOrder);
+					AITEST_TRUE(TEXT("3rd Tick should tick Tree2GlobalTaskA"), LogOrder);
 					LogOrder = LogOrder.Then("Tree2StateRootTaskA", TEXT("Tick22"));
-					AITEST_TRUE("3rd Tick should tick Tree2StateRootTaskA", LogOrder);
+					AITEST_TRUE(TEXT("3rd Tick should tick Tree2StateRootTaskA"), LogOrder);
 				}
 				else
 				{
 					FTestStateTreeExecutionContext::FLogOrder LogOrder = Exec.Expect("Tree1GlobalTaskA", TEXT("Tick11"));
-					AITEST_TRUE("3rd Tick should tick Tree1GlobalTaskA", LogOrder);
+					AITEST_TRUE(TEXT("3rd Tick should tick Tree1GlobalTaskA"), LogOrder);
 					LogOrder = LogOrder.Then("Tree2GlobalTaskA", TEXT("Tick21"));
-					AITEST_TRUE("3rd Tick should tick Tree2GlobalTaskA", LogOrder);
+					AITEST_TRUE(TEXT("3rd Tick should tick Tree2GlobalTaskA"), LogOrder);
 					LogOrder = LogOrder.Then("Tree1StateRootTaskA", TEXT("Tick12"));
-					AITEST_TRUE("3rd Tick should tick Tree1StateRootTaskA", LogOrder);
+					AITEST_TRUE(TEXT("3rd Tick should tick Tree1StateRootTaskA"), LogOrder);
 					LogOrder = LogOrder.Then("Tree2StateRootTaskA", TEXT("Tick22"));
-					AITEST_TRUE("3rd Tick should tick Tree2StateRootTaskA", LogOrder);
+					AITEST_TRUE(TEXT("3rd Tick should tick Tree2StateRootTaskA"), LogOrder);
 				}
 				Exec.LogClear();
 			}
 			{
 				EStateTreeRunStatus Status = Exec.Tick(1.0f);
-				AITEST_EQUAL("4th Tick should complete with Running", Status, EStateTreeRunStatus::Running);
+				AITEST_EQUAL(TEXT("4th Tick should complete with Running"), Status, EStateTreeRunStatus::Running);
 				if (bTickGlobalNodesWithHierarchy)
 				{
 					FTestStateTreeExecutionContext::FLogOrder LogOrder = Exec.Expect("Tree1GlobalTaskA", TEXT("Tick11"));
-					AITEST_TRUE("4th Tick should tick Tree1GlobalTaskA", LogOrder);
+					AITEST_TRUE(TEXT("4th Tick should tick Tree1GlobalTaskA"), LogOrder);
 					LogOrder = LogOrder.Then("Tree1StateRootTaskA", TEXT("Tick12"));
-					AITEST_TRUE("4th Tick should tick Tree1StateRootTaskA", LogOrder);
+					AITEST_TRUE(TEXT("4th Tick should tick Tree1StateRootTaskA"), LogOrder);
 					LogOrder = LogOrder.Then("Tree2GlobalTaskA", TEXT("Tick21"));
-					AITEST_TRUE("4th Tick should tick Tree2GlobalTaskA", LogOrder);
+					AITEST_TRUE(TEXT("4th Tick should tick Tree2GlobalTaskA"), LogOrder);
 					LogOrder = LogOrder.Then("Tree2StateRootTaskA", TEXT("Tick22"));
-					AITEST_TRUE("4th Tick should tick Tree2StateRootTaskA", LogOrder);
+					AITEST_TRUE(TEXT("4th Tick should tick Tree2StateRootTaskA"), LogOrder);
 					LogOrder = LogOrder.Then("Tree2StateRootTaskA", TEXT("ExitState22"));
-					AITEST_TRUE("4th Tick should exit Tree2StateRootTaskA", LogOrder);
+					AITEST_TRUE(TEXT("4th Tick should exit Tree2StateRootTaskA"), LogOrder);
 					LogOrder = LogOrder.Then("Tree2GlobalTaskA", TEXT("ExitState21"));
-					AITEST_TRUE("4th Tick should exit Tree2GlobalTaskA", LogOrder);
+					AITEST_TRUE(TEXT("4th Tick should exit Tree2GlobalTaskA"), LogOrder);
 					LogOrder = LogOrder.Then("Tree1StateBTaskA", TEXT("EnterState14"));
-					AITEST_TRUE("4th Tick should enter Tree1StateBTaskA", LogOrder);
+					AITEST_TRUE(TEXT("4th Tick should enter Tree1StateBTaskA"), LogOrder);
 				}
 				else
 				{
 					FTestStateTreeExecutionContext::FLogOrder LogOrder = Exec.Expect("Tree1GlobalTaskA", TEXT("Tick11"));
-					AITEST_TRUE("4th Tick should tick Tree1GlobalTaskA", LogOrder);
+					AITEST_TRUE(TEXT("4th Tick should tick Tree1GlobalTaskA"), LogOrder);
 					LogOrder = LogOrder.Then("Tree2GlobalTaskA", TEXT("Tick21"));
-					AITEST_TRUE("4th Tick should tick Tree2GlobalTaskA", LogOrder);
+					AITEST_TRUE(TEXT("4th Tick should tick Tree2GlobalTaskA"), LogOrder);
 					LogOrder = LogOrder.Then("Tree1StateRootTaskA", TEXT("Tick12"));
-					AITEST_TRUE("4th Tick should tick Tree1StateRootTaskA", LogOrder);
+					AITEST_TRUE(TEXT("4th Tick should tick Tree1StateRootTaskA"), LogOrder);
 					LogOrder = LogOrder.Then("Tree2StateRootTaskA", TEXT("Tick22"));
-					AITEST_TRUE("4th Tick should tick Tree2StateRootTaskA", LogOrder);
+					AITEST_TRUE(TEXT("4th Tick should tick Tree2StateRootTaskA"), LogOrder);
 					LogOrder = LogOrder.Then("Tree2StateRootTaskA", TEXT("ExitState22"));
-					AITEST_TRUE("4th Tick should exit Tree2StateRootTaskA", LogOrder);
+					AITEST_TRUE(TEXT("4th Tick should exit Tree2StateRootTaskA"), LogOrder);
 					LogOrder = LogOrder.Then("Tree2GlobalTaskA", TEXT("ExitState21"));
-					AITEST_TRUE("4th Tick should exit Tree2GlobalTaskA", LogOrder);
+					AITEST_TRUE(TEXT("4th Tick should exit Tree2GlobalTaskA"), LogOrder);
 					LogOrder = LogOrder.Then("Tree1StateBTaskA", TEXT("EnterState14"));
-					AITEST_TRUE("4th Tick should enter Tree1StateBTaskA", LogOrder);
+					AITEST_TRUE(TEXT("4th Tick should enter Tree1StateBTaskA"), LogOrder);
 				}
-				AITEST_TRUE("4th Tick should be in the correct states", Exec.ExpectInActiveStates("Tree1StateRoot", "Tree1StateB"));
+				AITEST_TRUE(TEXT("4th Tick should be in the correct states"), Exec.ExpectInActiveStates("Tree1StateRoot", "Tree1StateB"));
 				Exec.LogClear();
 			}
 			{
 				EStateTreeRunStatus Status = Exec.Tick(0.001f);
-				AITEST_EQUAL("5th Tick should complete with Running", Status, EStateTreeRunStatus::Running);
+				AITEST_EQUAL(TEXT("5th Tick should complete with Running"), Status, EStateTreeRunStatus::Running);
 				FTestStateTreeExecutionContext::FLogOrder LogOrder = Exec.Expect("Tree1GlobalTaskA", TEXT("Tick11"));
-				AITEST_TRUE("5th Tick should tick Tree1GlobalTaskA", LogOrder);
+				AITEST_TRUE(TEXT("5th Tick should tick Tree1GlobalTaskA"), LogOrder);
 				LogOrder = LogOrder.Then("Tree1StateRootTaskA", TEXT("Tick12"));
-				AITEST_TRUE("5th Tick should tick Tree1StateRootTaskA", LogOrder);
+				AITEST_TRUE(TEXT("5th Tick should tick Tree1StateRootTaskA"), LogOrder);
 				LogOrder = LogOrder.Then("Tree1StateBTaskA", TEXT("Tick14"));
-				AITEST_TRUE("5th Tick should tick Tree1StateBTaskA", LogOrder);
-				AITEST_TRUE("5th Tick should be in the correct states", Exec.ExpectInActiveStates("Tree1StateRoot", "Tree1StateB"));
+				AITEST_TRUE(TEXT("5th Tick should tick Tree1StateBTaskA"), LogOrder);
+				AITEST_TRUE(TEXT("5th Tick should be in the correct states"), Exec.ExpectInActiveStates("Tree1StateRoot", "Tree1StateB"));
 				Exec.LogClear();
 			}
 			Exec.Stop();
@@ -750,7 +750,7 @@ struct FStateTreeTest_RecursiveLinkedAsset : FStateTreeTestBase
 			FStateTreeCompilerLog Log;
 			FStateTreeCompiler Compiler(Log);
 			const bool bResult = Compiler.Compile(StateTree1);
-			AITEST_TRUE("StateTree 1 should get compiled", bResult);
+			AITEST_TRUE(TEXT("StateTree 1 should get compiled"), bResult);
 		}
 
 		UStateTree& StateTree2 = NewStateTree();
@@ -763,7 +763,7 @@ struct FStateTreeTest_RecursiveLinkedAsset : FStateTreeTestBase
 			FStateTreeCompilerLog Log;
 			FStateTreeCompiler Compiler(Log);
 			const bool bResult = Compiler.Compile(StateTree2);
-			AITEST_TRUE("StateTree 2 should get compiled", bResult);
+			AITEST_TRUE(TEXT("StateTree 2 should get compiled"), bResult);
 		}
 		// Asset 1 implementation
 		{
@@ -789,7 +789,7 @@ struct FStateTreeTest_RecursiveLinkedAsset : FStateTreeTestBase
 			FStateTreeCompilerLog Log;
 			FStateTreeCompiler Compiler(Log);
 			const bool bResult = Compiler.Compile(StateTree1);
-			AITEST_TRUE("StateTree should get compiled", bResult);
+			AITEST_TRUE(TEXT("StateTree should get compiled"), bResult);
 		}
 		// Asset 2 implementation
 		UStateTreeState* Tree2StateLinkedTree1 = nullptr;
@@ -818,7 +818,7 @@ struct FStateTreeTest_RecursiveLinkedAsset : FStateTreeTestBase
 			FStateTreeCompilerLog Log;
 			FStateTreeCompiler Compiler(Log);
 			const bool bResult = Compiler.Compile(StateTree2);
-			AITEST_FALSE("StateTree should not compiled", bResult);
+			AITEST_FALSE(TEXT("StateTree should not compiled"), bResult);
 		}
 		// Fix circular dependency
 		{
@@ -826,7 +826,7 @@ struct FStateTreeTest_RecursiveLinkedAsset : FStateTreeTestBase
 			FStateTreeCompilerLog Log;
 			FStateTreeCompiler Compiler(Log);
 			const bool bResult = Compiler.Compile(StateTree2);
-			AITEST_TRUE("StateTree should get compiled", bResult);
+			AITEST_TRUE(TEXT("StateTree should get compiled"), bResult);
 		}
 		// Run test
 		{
@@ -834,7 +834,7 @@ struct FStateTreeTest_RecursiveLinkedAsset : FStateTreeTestBase
 			FTestStateTreeExecutionContext Exec(StateTree1, StateTree1, InstanceData);
 			{
 				const bool bInitSucceeded = Exec.IsValid();
-				AITEST_TRUE("StateTree should init", bInitSucceeded);
+				AITEST_TRUE(TEXT("StateTree should init"), bInitSucceeded);
 			}
 
 			{
@@ -849,7 +849,7 @@ struct FStateTreeTest_RecursiveLinkedAsset : FStateTreeTestBase
 				LogOrder = LogOrder.Then("Tree2StateA", TEXT("EnterState202"));
 				AITEST_TRUE(TEXT("Start should enter Tree2StateA"), LogOrder);
 				Exec.LogClear();
-				AITEST_TRUE("Doesn't have the expected error message.", GetTestRunner().HasMetExpectedMessages());
+				AITEST_TRUE(TEXT("Doesn't have the expected error message."), GetTestRunner().HasMetExpectedMessages());
 			}
 			Exec.Stop();
 		}
@@ -886,7 +886,7 @@ struct FStateTreeTest_LinkedAssetTransitionSameTick : FStateTreeTestBase
 			FStateTreeCompilerLog Log;
 			FStateTreeCompiler Compiler(Log);
 			const bool bResult = Compiler.Compile(StateTree2);
-			AITEST_TRUE("StateTree should not compiled", bResult);
+			AITEST_TRUE(TEXT("StateTree should not compiled"), bResult);
 		}
 		// Asset 1
 		UStateTree& StateTree1 = NewStateTree();
@@ -899,7 +899,7 @@ struct FStateTreeTest_LinkedAssetTransitionSameTick : FStateTreeTestBase
 				FStateTreeCompilerLog Log;
 				FStateTreeCompiler Compiler(Log);
 				const bool bResult = Compiler.Compile(StateTree1);
-				AITEST_TRUE("StateTree 1 should get compiled", bResult);
+				AITEST_TRUE(TEXT("StateTree 1 should get compiled"), bResult);
 			}
 			{
 				TStateTreeEditorNode<FTestTask_PrintValue>& Task1 = Root1->AddTask<FTestTask_PrintValue>("Tree1StateRootTask1");
@@ -929,7 +929,7 @@ struct FStateTreeTest_LinkedAssetTransitionSameTick : FStateTreeTestBase
 			FStateTreeCompilerLog Log;
 			FStateTreeCompiler Compiler(Log);
 			const bool bResult = Compiler.Compile(StateTree1);
-			AITEST_TRUE("StateTree should get compiled", bResult);
+			AITEST_TRUE(TEXT("StateTree should get compiled"), bResult);
 		}
 
 		// Run test
@@ -938,7 +938,7 @@ struct FStateTreeTest_LinkedAssetTransitionSameTick : FStateTreeTestBase
 			FTestStateTreeExecutionContext Exec(StateTree1, StateTree1, InstanceData);
 			{
 				const bool bInitSucceeded = Exec.IsValid();
-				AITEST_TRUE("StateTree should init", bInitSucceeded);
+				AITEST_TRUE(TEXT("StateTree should init"), bInitSucceeded);
 			}
 
 			{
@@ -1101,7 +1101,7 @@ struct FStateTreeTest_Linked_GlobalParameter : FStateTreeTestBase
 			FStateTreeCompilerLog Log;
 			FStateTreeCompiler Compiler2(Log);
 			const bool bResult2 = Compiler2.Compile(StateTree2);
-			AITEST_TRUE("StateTree2 should get compiled", bResult2);
+			AITEST_TRUE(TEXT("StateTree2 should get compiled"), bResult2);
 		}
 
 		// Tree 1
@@ -1210,14 +1210,14 @@ struct FStateTreeTest_Linked_GlobalParameter : FStateTreeTestBase
 			FStateTreeCompilerLog Log;
 			FStateTreeCompiler Compiler1(Log);
 			const bool bResult1 = Compiler1.Compile(StateTree1);
-			AITEST_TRUE("StateTree1 should get compiled", bResult1);
+			AITEST_TRUE(TEXT("StateTree1 should get compiled"), bResult1);
 		}
 
 		FStateTreeInstanceData InstanceData;
 		{
 			FTestStateTreeExecutionContext Exec(StateTree1, StateTree1, InstanceData);
 			const bool bInitSucceeded = Exec.IsValid();
-			AITEST_TRUE("StateTree should init", bInitSucceeded);
+			AITEST_TRUE(TEXT("StateTree should init"), bInitSucceeded);
 		}
 		{
 			FStateTreeReference StateTreeRef;
@@ -1229,22 +1229,22 @@ struct FStateTreeTest_Linked_GlobalParameter : FStateTreeTestBase
 				{
 					.GlobalParameters = &StateTreeRef.GetParameters(),
 				});
-			AITEST_EQUAL("Start should complete with Running", Status, EStateTreeRunStatus::Running);
+			AITEST_EQUAL(TEXT("Start should complete with Running"), Status, EStateTreeRunStatus::Running);
 			AITEST_TRUE(TEXT("In correct states"), Exec.ExpectInActiveStates("Tree1StateRoot", "Tree1State1", "Tree2StateRoot", "Tree2State1"));
-			AITEST_TRUE("Start should enter Tree1GlobalTask1", Exec.Expect("Tree1GlobalTask1", TEXT("EnterState5")));
-			AITEST_TRUE("Start should enter Tree1GlobalTask2", Exec.Expect("Tree1GlobalTask2", TEXT("EnterState5")));
-			AITEST_TRUE("Start should enter Tree1StateRootTask1", Exec.Expect("Tree1StateRootTask1", TEXT("EnterState5")));
-			AITEST_TRUE("Start should enter Tree1StateRootTask2", Exec.Expect("Tree1StateRootTask2", TEXT("EnterState5")));
-			AITEST_TRUE("Start should enter Tree1StateRootTask3", Exec.Expect("Tree1StateRootTask3", TEXT("EnterState5")));
-			AITEST_TRUE("Start should enter Tree2GlobalTask1", Exec.Expect("Tree2GlobalTask1", TEXT("EnterState5")));
-			AITEST_TRUE("Start should enter Tree2GlobalTask2", Exec.Expect("Tree2GlobalTask2", TEXT("EnterState5")));
-			AITEST_TRUE("Start should enter Tree2StateRootTask1", Exec.Expect("Tree2StateRootTask1", TEXT("EnterState5")));
-			AITEST_TRUE("Start should enter Tree2StateRootTask2", Exec.Expect("Tree2StateRootTask2", TEXT("EnterState5")));
-			AITEST_TRUE("Start should enter Tree2StateRootTask3", Exec.Expect("Tree2StateRootTask3", TEXT("EnterState5")));
-			AITEST_TRUE("Start should enter Tree2State1Task1", Exec.Expect("Tree2State1Task1", TEXT("EnterState5")));
-			AITEST_TRUE("Start should enter Tree2State1Task2", Exec.Expect("Tree2State1Task2", TEXT("EnterState5")));
-			AITEST_TRUE("Start should enter Tree2State1Task3", Exec.Expect("Tree2State1Task3", TEXT("EnterState5")));
-			AITEST_TRUE("Start should enter Tree2State1Task4", Exec.Expect("Tree2State1Task4", TEXT("EnterState5")));
+			AITEST_TRUE(TEXT("Start should enter Tree1GlobalTask1"), Exec.Expect("Tree1GlobalTask1", TEXT("EnterState5")));
+			AITEST_TRUE(TEXT("Start should enter Tree1GlobalTask2"), Exec.Expect("Tree1GlobalTask2", TEXT("EnterState5")));
+			AITEST_TRUE(TEXT("Start should enter Tree1StateRootTask1"), Exec.Expect("Tree1StateRootTask1", TEXT("EnterState5")));
+			AITEST_TRUE(TEXT("Start should enter Tree1StateRootTask2"), Exec.Expect("Tree1StateRootTask2", TEXT("EnterState5")));
+			AITEST_TRUE(TEXT("Start should enter Tree1StateRootTask3"), Exec.Expect("Tree1StateRootTask3", TEXT("EnterState5")));
+			AITEST_TRUE(TEXT("Start should enter Tree2GlobalTask1"), Exec.Expect("Tree2GlobalTask1", TEXT("EnterState5")));
+			AITEST_TRUE(TEXT("Start should enter Tree2GlobalTask2"), Exec.Expect("Tree2GlobalTask2", TEXT("EnterState5")));
+			AITEST_TRUE(TEXT("Start should enter Tree2StateRootTask1"), Exec.Expect("Tree2StateRootTask1", TEXT("EnterState5")));
+			AITEST_TRUE(TEXT("Start should enter Tree2StateRootTask2"), Exec.Expect("Tree2StateRootTask2", TEXT("EnterState5")));
+			AITEST_TRUE(TEXT("Start should enter Tree2StateRootTask3"), Exec.Expect("Tree2StateRootTask3", TEXT("EnterState5")));
+			AITEST_TRUE(TEXT("Start should enter Tree2State1Task1"), Exec.Expect("Tree2State1Task1", TEXT("EnterState5")));
+			AITEST_TRUE(TEXT("Start should enter Tree2State1Task2"), Exec.Expect("Tree2State1Task2", TEXT("EnterState5")));
+			AITEST_TRUE(TEXT("Start should enter Tree2State1Task3"), Exec.Expect("Tree2State1Task3", TEXT("EnterState5")));
+			AITEST_TRUE(TEXT("Start should enter Tree2State1Task4"), Exec.Expect("Tree2State1Task4", TEXT("EnterState5")));
 			Exec.LogClear();
 		}
 		{
@@ -1257,33 +1257,33 @@ struct FStateTreeTest_Linked_GlobalParameter : FStateTreeTestBase
 			InstanceData.GetMutableStorage().SetGlobalParameters(StateTreeRef.GetParameters());
 
 			const EStateTreeRunStatus Status = Exec.Tick(1.0f);
-			AITEST_EQUAL("Tick should complete with Running", Status, EStateTreeRunStatus::Running);
+			AITEST_EQUAL(TEXT("Tick should complete with Running"), Status, EStateTreeRunStatus::Running);
 
 			AITEST_TRUE(TEXT("In correct states"), Exec.ExpectInActiveStates("Tree1StateRoot", "Tree1State2", "Tree1StateSub1", "Tree1State3"));
-			AITEST_TRUE("Tick should tick Tree1GlobalTask1", Exec.Expect("Tree1GlobalTask1", TEXT("Tick6")));
-			AITEST_TRUE("Tick should tick Tree1GlobalTask2", Exec.Expect("Tree1GlobalTask2", TEXT("Tick6")));
-			AITEST_TRUE("Tick should tick Tree1StateRootTask1", Exec.Expect("Tree1StateRootTask1", TEXT("Tick6")));
+			AITEST_TRUE(TEXT("Tick should tick Tree1GlobalTask1"), Exec.Expect("Tree1GlobalTask1", TEXT("Tick6")));
+			AITEST_TRUE(TEXT("Tick should tick Tree1GlobalTask2"), Exec.Expect("Tree1GlobalTask2", TEXT("Tick6")));
+			AITEST_TRUE(TEXT("Tick should tick Tree1StateRootTask1"), Exec.Expect("Tree1StateRootTask1", TEXT("Tick6")));
 			//@todo Binding on the state are not updated on tick. Is this a bug? See Tree1StateRootParametersInt
-			//AITEST_TRUE("Tick should tick Tree1StateRootTask2", Exec.Expect("Tree1StateRootTask2", TEXT("Tick6")));
-			AITEST_TRUE("Tick should tick Tree1StateRootTask3", Exec.Expect("Tree1StateRootTask3", TEXT("Tick6")));
-			AITEST_TRUE("Tick should tick Tree2GlobalTask1", Exec.Expect("Tree2GlobalTask1", TEXT("Tick6")));
-			AITEST_TRUE("Tick should tick Tree2GlobalTask2", Exec.Expect("Tree2GlobalTask2", TEXT("Tick6")));
-			AITEST_TRUE("Tick should tick Tree2StateRootTask1", Exec.Expect("Tree2StateRootTask1", TEXT("Tick6")));
+			//AITEST_TRUE(TEXT("Tick should tick Tree1StateRootTask2), Exec.Expect("Tree1StateRootTask2", TEXT("Tick6")));
+			AITEST_TRUE(TEXT("Tick should tick Tree1StateRootTask3"), Exec.Expect("Tree1StateRootTask3", TEXT("Tick6")));
+			AITEST_TRUE(TEXT("Tick should tick Tree2GlobalTask1"), Exec.Expect("Tree2GlobalTask1", TEXT("Tick6")));
+			AITEST_TRUE(TEXT("Tick should tick Tree2GlobalTask2"), Exec.Expect("Tree2GlobalTask2", TEXT("Tick6")));
+			AITEST_TRUE(TEXT("Tick should tick Tree2StateRootTask1"), Exec.Expect("Tree2StateRootTask1", TEXT("Tick6")));
 			//@todo Binding on the state are not updated on tick. Is this a bug? See Tree2StateRootParametersInt
-			//AITEST_TRUE("Tick should tick Tree2StateRootTask2", Exec.Expect("Tree2StateRootTask2", TEXT("Tick6")));
-			AITEST_TRUE("Tick should tick Tree2StateRootTask3", Exec.Expect("Tree2StateRootTask3", TEXT("Tick6")));
-			AITEST_TRUE("Tick should tick Tree2State1Task1", Exec.Expect("Tree2State1Task1", TEXT("Tick6")));
+			//AITEST_TRUE(TEXT("Tick should tick Tree2StateRootTask2"), Exec.Expect("Tree2StateRootTask2", TEXT("Tick6")));
+			AITEST_TRUE(TEXT("Tick should tick Tree2StateRootTask3"), Exec.Expect("Tree2StateRootTask3", TEXT("Tick6")));
+			AITEST_TRUE(TEXT("Tick should tick Tree2State1Task1"), Exec.Expect("Tree2State1Task1", TEXT("Tick6")));
 			//@todo Binding on the state are not updated on tick. Is this a bug? See Tree2StateRootParametersInt
-			//AITEST_TRUE("Tick should tick Tree2State1Task2", Exec.Expect("Tree2State1Task2", TEXT("Tick6")));
-			AITEST_TRUE("Tick should tick Tree2State1Task3", Exec.Expect("Tree2State1Task3", TEXT("Tick6")));
+			//AITEST_TRUE(TEXT("Tick should tick Tree2State1Task2"), Exec.Expect("Tree2State1Task2", TEXT("Tick6")));
+			AITEST_TRUE(TEXT("Tick should tick Tree2State1Task3"), Exec.Expect("Tree2State1Task3", TEXT("Tick6")));
 			//@todo Binding on the state are not updated on tick. Is this a bug? See Tree2State1ParametersInt
-			//AITEST_TRUE("Tick should tick Tree2State1Task4", Exec.Expect("Tree2State1Task4", TEXT("Tick6")));
+			//AITEST_TRUE(TEXT("Tick should tick Tree2State1Task4"), Exec.Expect("Tree2State1Task4", TEXT("Tick6")));
 
-			AITEST_TRUE("Tick should enter Tree1StateSub1Task1", Exec.Expect("Tree1StateSub1Task1", TEXT("EnterState6")));
-			AITEST_TRUE("Tick should enter Tree1State3Task1", Exec.Expect("Tree1State3Task1", TEXT("EnterState6")));
-			AITEST_TRUE("Tick should enter Tree1State3Task2", Exec.Expect("Tree1State3Task2", TEXT("EnterState6")));
-			AITEST_TRUE("Tick should enter Tree1State3Task3", Exec.Expect("Tree1State3Task3", TEXT("EnterState6")));
-			AITEST_TRUE("Tick should enter Tree1State3Task4", Exec.Expect("Tree1State3Task4", TEXT("EnterState6")));
+			AITEST_TRUE(TEXT("Tick should enter Tree1StateSub1Task1"), Exec.Expect("Tree1StateSub1Task1", TEXT("EnterState6")));
+			AITEST_TRUE(TEXT("Tick should enter Tree1State3Task1"), Exec.Expect("Tree1State3Task1", TEXT("EnterState6")));
+			AITEST_TRUE(TEXT("Tick should enter Tree1State3Task2"), Exec.Expect("Tree1State3Task2", TEXT("EnterState6")));
+			AITEST_TRUE(TEXT("Tick should enter Tree1State3Task3"), Exec.Expect("Tree1State3Task3", TEXT("EnterState6")));
+			AITEST_TRUE(TEXT("Tick should enter Tree1State3Task4"), Exec.Expect("Tree1State3Task4", TEXT("EnterState6")));
 		}
 		{
 			FTestStateTreeExecutionContext Exec(StateTree1, StateTree1, InstanceData);
diff --git a/Engine/Plugins/Runtime/StateTree/Source/StateTreeTestSuite/Private/StateTreePropertyPathTest.cpp b/Engine/Plugins/Runtime/StateTree/Source/StateTreeTestSuite/Private/StateTreePropertyPathTest.cpp
--- a/Engine/Plugins/Runtime/StateTree/Source/StateTreeTestSuite/Private/StateTreePropertyPathTest.cpp
+++ b/Engine/Plugins/Runtime/StateTree/Source/StateTreeTestSuite/Private/StateTreePropertyPathTest.cpp
@@ -21,19 +21,19 @@ struct FStateTreeTest_PropertyPathOffset : FStateTreeTestBase
 		FPropertyBindingPath Path;
 		const bool bParseResult = Path.FromString(TEXT("StructB.B"));
 
-		AITEST_TRUE("Parsing path should succeeed", bParseResult);
-		AITEST_EQUAL("Should have 2 path segments", Path.NumSegments(), 2);
+		AITEST_TRUE(TEXT("Parsing path should succeeed"), bParseResult);
+		AITEST_EQUAL(TEXT("Should have 2 path segments"), Path.NumSegments(), 2);
 
 		FString ResolveErrors;
 		TArray<FPropertyBindingPathIndirection> Indirections;
 		const bool bResolveResult = Path.ResolveIndirections(FStateTreeTest_PropertyStruct::StaticStruct(), Indirections, &ResolveErrors);
 
-		AITEST_TRUE("Resolve path should succeeed", bResolveResult);
-		AITEST_EQUAL("Should have no resolve errors", ResolveErrors.Len(), 0);
+		AITEST_TRUE(TEXT("Resolve path should succeeed"), bResolveResult);
+		AITEST_EQUAL(TEXT("Should have no resolve errors"), ResolveErrors.Len(), 0);
 		
-		AITEST_EQUAL("Should have 2 indirections", Indirections.Num(), 2);
-		AITEST_EQUAL("Indirection 0 should be Offset type", Indirections[0].GetAccessType(), EPropertyBindingPropertyAccessType::Offset);
-		AITEST_EQUAL("Indirection 1 should be Offset type", Indirections[1].GetAccessType(), EPropertyBindingPropertyAccessType::Offset);
+		AITEST_EQUAL(TEXT("Should have 2 indirections"), Indirections.Num(), 2);
+		AITEST_EQUAL(TEXT("Indirection 0 should be Offset type"), Indirections[0].GetAccessType(), EPropertyBindingPropertyAccessType::Offset);
+		AITEST_EQUAL(TEXT("Indirection 1 should be Offset type"), Indirections[1].GetAccessType(), EPropertyBindingPropertyAccessType::Offset);
 
 		return true;
 	}
@@ -47,31 +47,31 @@ struct FStateTreeTest_PropertyPathParseFail : FStateTreeTestBase
 		{
 			FPropertyBindingPath Path;
 			const bool bParseResult = Path.FromString(TEXT("")); // empty is valid.
-			AITEST_TRUE("Parsing path should succeed", bParseResult);
+			AITEST_TRUE(TEXT("Parsing path should succeed"), bParseResult);
 		}
 
 		{
 			FPropertyBindingPath Path;
 			const bool bParseResult = Path.FromString(TEXT("StructB.[0]B"));
-			AITEST_FALSE("Parsing path should fail", bParseResult);
+			AITEST_FALSE(TEXT("Parsing path should fail"), bParseResult);
 		}
 
 		{
 			FPropertyBindingPath Path;
 			const bool bParseResult = Path.FromString(TEXT("StructB..NoThere"));
-			AITEST_FALSE("Parsing path should fail", bParseResult);
+			AITEST_FALSE(TEXT("Parsing path should fail"), bParseResult);
 		}
 
 		{
 			FPropertyBindingPath Path;
 			const bool bParseResult = Path.FromString(TEXT("."));
-			AITEST_FALSE("Parsing path should fail", bParseResult);
+			AITEST_FALSE(TEXT("Parsing path should fail"), bParseResult);
 		}
 
 		{
 			FPropertyBindingPath Path;
 			const bool bParseResult = Path.FromString(TEXT("StructB..B"));
-			AITEST_FALSE("Parsing path should fail", bParseResult);
+			AITEST_FALSE(TEXT("Parsing path should fail"), bParseResult);
 		}
 
 		return true;
@@ -86,17 +86,17 @@ struct FStateTreeTest_PropertyPathOffsetFail : FStateTreeTestBase
 		FPropertyBindingPath Path;
 		const bool bParseResult = Path.FromString(TEXT("StructB.Q"));
 
-		AITEST_TRUE("Parsing path should succeeed", bParseResult);
-		AITEST_EQUAL("Should have 2 path segments", Path.NumSegments(), 2);
+		AITEST_TRUE(TEXT("Parsing path should succeeed"), bParseResult);
+		AITEST_EQUAL(TEXT("Should have 2 path segments"), Path.NumSegments(), 2);
 
 		FString ResolveErrors;
 		TArray<FPropertyBindingPathIndirection> Indirections;
 		const bool bResolveResult = Path.ResolveIndirections(FStateTreeTest_PropertyStruct::StaticStruct(), Indirections, &ResolveErrors);
 
-		AITEST_FALSE("Resolve path should not succeeed", bResolveResult);
-		AITEST_NOT_EQUAL("Should have errors", ResolveErrors.Len(), 0);
+		AITEST_FALSE(TEXT("Resolve path should not succeeed"), bResolveResult);
+		AITEST_NOT_EQUAL(TEXT("Should have errors"), ResolveErrors.Len(), 0);
 		
-		AITEST_EQUAL("Should have 0 indirections", Indirections.Num(), 0);
+		AITEST_EQUAL(TEXT("Should have 0 indirections"), Indirections.Num(), 0);
 
 		return true;
 	}
@@ -110,17 +110,17 @@ struct FStateTreeTest_PropertyPathObject : FStateTreeTestBase
 		FPropertyBindingPath Path;
 		const bool bParseResult = Path.FromString(TEXT("InstancedObject.A"));
 
-		AITEST_TRUE("Parsing path should succeeed", bParseResult);
-		AITEST_EQUAL("Should have 2 path segments", Path.NumSegments(), 2);
+		AITEST_TRUE(TEXT("Parsing path should succeeed"), bParseResult);
+		AITEST_EQUAL(TEXT("Should have 2 path segments"), Path.NumSegments(), 2);
 
 		UStateTreeTest_PropertyObject* Object = NewObject<UStateTreeTest_PropertyObject>();
 		Object->InstancedObject = NewObject<UStateTreeTest_PropertyObjectInstanced>();
 		
 		const bool bUpdateResult = Path.UpdateSegmentsFromValue(FStateTreeDataView(Object));
 
-		AITEST_TRUE("Update instance types should succeeed", bUpdateResult);
-		AITEST_TRUE("Path segment 0 instance type should be UStateTreeTest_PropertyObjectInstanced", Path.GetSegment(0).GetInstanceStruct() == UStateTreeTest_PropertyObjectInstanced::StaticClass());
-		AITEST_TRUE("Path segment 1 instance type should be nullptr", Path.GetSegment(1).GetInstanceStruct() == nullptr);
+		AITEST_TRUE(TEXT("Update instance types should succeeed"), bUpdateResult);
+		AITEST_TRUE(TEXT("Path segment 0 instance type should be UStateTreeTest_PropertyObjectInstanced"), Path.GetSegment(0).GetInstanceStruct() == UStateTreeTest_PropertyObjectInstanced::StaticClass());
+		AITEST_TRUE(TEXT("Path segment 1 instance type should be nullptr"), Path.GetSegment(1).GetInstanceStruct() == nullptr);
 
 		return true;
 	}
@@ -134,8 +134,8 @@ struct FStateTreeTest_PropertyPathWrongObject : FStateTreeTestBase
 		FPropertyBindingPath Path;
 		const bool bParseResult = Path.FromString(TEXT("InstancedObject.B"));
 
-		AITEST_TRUE("Parsing path should succeeed", bParseResult);
-		AITEST_EQUAL("Should have 2 path segments", Path.NumSegments(), 2);
+		AITEST_TRUE(TEXT("Parsing path should succeeed"), bParseResult);
+		AITEST_EQUAL(TEXT("Should have 2 path segments"), Path.NumSegments(), 2);
 
 		UStateTreeTest_PropertyObject* Object = NewObject<UStateTreeTest_PropertyObject>();
 
@@ -145,12 +145,12 @@ struct FStateTreeTest_PropertyPathWrongObject : FStateTreeTestBase
 			TArray<FPropertyBindingPathIndirection> Indirections;
 			const bool bResolveResult = Path.ResolveIndirectionsWithValue(FStateTreeDataView(Object), Indirections, &ResolveErrors);
 
-			AITEST_TRUE("Resolve path should succeeed", bResolveResult);
-			AITEST_EQUAL("Should have 2 indirections", Indirections.Num(), 2);
-			AITEST_TRUE("Object ", Indirections[0].GetAccessType() == EPropertyBindingPropertyAccessType::ObjectInstance);
-			AITEST_TRUE("Object ", Indirections[0].GetContainerStruct() == Object->GetClass());
-			AITEST_TRUE("Object ", Indirections[0].GetInstanceStruct() == UStateTreeTest_PropertyObjectInstancedWithB::StaticClass());
-			AITEST_EQUAL("Should not have error", ResolveErrors.Len(), 0);
+			AITEST_TRUE(TEXT("Resolve path should succeeed"), bResolveResult);
+			AITEST_EQUAL(TEXT("Should have 2 indirections"), Indirections.Num(), 2);
+			AITEST_TRUE(TEXT("Object "), Indirections[0].GetAccessType() == EPropertyBindingPropertyAccessType::ObjectInstance);
+			AITEST_TRUE(TEXT("Object "), Indirections[0].GetContainerStruct() == Object->GetClass());
+			AITEST_TRUE(TEXT("Object "), Indirections[0].GetInstanceStruct() == UStateTreeTest_PropertyObjectInstancedWithB::StaticClass());
+			AITEST_EQUAL(TEXT("Should not have error"), ResolveErrors.Len(), 0);
 		}
 
 		Object->InstancedObject = NewObject<UStateTreeTest_PropertyObjectInstanced>();
@@ -159,9 +159,9 @@ struct FStateTreeTest_PropertyPathWrongObject : FStateTreeTestBase
 			TArray<FPropertyBindingPathIndirection> Indirections;
 			const bool bResolveResult = Path.ResolveIndirectionsWithValue(FStateTreeDataView(Object), Indirections, &ResolveErrors);
 
-			AITEST_FALSE("Resolve path should fail", bResolveResult);
-			AITEST_EQUAL("Should have 0 indirections", Indirections.Num(), 0);
-			AITEST_NOT_EQUAL("Should have error", ResolveErrors.Len(), 0);
+			AITEST_FALSE(TEXT("Resolve path should fail"), bResolveResult);
+			AITEST_EQUAL(TEXT("Should have 0 indirections"), Indirections.Num(), 0);
+			AITEST_NOT_EQUAL(TEXT("Should have error"), ResolveErrors.Len(), 0);
 		}
 		
 		return true;
@@ -176,8 +176,8 @@ struct FStateTreeTest_PropertyPathArray : FStateTreeTestBase
 		FPropertyBindingPath Path;
 		const bool bParseResult = Path.FromString(TEXT("ArrayOfInts[1]"));
 
-		AITEST_TRUE("Parsing path should succeeed", bParseResult);
-		AITEST_EQUAL("Should have 1 path segments", Path.NumSegments(), 1);
+		AITEST_TRUE(TEXT("Parsing path should succeeed"), bParseResult);
+		AITEST_EQUAL(TEXT("Should have 1 path segments"), Path.NumSegments(), 1);
 
 		UStateTreeTest_PropertyObject* Object = NewObject<UStateTreeTest_PropertyObject>();
 		Object->ArrayOfInts.Add(42);
@@ -187,14 +187,14 @@ struct FStateTreeTest_PropertyPathArray : FStateTreeTestBase
 		TArray<FPropertyBindingPathIndirection> Indirections;
 		const bool bResolveResult = Path.ResolveIndirectionsWithValue(FStateTreeDataView(Object), Indirections, &ResolveErrors);
 
-		AITEST_TRUE("Resolve path should succeeed", bResolveResult);
-		AITEST_EQUAL("Should have no resolve errors", ResolveErrors.Len(), 0);
-		AITEST_EQUAL("Should have 2 indirections", Indirections.Num(), 2);
-		AITEST_EQUAL("Indirection 0 should be IndexArray type", Indirections[0].GetAccessType(), EPropertyBindingPropertyAccessType::IndexArray);
-		AITEST_EQUAL("Indirection 1 should be Offset type", Indirections[1].GetAccessType(), EPropertyBindingPropertyAccessType::Offset);
+		AITEST_TRUE(TEXT("Resolve path should succeeed"), bResolveResult);
+		AITEST_EQUAL(TEXT("Should have no resolve errors"), ResolveErrors.Len(), 0);
+		AITEST_EQUAL(TEXT("Should have 2 indirections"), Indirections.Num(), 2);
+		AITEST_EQUAL(TEXT("Indirection 0 should be IndexArray type"), Indirections[0].GetAccessType(), EPropertyBindingPropertyAccessType::IndexArray);
+		AITEST_EQUAL(TEXT("Indirection 1 should be Offset type"), Indirections[1].GetAccessType(), EPropertyBindingPropertyAccessType::Offset);
 
 		const int32 Value = *reinterpret_cast<const int32*>(Indirections[1].GetPropertyAddress());
-		AITEST_EQUAL("Value should be 123", Value, 123);
+		AITEST_EQUAL(TEXT("Value should be 123"), Value, 123);
 		
 		return true;
 	}
@@ -208,8 +208,8 @@ struct FStateTreeTest_PropertyPathArrayInvalidIndex : FStateTreeTestBase
 		FPropertyBindingPath Path;
 		const bool bParseResult = Path.FromString(TEXT("ArrayOfInts[123]"));
 
-		AITEST_TRUE("Parsing path should succeeed", bParseResult);
-		AITEST_EQUAL("Should have 1 path segments", Path.NumSegments(), 1);
+		AITEST_TRUE(TEXT("Parsing path should succeeed"), bParseResult);
+		AITEST_EQUAL(TEXT("Should have 1 path segments"), Path.NumSegments(), 1);
 
 		UStateTreeTest_PropertyObject* Object = NewObject<UStateTreeTest_PropertyObject>();
 		Object->ArrayOfInts.Add(42);
@@ -219,7 +219,7 @@ struct FStateTreeTest_PropertyPathArrayInvalidIndex : FStateTreeTestBase
 		TArray<FPropertyBindingPathIndirection> Indirections;
 		const bool bResolveResult = Path.ResolveIndirectionsWithValue(FStateTreeDataView(Object), Indirections, &ResolveErrors);
 
-		AITEST_FALSE("Resolve path should fail", bResolveResult);
+		AITEST_FALSE(TEXT("Resolve path should fail"), bResolveResult);
 		
 		return true;
 	}
@@ -246,15 +246,15 @@ struct FStateTreeTest_PropertyPathArrayOfStructs : FStateTreeTestBase
 			TArray<FPropertyBindingPathIndirection> Indirections;
 			const bool bResolveResult = Path1.ResolveIndirectionsWithValue(FStateTreeDataView(Object), Indirections, &ResolveErrors);
 
-			AITEST_TRUE("Resolve path1 should succeeed", bResolveResult);
-			AITEST_EQUAL("Should have no resolve errors", ResolveErrors.Len(), 0);
-			AITEST_EQUAL("Should have 3 indirections", Indirections.Num(), 3);
-			AITEST_EQUAL("Indirection 0 should be ArrayIndex type", Indirections[0].GetAccessType(), EPropertyBindingPropertyAccessType::IndexArray);
-			AITEST_EQUAL("Indirection 1 should be Offset type", Indirections[1].GetAccessType(), EPropertyBindingPropertyAccessType::Offset);
-			AITEST_EQUAL("Indirection 2 should be Offset type", Indirections[2].GetAccessType(), EPropertyBindingPropertyAccessType::Offset);
+			AITEST_TRUE(TEXT("Resolve path1 should succeeed"), bResolveResult);
+			AITEST_EQUAL(TEXT("Should have no resolve errors"), ResolveErrors.Len(), 0);
+			AITEST_EQUAL(TEXT("Should have 3 indirections"), Indirections.Num(), 3);
+			AITEST_EQUAL(TEXT("Indirection 0 should be ArrayIndex type"), Indirections[0].GetAccessType(), EPropertyBindingPropertyAccessType::IndexArray);
+			AITEST_EQUAL(TEXT("Indirection 1 should be Offset type"), Indirections[1].GetAccessType(), EPropertyBindingPropertyAccessType::Offset);
+			AITEST_EQUAL(TEXT("Indirection 2 should be Offset type"), Indirections[2].GetAccessType(), EPropertyBindingPropertyAccessType::Offset);
 
 			const int32 Value = *reinterpret_cast<const int32*>(Indirections[2].GetPropertyAddress());
-			AITEST_EQUAL("Value should be 3", Value, 3);
+			AITEST_EQUAL(TEXT("Value should be 3"), Value, 3);
 		}
 
 		{
@@ -262,16 +262,16 @@ struct FStateTreeTest_PropertyPathArrayOfStructs : FStateTreeTestBase
 			TArray<FPropertyBindingPathIndirection> Indirections;
 			const bool bResolveResult = Path2.ResolveIndirectionsWithValue(FStateTreeDataView(Object), Indirections, &ResolveErrors);
 
-			AITEST_TRUE("Resolve path2 should succeeed", bResolveResult);
-			AITEST_EQUAL("Should have no resolve errors", ResolveErrors.Len(), 0);
-			AITEST_EQUAL("Should have 4 indirections", Indirections.Num(), 4);
-			AITEST_EQUAL("Indirection 0 should be ArrayIndex type", Indirections[0].GetAccessType(), EPropertyBindingPropertyAccessType::IndexArray);
-			AITEST_EQUAL("Indirection 1 should be Offset type", Indirections[1].GetAccessType(), EPropertyBindingPropertyAccessType::Offset);
-			AITEST_EQUAL("Indirection 2 should be Offset type", Indirections[2].GetAccessType(), EPropertyBindingPropertyAccessType::Offset);
-			AITEST_EQUAL("Indirection 3 should be Offset type", Indirections[3].GetAccessType(), EPropertyBindingPropertyAccessType::Offset);
+			AITEST_TRUE(TEXT("Resolve path2 should succeeed"), bResolveResult);
+			AITEST_EQUAL(TEXT("Should have no resolve errors"), ResolveErrors.Len(), 0);
+			AITEST_EQUAL(TEXT("Should have 4 indirections"), Indirections.Num(), 4);
+			AITEST_EQUAL(TEXT("Indirection 0 should be ArrayIndex type"), Indirections[0].GetAccessType(), EPropertyBindingPropertyAccessType::IndexArray);
+			AITEST_EQUAL(TEXT("Indirection 1 should be Offset type"), Indirections[1].GetAccessType(), EPropertyBindingPropertyAccessType::Offset);
+			AITEST_EQUAL(TEXT("Indirection 2 should be Offset type"), Indirections[2].GetAccessType(), EPropertyBindingPropertyAccessType::Offset);
+			AITEST_EQUAL(TEXT("Indirection 3 should be Offset type"), Indirections[3].GetAccessType(), EPropertyBindingPropertyAccessType::Offset);
 
 			const int32 Value = *reinterpret_cast<const int32*>(Indirections[3].GetPropertyAddress());
-			AITEST_EQUAL("Value should be 42", Value, 42);
+			AITEST_EQUAL(TEXT("Value should be 42"), Value, 42);
 		}
 		return true;
 	}
@@ -292,22 +292,22 @@ struct FStateTreeTest_PropertyPathArrayOfInstancedObjects : FStateTreeTestBase
 		Object->ArrayOfInstancedStructs.Emplace(FConstStructView::Make(Struct));
 
 		const bool bUpdateResult = Path.UpdateSegmentsFromValue(FStateTreeDataView(Object));
-		AITEST_TRUE("Update instance types should succeeed", bUpdateResult);
-		AITEST_EQUAL("Should have 2 path segments", Path.NumSegments(), 2);
-		AITEST_TRUE("Path segment 0 instance type should be FStateTreeTest_PropertyStruct", Path.GetSegment(0).GetInstanceStruct() == FStateTreeTest_PropertyStruct::StaticStruct());
-		AITEST_TRUE("Path segment 1 instance type should be nullptr", Path.GetSegment(1).GetInstanceStruct() == nullptr);
+		AITEST_TRUE(TEXT("Update instance types should succeeed"), bUpdateResult);
+		AITEST_EQUAL(TEXT("Should have 2 path segments"), Path.NumSegments(), 2);
+		AITEST_TRUE(TEXT("Path segment 0 instance type should be FStateTreeTest_PropertyStruct"), Path.GetSegment(0).GetInstanceStruct() == FStateTreeTest_PropertyStruct::StaticStruct());
+		AITEST_TRUE(TEXT("Path segment 1 instance type should be nullptr"), Path.GetSegment(1).GetInstanceStruct() == nullptr);
 
 		{
 			FString ResolveErrors;
 			TArray<FPropertyBindingPathIndirection> Indirections;
 			const bool bResolveResult = Path.ResolveIndirections(UStateTreeTest_PropertyObject::StaticClass(), Indirections, &ResolveErrors);
 
-			AITEST_TRUE("Resolve path should succeeed", bResolveResult);
-			AITEST_EQUAL("Should have no resolve errors", ResolveErrors.Len(), 0);
-			AITEST_EQUAL("Should have 3 indirections", Indirections.Num(), 3);
-			AITEST_EQUAL("Indirection 0 should be ArrayIndex type", Indirections[0].GetAccessType(), EPropertyBindingPropertyAccessType::IndexArray);
-			AITEST_EQUAL("Indirection 1 should be StructInstance type", Indirections[1].GetAccessType(), EPropertyBindingPropertyAccessType::StructInstance);
-			AITEST_EQUAL("Indirection 2 should be Offset type", Indirections[2].GetAccessType(), EPropertyBindingPropertyAccessType::Offset);
+			AITEST_TRUE(TEXT("Resolve path should succeeed"), bResolveResult);
+			AITEST_EQUAL(TEXT("Should have no resolve errors"), ResolveErrors.Len(), 0);
+			AITEST_EQUAL(TEXT("Should have 3 indirections"), Indirections.Num(), 3);
+			AITEST_EQUAL(TEXT("Indirection 0 should be ArrayIndex type"), Indirections[0].GetAccessType(), EPropertyBindingPropertyAccessType::IndexArray);
+			AITEST_EQUAL(TEXT("Indirection 1 should be StructInstance type"), Indirections[1].GetAccessType(), EPropertyBindingPropertyAccessType::StructInstance);
+			AITEST_EQUAL(TEXT("Indirection 2 should be Offset type"), Indirections[2].GetAccessType(), EPropertyBindingPropertyAccessType::Offset);
 		}
 
 		{
@@ -315,15 +315,15 @@ struct FStateTreeTest_PropertyPathArrayOfInstancedObjects : FStateTreeTestBase
 			TArray<FPropertyBindingPathIndirection> Indirections;
 			const bool bResolveResult = Path.ResolveIndirectionsWithValue(FStateTreeDataView(Object), Indirections, &ResolveErrors);
 
-			AITEST_TRUE("Resolve path should succeeed", bResolveResult);
-			AITEST_EQUAL("Should have no resolve errors", ResolveErrors.Len(), 0);
-			AITEST_EQUAL("Should have 3 indirections", Indirections.Num(), 3);
-			AITEST_EQUAL("Indirection 0 should be ArrayIndex type", Indirections[0].GetAccessType(), EPropertyBindingPropertyAccessType::IndexArray);
-			AITEST_EQUAL("Indirection 1 should be StructInstance type", Indirections[1].GetAccessType(), EPropertyBindingPropertyAccessType::StructInstance);
-			AITEST_EQUAL("Indirection 2 should be Offset type", Indirections[2].GetAccessType(), EPropertyBindingPropertyAccessType::Offset);
+			AITEST_TRUE(TEXT("Resolve path should succeeed"), bResolveResult);
+			AITEST_EQUAL(TEXT("Should have no resolve errors"), ResolveErrors.Len(), 0);
+			AITEST_EQUAL(TEXT("Should have 3 indirections"), Indirections.Num(), 3);
+			AITEST_EQUAL(TEXT("Indirection 0 should be ArrayIndex type"), Indirections[0].GetAccessType(), EPropertyBindingPropertyAccessType::IndexArray);
+			AITEST_EQUAL(TEXT("Indirection 1 should be StructInstance type"), Indirections[1].GetAccessType(), EPropertyBindingPropertyAccessType::StructInstance);
+			AITEST_EQUAL(TEXT("Indirection 2 should be Offset type"), Indirections[2].GetAccessType(), EPropertyBindingPropertyAccessType::Offset);
 
 			const int32 Value = *reinterpret_cast<const int32*>(Indirections[2].GetPropertyAddress());
-			AITEST_EQUAL("Value should be 123", Value, 123);
+			AITEST_EQUAL(TEXT("Value should be 123"), Value, 123);
 		}
 		
 		return true;
diff --git a/Engine/Plugins/Runtime/StateTree/Source/StateTreeTestSuite/Private/StateTreeStatePathTest.cpp b/Engine/Plugins/Runtime/StateTree/Source/StateTreeTestSuite/Private/StateTreeStatePathTest.cpp
--- a/Engine/Plugins/Runtime/StateTree/Source/StateTreeTestSuite/Private/StateTreeStatePathTest.cpp
+++ b/Engine/Plugins/Runtime/StateTree/Source/StateTreeTestSuite/Private/StateTreeStatePathTest.cpp
@@ -225,19 +225,19 @@ struct FStateTreeTest_StatePath_LinkStates : FStateTreeTestBase
 			FStateTreeCompilerLog Log;
 			FStateTreeCompiler Compiler(Log);
 			const bool bResult = Compiler.Compile(StateTree3);
-			AITEST_TRUE("StateTree3 should get compiled", bResult);
+			AITEST_TRUE(TEXT("StateTree3 should get compiled"), bResult);
 		}
 		{
 			FStateTreeCompilerLog Log;
 			FStateTreeCompiler Compiler(Log);
 			const bool bResult = Compiler.Compile(StateTree2);
-			AITEST_TRUE("StateTree2 should get compiled", bResult);
+			AITEST_TRUE(TEXT("StateTree2 should get compiled"), bResult);
 		}
 		{
 			FStateTreeCompilerLog Log;
 			FStateTreeCompiler Compiler(Log);
 			const bool bResult = Compiler.Compile(StateTree1);
-			AITEST_TRUE("StateTree1 should get compiled", bResult);
+			AITEST_TRUE(TEXT("StateTree1 should get compiled"), bResult);
 		}
 
 		// Create context
@@ -245,11 +245,11 @@ struct FStateTreeTest_StatePath_LinkStates : FStateTreeTestBase
 		FTestStateTreeExecutionContext Exec(StateTree1, StateTree1, InstanceData);
 		{
 			const bool bInitSucceeded = Exec.IsValid();
-			AITEST_TRUE("StateTree should init", bInitSucceeded);
+			AITEST_TRUE(TEXT("StateTree should init"), bInitSucceeded);
 		}
 		{
 			UE::StateTree::FActiveStatePath ExecPath = InstanceData.GetExecutionState()->GetActiveStatePath();
-			AITEST_TRUE("ExecPath should be empty.", ExecPath.Num() == 0);
+			AITEST_TRUE(TEXT("ExecPath should be empty."), ExecPath.Num() == 0);
 		}
 
 		// Test variable and helper functions
@@ -288,39 +288,39 @@ struct FStateTreeTest_StatePath_LinkStates : FStateTreeTestBase
 			using namespace UE::StateTree;
 
 			const EStateTreeRunStatus Status = Exec.Start();
-			AITEST_EQUAL("Start should complete with Running", Status, EStateTreeRunStatus::Running);
-			AITEST_TRUE("Should be in the correct state", Exec.ExpectInActiveStates("Tree1Root", "Tree1RootState1", "Tree1State1StateSub1A", "Tree1Sub1", "Tree1Sub1StateLinkTree2A", "Tree2StateRoot", "Tree2RootStateSub1A", "Tree2Sub1", "Tree2Sub1State1"));
+			AITEST_EQUAL(TEXT("Start should complete with Running"), Status, EStateTreeRunStatus::Running);
+			AITEST_TRUE(TEXT("Should be in the correct state"), Exec.ExpectInActiveStates("Tree1Root", "Tree1RootState1", "Tree1State1StateSub1A", "Tree1Sub1", "Tree1Sub1StateLinkTree2A", "Tree2StateRoot", "Tree2RootStateSub1A", "Tree2Sub1", "Tree2Sub1State1"));
 
 			const FActiveStatePath ExecPath = InstanceData.GetExecutionState()->GetActiveStatePath();
-			AITEST_TRUE("Has the correct number of path elements", ExecPath.Num() == 9);
-			AITEST_TRUE("Has the correct number of active states", InstanceData.GetExecutionState()->ActiveFrames.Num() == 4);
+			AITEST_TRUE(TEXT("Has the correct number of path elements"), ExecPath.Num() == 9);
+			AITEST_TRUE(TEXT("Has the correct number of active states"), InstanceData.GetExecutionState()->ActiveFrames.Num() == 4);
 			{
 				const FActiveFrameID FirstFrameID = InstanceData.GetExecutionState()->ActiveFrames[0].FrameID;
-				AITEST_TRUE("Frame for Tree1Root is active", FirstFrameID == FActiveFrameID(++ActiveCounter));
-				AITEST_TRUE("State Tree1Root is active", ExecPath.Contains(FActiveStateID(++ActiveCounter)));
-				AITEST_TRUE("State Tree1Root is active", ExecPath.Contains(FActiveState(FirstFrameID, FActiveStateID(ActiveCounter), FStateTreeStateHandle(0))));
-				AITEST_TRUE("State Tree1RootState1 is active", ExecPath.Contains(FActiveStateID(++ActiveCounter)));
-				AITEST_TRUE("State Tree1RootState1 is active", ExecPath.Contains(FActiveState(FirstFrameID, FActiveStateID(ActiveCounter), FStateTreeStateHandle(1))));
-				AITEST_TRUE("State Tree1State1StateSub1A is active", ExecPath.Contains(FActiveStateID(++ActiveCounter)));
-				AITEST_TRUE("State Tree1State1StateSub1A is active", ExecPath.Contains(FActiveState(FirstFrameID, FActiveStateID(4), FStateTreeStateHandle(2))));
+				AITEST_TRUE(TEXT("Frame for Tree1Root is active"), FirstFrameID == FActiveFrameID(++ActiveCounter));
+				AITEST_TRUE(TEXT("State Tree1Root is active"), ExecPath.Contains(FActiveStateID(++ActiveCounter)));
+				AITEST_TRUE(TEXT("State Tree1Root is active"), ExecPath.Contains(FActiveState(FirstFrameID, FActiveStateID(ActiveCounter), FStateTreeStateHandle(0))));
+				AITEST_TRUE(TEXT("State Tree1RootState1 is active"), ExecPath.Contains(FActiveStateID(++ActiveCounter)));
+				AITEST_TRUE(TEXT("State Tree1RootState1 is active"), ExecPath.Contains(FActiveState(FirstFrameID, FActiveStateID(ActiveCounter), FStateTreeStateHandle(1))));
+				AITEST_TRUE(TEXT("State Tree1State1StateSub1A is active"), ExecPath.Contains(FActiveStateID(++ActiveCounter)));
+				AITEST_TRUE(TEXT("State Tree1State1StateSub1A is active"), ExecPath.Contains(FActiveState(FirstFrameID, FActiveStateID(4), FStateTreeStateHandle(2))));
 			}
 			{
-				AITEST_TRUE("Frame for Tree1Sub1 is active", InstanceData.GetExecutionState()->ActiveFrames[1].FrameID == FActiveFrameID(++ActiveCounter));
-				AITEST_TRUE("State Tree1Sub1 is active", ExecPath.Contains(FActiveStateID(++ActiveCounter)));
-				AITEST_TRUE("State Tree1Sub1StateLinkTree2A is active", ExecPath.Contains(FActiveStateID(++ActiveCounter)));
+				AITEST_TRUE(TEXT("Frame for Tree1Sub1 is active"), InstanceData.GetExecutionState()->ActiveFrames[1].FrameID == FActiveFrameID(++ActiveCounter));
+				AITEST_TRUE(TEXT("State Tree1Sub1 is active"), ExecPath.Contains(FActiveStateID(++ActiveCounter)));
+				AITEST_TRUE(TEXT("State Tree1Sub1StateLinkTree2A is active"), ExecPath.Contains(FActiveStateID(++ActiveCounter)));
 			}
 			{
-				AITEST_TRUE("Frame for Tree2StateRoot is active", InstanceData.GetExecutionState()->ActiveFrames[2].FrameID == FActiveFrameID(++ActiveCounter));
-				AITEST_TRUE("State Tree2StateRoot is active", ExecPath.Contains(FActiveStateID(++ActiveCounter)));
-				AITEST_TRUE("State Tree2RootStateSub1A is active", ExecPath.Contains(FActiveStateID(++ActiveCounter)));
+				AITEST_TRUE(TEXT("Frame for Tree2StateRoot is active"), InstanceData.GetExecutionState()->ActiveFrames[2].FrameID == FActiveFrameID(++ActiveCounter));
+				AITEST_TRUE(TEXT("State Tree2StateRoot is active"), ExecPath.Contains(FActiveStateID(++ActiveCounter)));
+				AITEST_TRUE(TEXT("State Tree2RootStateSub1A is active"), ExecPath.Contains(FActiveStateID(++ActiveCounter)));
 			}
 			{
-				AITEST_TRUE("Frame for Tree2Sub1 is active", InstanceData.GetExecutionState()->ActiveFrames[3].FrameID == FActiveFrameID(++ActiveCounter));
-				AITEST_TRUE("State Tree2Sub1 is active", ExecPath.Contains(FActiveStateID(++ActiveCounter)));
-				AITEST_TRUE("State Tree2Sub1State1 is active", ExecPath.Contains(FActiveStateID(++ActiveCounter)));
+				AITEST_TRUE(TEXT("Frame for Tree2Sub1 is active"), InstanceData.GetExecutionState()->ActiveFrames[3].FrameID == FActiveFrameID(++ActiveCounter));
+				AITEST_TRUE(TEXT("State Tree2Sub1 is active"), ExecPath.Contains(FActiveStateID(++ActiveCounter)));
+				AITEST_TRUE(TEXT("State Tree2Sub1State1 is active"), ExecPath.Contains(FActiveStateID(++ActiveCounter)));
 			}
 			{
-				AITEST_FALSE("No accidental increment", ExecPath.Contains(FActiveStateID(ActiveCounter+1)));
+				AITEST_FALSE(TEXT("No accidental increment"), ExecPath.Contains(FActiveStateID(ActiveCounter+1)));
 			}
 			Exec.LogClear();
 		}
@@ -329,114 +329,114 @@ struct FStateTreeTest_StatePath_LinkStates : FStateTreeTestBase
 
 			PreTestFrameId();
 			EStateTreeRunStatus Status = Exec.Tick(1.0f);
-			AITEST_EQUAL("Start should complete with Running", Status, EStateTreeRunStatus::Running);
-			AITEST_TRUE("Should be in the correct state", Exec.ExpectInActiveStates("Tree1Root", "Tree1RootState1", "Tree1State1StateSub1A", "Tree1Sub1", "Tree1Sub1StateLinkTree2A", "Tree2StateRoot", "Tree2RootStateSub1A", "Tree2Sub1", "Tree2Sub1State2"));
+			AITEST_EQUAL(TEXT("Start should complete with Running"), Status, EStateTreeRunStatus::Running);
+			AITEST_TRUE(TEXT("Should be in the correct state"), Exec.ExpectInActiveStates("Tree1Root", "Tree1RootState1", "Tree1State1StateSub1A", "Tree1Sub1", "Tree1Sub1StateLinkTree2A", "Tree2StateRoot", "Tree2RootStateSub1A", "Tree2Sub1", "Tree2Sub1State2"));
 
 			FActiveStatePath ExecPath = InstanceData.GetExecutionState()->GetActiveStatePath();
-			AITEST_TRUE("Has the correct number of path elements", ExecPath.Num() == 9);
-			AITEST_TRUE("Has the correct number of active states", InstanceData.GetExecutionState()->ActiveFrames.Num() == 4);
-			AITEST_TRUE("State Tree2Sub1State2 is active", ExecPath.Contains(FActiveStateID(++ActiveCounter)));
-			AITEST_FALSE("No accidental increment", ExecPath.Contains(FActiveStateID(ActiveCounter + 1)));
-			AITEST_TRUE("All frames are the same", TestFrameId(4));
+			AITEST_TRUE(TEXT("Has the correct number of path elements"), ExecPath.Num() == 9);
+			AITEST_TRUE(TEXT("Has the correct number of active states"), InstanceData.GetExecutionState()->ActiveFrames.Num() == 4);
+			AITEST_TRUE(TEXT("State Tree2Sub1State2 is active"), ExecPath.Contains(FActiveStateID(++ActiveCounter)));
+			AITEST_FALSE(TEXT("No accidental increment"), ExecPath.Contains(FActiveStateID(ActiveCounter + 1)));
+			AITEST_TRUE(TEXT("All frames are the same"), TestFrameId(4));
 			Exec.LogClear();
 		}
 		{
 			PreTestFrameId();
 			EStateTreeRunStatus Status = Exec.Tick(1.0f);
-			AITEST_EQUAL("Start should complete with Running", Status, EStateTreeRunStatus::Running);
-			AITEST_TRUE("Should be in the correct state", Exec.ExpectInActiveStates("Tree1Root", "Tree1RootState1", "Tree1State1StateSub1A", "Tree1Sub1", "Tree1Sub1StateLinkTree2A", "Tree2StateRoot", "Tree2RootStateSub1B", "Tree2Sub1", "Tree2Sub1State1"));
+			AITEST_EQUAL(TEXT("Start should complete with Running"), Status, EStateTreeRunStatus::Running);
+			AITEST_TRUE(TEXT("Should be in the correct state"), Exec.ExpectInActiveStates("Tree1Root", "Tree1RootState1", "Tree1State1StateSub1A", "Tree1Sub1", "Tree1Sub1StateLinkTree2A", "Tree2StateRoot", "Tree2RootStateSub1B", "Tree2Sub1", "Tree2Sub1State1"));
 
 			const UE::StateTree::FActiveStatePath ExecPath = InstanceData.GetExecutionState()->GetActiveStatePath();
-			AITEST_TRUE("Has the correct number of path elements", ExecPath.Num() == 9);
-			AITEST_TRUE("Has the correct number of active states", InstanceData.GetExecutionState()->ActiveFrames.Num() == 4);
+			AITEST_TRUE(TEXT("Has the correct number of path elements"), ExecPath.Num() == 9);
+			AITEST_TRUE(TEXT("Has the correct number of active states"), InstanceData.GetExecutionState()->ActiveFrames.Num() == 4);
 			//@todo Temporary reverted the frameID fixed until TransitionID
-			//AITEST_TRUE("The last frame changed", TestFrameId(3));
+			//AITEST_TRUE(TEXT("The last frame changed"), TestFrameId(3));
 			Exec.LogClear();
 		}
 		{
 			PreTestFrameId();
 			EStateTreeRunStatus Status = Exec.Tick(1.0f);
-			AITEST_EQUAL("Start should complete with Running", Status, EStateTreeRunStatus::Running);
-			AITEST_TRUE("Should be in the correct state", Exec.ExpectInActiveStates("Tree1Root", "Tree1RootState1", "Tree1State1StateSub1A", "Tree1Sub1", "Tree1Sub1StateLinkTree2A", "Tree2StateRoot", "Tree2RootStateSub1B", "Tree2Sub1", "Tree2Sub1State2"));
+			AITEST_EQUAL(TEXT("Start should complete with Running"), Status, EStateTreeRunStatus::Running);
+			AITEST_TRUE(TEXT("Should be in the correct state"), Exec.ExpectInActiveStates("Tree1Root", "Tree1RootState1", "Tree1State1StateSub1A", "Tree1Sub1", "Tree1Sub1StateLinkTree2A", "Tree2StateRoot", "Tree2RootStateSub1B", "Tree2Sub1", "Tree2Sub1State2"));
 
 			UE::StateTree::FActiveStatePath ExecPath = InstanceData.GetExecutionState()->GetActiveStatePath();
-			AITEST_TRUE("Has the correct number of path elements", ExecPath.Num() == 9);
-			AITEST_TRUE("Has the correct number of active states", InstanceData.GetExecutionState()->ActiveFrames.Num() == 4);
-			AITEST_TRUE("All frames are the same", TestFrameId(4));
+			AITEST_TRUE(TEXT("Has the correct number of path elements"), ExecPath.Num() == 9);
+			AITEST_TRUE(TEXT("Has the correct number of active states"), InstanceData.GetExecutionState()->ActiveFrames.Num() == 4);
+			AITEST_TRUE(TEXT("All frames are the same"), TestFrameId(4));
 			Exec.LogClear();
 		}
 		{
 			PreTestFrameId();
 			EStateTreeRunStatus Status = Exec.Tick(1.0f);
-			AITEST_EQUAL("Start should complete with Running", Status, EStateTreeRunStatus::Running);
-			AITEST_TRUE("Should be in the correct state", Exec.ExpectInActiveStates("Tree1Root", "Tree1RootState1", "Tree1State1StateSub1A", "Tree1Sub1", "Tree1Sub1StateLinkTree2A", "Tree2StateRoot", "Tree2RootState1"));
+			AITEST_EQUAL(TEXT("Start should complete with Running"), Status, EStateTreeRunStatus::Running);
+			AITEST_TRUE(TEXT("Should be in the correct state"), Exec.ExpectInActiveStates("Tree1Root", "Tree1RootState1", "Tree1State1StateSub1A", "Tree1Sub1", "Tree1Sub1StateLinkTree2A", "Tree2StateRoot", "Tree2RootState1"));
 
 			UE::StateTree::FActiveStatePath ExecPath = InstanceData.GetExecutionState()->GetActiveStatePath();
-			AITEST_TRUE("Has the correct number of path elements", ExecPath.Num() == 7);
-			AITEST_TRUE("Has the correct number of active states", InstanceData.GetExecutionState()->ActiveFrames.Num() == 3);
-			AITEST_TRUE("All frames are the same", TestFrameId(3));
+			AITEST_TRUE(TEXT("Has the correct number of path elements"), ExecPath.Num() == 7);
+			AITEST_TRUE(TEXT("Has the correct number of active states"), InstanceData.GetExecutionState()->ActiveFrames.Num() == 3);
+			AITEST_TRUE(TEXT("All frames are the same"), TestFrameId(3));
 			Exec.LogClear();
 		}
 		{
 			PreTestFrameId();
 			EStateTreeRunStatus Status = Exec.Tick(1.0f);
-			AITEST_EQUAL("Start should complete with Running", Status, EStateTreeRunStatus::Running);
-			AITEST_TRUE("Should be in the correct state", Exec.ExpectInActiveStates("Tree1Root", "Tree1RootState1", "Tree1State1StateSub1A", "Tree1Sub1", "Tree1Sub1StateLinkTree2B", "Tree2StateRoot", "Tree2RootStateSub1A", "Tree2Sub1", "Tree2Sub1State1"));
+			AITEST_EQUAL(TEXT("Start should complete with Running"), Status, EStateTreeRunStatus::Running);
+			AITEST_TRUE(TEXT("Should be in the correct state"), Exec.ExpectInActiveStates("Tree1Root", "Tree1RootState1", "Tree1State1StateSub1A", "Tree1Sub1", "Tree1Sub1StateLinkTree2B", "Tree2StateRoot", "Tree2RootStateSub1A", "Tree2Sub1", "Tree2Sub1State1"));
 
 			UE::StateTree::FActiveStatePath ExecPath = InstanceData.GetExecutionState()->GetActiveStatePath();
-			AITEST_TRUE("Has the correct number of path elements", ExecPath.Num() == 9);
-			AITEST_TRUE("Has the correct number of active states", InstanceData.GetExecutionState()->ActiveFrames.Num() == 4);
+			AITEST_TRUE(TEXT("Has the correct number of path elements"), ExecPath.Num() == 9);
+			AITEST_TRUE(TEXT("Has the correct number of active states"), InstanceData.GetExecutionState()->ActiveFrames.Num() == 4);
 			//@todo Temporary reverted the frameID fixed until TransitionID
-			//AITEST_TRUE("All frames are the same", TestFrameId(2));
+			//AITEST_TRUE(TEXT("All frames are the same"), TestFrameId(2));
 			Exec.LogClear();
 		}
 		{
 			PreTestFrameId();
 			EStateTreeRunStatus Status = Exec.Tick(1.0f);
-			AITEST_EQUAL("Start should complete with Running", Status, EStateTreeRunStatus::Running);
-			AITEST_TRUE("Should be in the correct state", Exec.ExpectInActiveStates("Tree1Root", "Tree1RootState1", "Tree1State1StateSub1A", "Tree1Sub1", "Tree1Sub1StateLinkTree2B", "Tree2StateRoot", "Tree2RootStateSub1A", "Tree2Sub1", "Tree2Sub1State2"));
-			AITEST_TRUE("All frames are the same", TestFrameId(4));
+			AITEST_EQUAL(TEXT("Start should complete with Running"), Status, EStateTreeRunStatus::Running);
+			AITEST_TRUE(TEXT("Should be in the correct state"), Exec.ExpectInActiveStates("Tree1Root", "Tree1RootState1", "Tree1State1StateSub1A", "Tree1Sub1", "Tree1Sub1StateLinkTree2B", "Tree2StateRoot", "Tree2RootStateSub1A", "Tree2Sub1", "Tree2Sub1State2"));
+			AITEST_TRUE(TEXT("All frames are the same"), TestFrameId(4));
 			Exec.LogClear();
 
 			PreTestFrameId();
 			Status = Exec.Tick(1.0f);
-			AITEST_EQUAL("Start should complete with Running", Status, EStateTreeRunStatus::Running);
-			AITEST_TRUE("Should be in the correct state", Exec.ExpectInActiveStates("Tree1Root", "Tree1RootState1", "Tree1State1StateSub1A", "Tree1Sub1", "Tree1Sub1StateLinkTree2B", "Tree2StateRoot", "Tree2RootStateSub1B", "Tree2Sub1", "Tree2Sub1State1")); 
+			AITEST_EQUAL(TEXT("Start should complete with Running"), Status, EStateTreeRunStatus::Running);
+			AITEST_TRUE(TEXT("Should be in the correct state"), Exec.ExpectInActiveStates("Tree1Root", "Tree1RootState1", "Tree1State1StateSub1A", "Tree1Sub1", "Tree1Sub1StateLinkTree2B", "Tree2StateRoot", "Tree2RootStateSub1B", "Tree2Sub1", "Tree2Sub1State1"));
 			//@todo Temporary reverted the frameID fixed until TransitionID
-			//AITEST_TRUE("All frames are the same", TestFrameId(3));
+			//AITEST_TRUE(TEXT("All frames are the same"), TestFrameId(3));
 			Exec.LogClear();
 
 			PreTestFrameId();
 			Status = Exec.Tick(1.0f);
-			AITEST_EQUAL("Start should complete with Running", Status, EStateTreeRunStatus::Running);
-			AITEST_TRUE("Should be in the correct state", Exec.ExpectInActiveStates("Tree1Root", "Tree1RootState1", "Tree1State1StateSub1A", "Tree1Sub1", "Tree1Sub1StateLinkTree2B", "Tree2StateRoot", "Tree2RootStateSub1B", "Tree2Sub1", "Tree2Sub1State2"));
-			AITEST_TRUE("Has the correct number of active states", InstanceData.GetExecutionState()->ActiveFrames.Num() == 4);
-			AITEST_TRUE("All frames are the same", TestFrameId(4));
+			AITEST_EQUAL(TEXT("Start should complete with Running"), Status, EStateTreeRunStatus::Running);
+			AITEST_TRUE(TEXT("Should be in the correct state"), Exec.ExpectInActiveStates("Tree1Root", "Tree1RootState1", "Tree1State1StateSub1A", "Tree1Sub1", "Tree1Sub1StateLinkTree2B", "Tree2StateRoot", "Tree2RootStateSub1B", "Tree2Sub1", "Tree2Sub1State2"));
+			AITEST_TRUE(TEXT("Has the correct number of active states"), InstanceData.GetExecutionState()->ActiveFrames.Num() == 4);
+			AITEST_TRUE(TEXT("All frames are the same"), TestFrameId(4));
 			Exec.LogClear();
 
 			PreTestFrameId();
 			Status = Exec.Tick(1.0f);
-			AITEST_EQUAL("Start should complete with Running", Status, EStateTreeRunStatus::Running);
-			AITEST_TRUE("Should be in the correct state", Exec.ExpectInActiveStates("Tree1Root", "Tree1RootState1", "Tree1State1StateSub1A", "Tree1Sub1", "Tree1Sub1StateLinkTree2B", "Tree2StateRoot", "Tree2RootState1"));
-			AITEST_TRUE("Has the correct number of active states", InstanceData.GetExecutionState()->ActiveFrames.Num() == 3);
-			AITEST_TRUE("All frames are the same", TestFrameId(3));
+			AITEST_EQUAL(TEXT("Start should complete with Running"), Status, EStateTreeRunStatus::Running);
+			AITEST_TRUE(TEXT("Should be in the correct state"), Exec.ExpectInActiveStates("Tree1Root", "Tree1RootState1", "Tree1State1StateSub1A", "Tree1Sub1", "Tree1Sub1StateLinkTree2B", "Tree2StateRoot", "Tree2RootState1"));
+			AITEST_TRUE(TEXT("Has the correct number of active states"), InstanceData.GetExecutionState()->ActiveFrames.Num() == 3);
+			AITEST_TRUE(TEXT("All frames are the same"), TestFrameId(3));
 			Exec.LogClear();
 		}
 		{
 			PreTestFrameId();
 			EStateTreeRunStatus Status = Exec.Tick(1.0f);
-			AITEST_EQUAL("Start should complete with Running", Status, EStateTreeRunStatus::Running);
-			AITEST_TRUE("Should be in the correct state", Exec.ExpectInActiveStates("Tree1Root", "Tree1RootState1", "Tree1State1StateSub1A", "Tree1Sub1", "Tree1Sub1StateSub2A", "Tree1Sub2", "Tree1Sub2State1"));
-			AITEST_TRUE("Has the correct number of active states", InstanceData.GetExecutionState()->ActiveFrames.Num() == 3);
-			AITEST_TRUE("All frames are the same", TestFrameId(2));
+			AITEST_EQUAL(TEXT("Start should complete with Running"), Status, EStateTreeRunStatus::Running);
+			AITEST_TRUE(TEXT("Should be in the correct state"), Exec.ExpectInActiveStates("Tree1Root", "Tree1RootState1", "Tree1State1StateSub1A", "Tree1Sub1", "Tree1Sub1StateSub2A", "Tree1Sub2", "Tree1Sub2State1"));
+			AITEST_TRUE(TEXT("Has the correct number of active states"), InstanceData.GetExecutionState()->ActiveFrames.Num() == 3);
+			AITEST_TRUE(TEXT("All frames are the same"), TestFrameId(2));
 			Exec.LogClear();
 
 			PreTestFrameId();
 			Status = Exec.Tick(1.0f);
-			AITEST_EQUAL("Start should complete with Running", Status, EStateTreeRunStatus::Running);
-			AITEST_TRUE("Should be in the correct state", Exec.ExpectInActiveStates("Tree1Root", "Tree1RootState1", "Tree1State1StateSub1A", "Tree1Sub1", "Tree1Sub1StateSub2A", "Tree1Sub2", "Tree1Sub2State2"));
-			AITEST_TRUE("Has the correct number of active states", InstanceData.GetExecutionState()->ActiveFrames.Num() == 3);
-			AITEST_TRUE("All frames are the same", TestFrameId(3));
+			AITEST_EQUAL(TEXT("Start should complete with Running"), Status, EStateTreeRunStatus::Running);
+			AITEST_TRUE(TEXT("Should be in the correct state"), Exec.ExpectInActiveStates("Tree1Root", "Tree1RootState1", "Tree1State1StateSub1A", "Tree1Sub1", "Tree1Sub1StateSub2A", "Tree1Sub2", "Tree1Sub2State2"));
+			AITEST_TRUE(TEXT("Has the correct number of active states"), InstanceData.GetExecutionState()->ActiveFrames.Num() == 3);
+			AITEST_TRUE(TEXT("All frames are the same"), TestFrameId(3));
 			Exec.LogClear();
 		}
 
diff --git a/Engine/Plugins/Runtime/StateTree/Source/StateTreeTestSuite/Private/StateTreeTaskStateTest.cpp b/Engine/Plugins/Runtime/StateTree/Source/StateTreeTestSuite/Private/StateTreeTaskStateTest.cpp
--- a/Engine/Plugins/Runtime/StateTree/Source/StateTreeTestSuite/Private/StateTreeTaskStateTest.cpp
+++ b/Engine/Plugins/Runtime/StateTree/Source/StateTreeTestSuite/Private/StateTreeTaskStateTest.cpp
@@ -61,7 +61,7 @@ struct FStateTreeTest_TasksCompletion_AllAny : FStateTreeTestBase
 				FStateTreeCompilerLog Log;
 				FStateTreeCompiler Compiler(Log);
 				const bool bResult = Compiler.Compile(StateTree);
-				AITEST_TRUE("StateTree should get compiled", bResult);
+				AITEST_TRUE(TEXT("StateTree should get compiled"), bResult);
 			}
 		}
 
@@ -71,11 +71,11 @@ struct FStateTreeTest_TasksCompletion_AllAny : FStateTreeTestBase
 			FStateTreeInstanceData InstanceData;
 			FTestStateTreeExecutionContext Exec(StateTree, StateTree, InstanceData);
 			const bool bInitSucceeded = Exec.IsValid();
-			AITEST_TRUE("StateTree should init", bInitSucceeded);
+			AITEST_TRUE(TEXT("StateTree should init"), bInitSucceeded);
 
 			{
 				Status = Exec.Start();
-				AITEST_EQUAL("Start should complete with Running", Status, EStateTreeRunStatus::Running);
+				AITEST_EQUAL(TEXT("Start should complete with Running"), Status, EStateTreeRunStatus::Running);
 				Exec.LogClear();
 			}
 
@@ -106,7 +106,7 @@ struct FStateTreeTest_TasksCompletion_AllAny : FStateTreeTestBase
 						AITEST_TRUE(*FString::Printf(TEXT("Should tick Task A %d, %d"), TickIndex, TaskIndex), bTickedA);
 					}
 				}
-				const bool bStateSucceeded = Exec.Expect("Tree1StateRootTask_0", "StateCompleted");
+				const bool bStateSucceeded = Exec.Expect("Tree1StateRootTask_0", TEXT("StateCompleted"));
 				const bool bLastTick = TickIndex == FStateTreeTasksCompletionStatus::MaxNumberOfTasksPerGroup - 1;
 				AITEST_EQUAL(TEXT("Tick should not complete the task."), bStateSucceeded, bLastTick);
 				Exec.LogClear();
@@ -128,18 +128,18 @@ struct FStateTreeTest_TasksCompletion_AllAny : FStateTreeTestBase
 			FStateTreeCompilerLog Log;
 			FStateTreeCompiler Compiler(Log);
 			const bool bResult = Compiler.Compile(StateTree);
-			AITEST_TRUE("StateTree should get compiled", bResult);
+			AITEST_TRUE(TEXT("StateTree should get compiled"), bResult);
 		}
 		{
 			EStateTreeRunStatus Status = EStateTreeRunStatus::Unset;
 			FStateTreeInstanceData InstanceData;
 			FTestStateTreeExecutionContext Exec(StateTree, StateTree, InstanceData);
 			const bool bInitSucceeded = Exec.IsValid();
-			AITEST_TRUE("StateTree should init", bInitSucceeded);
+			AITEST_TRUE(TEXT("StateTree should init"), bInitSucceeded);
 
 			{
 				Status = Exec.Start();
-				AITEST_EQUAL("Start should complete with Running", Status, EStateTreeRunStatus::Running);
+				AITEST_EQUAL(TEXT("Start should complete with Running"), Status, EStateTreeRunStatus::Running);
 				Exec.LogClear();
 			}
 
@@ -152,7 +152,7 @@ struct FStateTreeTest_TasksCompletion_AllAny : FStateTreeTestBase
 					bool bTicked = Exec.Expect(*FString::Printf(TEXT("Tree1StateRootTask_%d"), TaskIndex), TEXT("Tick"));
 					AITEST_TRUE(*FString::Printf(TEXT("Should tick Task %d"), TaskIndex), bTicked);
 				}
-				const bool bStateSucceeded = Exec.Expect("Tree1StateRootTask_0", "StateCompleted");
+				const bool bStateSucceeded = Exec.Expect("Tree1StateRootTask_0", TEXT("StateCompleted"));
 				AITEST_TRUE(TEXT("Tick should not complete the task."), bStateSucceeded);
 				Exec.LogClear();
 			}
@@ -198,7 +198,7 @@ struct FStateTreeTest_TasksCompletion_FailureTasks : FStateTreeTestBase
 				FStateTreeCompilerLog Log;
 				FStateTreeCompiler Compiler(Log);
 				const bool bResult = Compiler.Compile(StateTree);
-				AITEST_TRUE("StateTree should get compiled", bResult);
+				AITEST_TRUE(TEXT("StateTree should get compiled"), bResult);
 			}
 		}
 
@@ -207,11 +207,11 @@ struct FStateTreeTest_TasksCompletion_FailureTasks : FStateTreeTestBase
 			FStateTreeInstanceData InstanceData;
 			FTestStateTreeExecutionContext Exec(StateTree, StateTree, InstanceData);
 			const bool bInitSucceeded = Exec.IsValid();
-			AITEST_TRUE("StateTree should init", bInitSucceeded);
+			AITEST_TRUE(TEXT("StateTree should init"), bInitSucceeded);
 
 			{
 				Status = Exec.Start();
-				AITEST_EQUAL("Start should complete with Running", Status, EStateTreeRunStatus::Running);
+				AITEST_EQUAL(TEXT("Start should complete with Running"), Status, EStateTreeRunStatus::Running);
 				Exec.LogClear();
 			}
 
@@ -233,7 +233,7 @@ struct FStateTreeTest_TasksCompletion_FailureTasks : FStateTreeTestBase
 						AITEST_TRUE(*FString::Printf(TEXT("Should tick Task %d, %d"), TickIndex, TaskIndex), bTicked);
 					}
 				}
-				const bool bStateSucceeded = Exec.Expect("Tree1StateRootTask_0", "StateCompleted");
+				const bool bStateSucceeded = Exec.Expect("Tree1StateRootTask_0", TEXT("StateCompleted"));
 				AITEST_EQUAL(TEXT("Tick should not complete the task."), bStateSucceeded, bLastTick);
 				Exec.LogClear();
 			}
diff --git a/Engine/Plugins/Runtime/StateTree/Source/StateTreeTestSuite/Private/StateTreeTest.cpp b/Engine/Plugins/Runtime/StateTree/Source/StateTreeTestSuite/Private/StateTreeTest.cpp
--- a/Engine/Plugins/Runtime/StateTree/Source/StateTreeTestSuite/Private/StateTreeTest.cpp
+++ b/Engine/Plugins/Runtime/StateTree/Source/StateTreeTestSuite/Private/StateTreeTest.cpp
@@ -66,8 +66,8 @@ struct FStateTreeTest_MakeAndBakeStateTree : FStateTreeTestBase
 		FStateTreeCompiler Compiler(Log);
 		const bool bResult = Compiler.Compile(StateTree);
 
-		AITEST_TRUE("StateTree should get compiled", bResult);
-		AITEST_TRUE("StateTree should be ready to run", StateTree.IsReadyToRun());
+		AITEST_TRUE(TEXT("StateTree should get compiled"), bResult);
+		AITEST_TRUE(TEXT("StateTree should be ready to run"), StateTree.IsReadyToRun());
 
 		return true;
 	}
@@ -89,20 +89,20 @@ struct FStateTreeTest_EmptyStateTree : FStateTreeTestBase
 		FStateTreeCompiler Compiler(Log);
 		const bool bResult = Compiler.Compile(StateTree);
 
-		AITEST_TRUE("StateTree should get compiled", bResult);
+		AITEST_TRUE(TEXT("StateTree should get compiled"), bResult);
 
 		EStateTreeRunStatus Status = EStateTreeRunStatus::Unset;
 		FStateTreeInstanceData InstanceData;
 		FTestStateTreeExecutionContext Exec(StateTree, StateTree, InstanceData);
 		const bool bInitSucceeded = Exec.IsValid();
-		AITEST_TRUE("StateTree should init", bInitSucceeded);
+		AITEST_TRUE(TEXT("StateTree should init"), bInitSucceeded);
 
 		Status = Exec.Start();
-		AITEST_TRUE("StateTree should be running", Status == EStateTreeRunStatus::Running);
+		AITEST_TRUE(TEXT("StateTree should be running"), Status == EStateTreeRunStatus::Running);
 		Exec.LogClear();
 
 		Status = Exec.Tick(0.1f);
-		AITEST_TRUE("StateTree should be completed", Status == EStateTreeRunStatus::Succeeded);
+		AITEST_TRUE(TEXT("StateTree should be completed"), Status == EStateTreeRunStatus::Succeeded);
 		Exec.LogClear();
 
 		Exec.Stop();
@@ -132,39 +132,39 @@ struct FStateTreeTest_Sequence : FStateTreeTestBase
 		FStateTreeCompilerLog Log;
 		FStateTreeCompiler Compiler(Log);
 		const bool bResult = Compiler.Compile(StateTree);
-		AITEST_TRUE("StateTree should get compiled", bResult);
+		AITEST_TRUE(TEXT("StateTree should get compiled"), bResult);
 
 		EStateTreeRunStatus Status = EStateTreeRunStatus::Unset;
 		FStateTreeInstanceData InstanceData;
 		FTestStateTreeExecutionContext Exec(StateTree, StateTree, InstanceData);
 		const bool bInitSucceeded = Exec.IsValid();
-		AITEST_TRUE("StateTree should init", bInitSucceeded);
+		AITEST_TRUE(TEXT("StateTree should init"), bInitSucceeded);
 
-		const FString TickStr(TEXT("Tick"));
-		const FString EnterStateStr(TEXT("EnterState"));
-		const FString ExitStateStr(TEXT("ExitState"));
+		const TCHAR* TickStr(TEXT("Tick"));
+		const TCHAR* EnterStateStr(TEXT("EnterState"));
+		const TCHAR* ExitStateStr(TEXT("ExitState"));
 		
 		Status = Exec.Start();
-		AITEST_TRUE("StateTree Task1 should enter state", Exec.Expect(Task1.GetName(), EnterStateStr));
-		AITEST_FALSE("StateTree Task1 should not tick", Exec.Expect(Task1.GetName(), TickStr));
+		AITEST_TRUE(TEXT("StateTree Task1 should enter state"), Exec.Expect(Task1.GetName(), EnterStateStr));
+		AITEST_FALSE(TEXT("StateTree Task1 should not tick"), Exec.Expect(Task1.GetName(), TickStr));
 		Exec.LogClear();
 
 		Status = Exec.Tick(0.1f);
-		AITEST_TRUE("StateTree Task1 should tick, and exit state", Exec.Expect(Task1.GetName(), TickStr).Then(Task1.GetName(), ExitStateStr));
-		AITEST_TRUE("StateTree Task2 should enter state", Exec.Expect(Task2.GetName(), EnterStateStr));
-		AITEST_FALSE("StateTree Task2 should not tick", Exec.Expect(Task2.GetName(), TickStr));
-		AITEST_TRUE("StateTree should be running", Status == EStateTreeRunStatus::Running);
+		AITEST_TRUE(TEXT("StateTree Task1 should tick, and exit state"), Exec.Expect(Task1.GetName(), TickStr).Then(Task1.GetName(), ExitStateStr));
+		AITEST_TRUE(TEXT("StateTree Task2 should enter state"), Exec.Expect(Task2.GetName(), EnterStateStr));
+		AITEST_FALSE(TEXT("StateTree Task2 should not tick"), Exec.Expect(Task2.GetName(), TickStr));
+		AITEST_TRUE(TEXT("StateTree should be running"), Status == EStateTreeRunStatus::Running);
 		Exec.LogClear();
 		
 		Status = Exec.Tick(0.1f);
-		AITEST_TRUE("StateTree Task2 should tick, and exit state", Exec.Expect(Task2.GetName(), TickStr).Then(Task2.GetName(), ExitStateStr));
-		AITEST_FALSE("StateTree Task1 should not tick", Exec.Expect(Task1.GetName(), TickStr));
-		AITEST_TRUE("StateTree should be completed", Status == EStateTreeRunStatus::Succeeded);
+		AITEST_TRUE(TEXT("StateTree Task2 should tick, and exit state"), Exec.Expect(Task2.GetName(), TickStr).Then(Task2.GetName(), ExitStateStr));
+		AITEST_FALSE(TEXT("StateTree Task1 should not tick"), Exec.Expect(Task1.GetName(), TickStr));
+		AITEST_TRUE(TEXT("StateTree should be completed"), Status == EStateTreeRunStatus::Succeeded);
 		Exec.LogClear();
 
 		Status = Exec.Tick(0.1f);
-		AITEST_FALSE("StateTree Task1 should not tick", Exec.Expect(Task1.GetName(), TickStr));
-		AITEST_FALSE("StateTree Task2 should not tick", Exec.Expect(Task2.GetName(), TickStr));
+		AITEST_FALSE(TEXT("StateTree Task1 should not tick"), Exec.Expect(Task1.GetName(), TickStr));
+		AITEST_FALSE(TEXT("StateTree Task2 should not tick"), Exec.Expect(Task2.GetName(), TickStr));
 		Exec.LogClear();
 
 		Exec.Stop();
@@ -198,47 +198,47 @@ struct FStateTreeTest_Select : FStateTreeTestBase
 		FStateTreeCompilerLog Log;
 		FStateTreeCompiler Compiler(Log);
 		const bool bResult = Compiler.Compile(StateTree);
-		AITEST_TRUE("StateTree should get compiled", bResult);
+		AITEST_TRUE(TEXT("StateTree should get compiled"), bResult);
 
 		EStateTreeRunStatus Status = EStateTreeRunStatus::Unset;
 		FStateTreeInstanceData InstanceData;
 		FTestStateTreeExecutionContext Exec(StateTree, StateTree, InstanceData);
 		const bool bInitSucceeded = Exec.IsValid();
-		AITEST_TRUE("StateTree should init", bInitSucceeded);
+		AITEST_TRUE(TEXT("StateTree should init"), bInitSucceeded);
 
-		const FString TickStr(TEXT("Tick"));
-		const FString EnterStateStr(TEXT("EnterState"));
-		const FString ExitStateStr(TEXT("ExitState"));
+		const TCHAR* TickStr(TEXT("Tick"));
+		const TCHAR* EnterStateStr(TEXT("EnterState"));
+		const TCHAR* ExitStateStr(TEXT("ExitState"));
 
 		// Start and enter state
 		Status = Exec.Start();
-		AITEST_TRUE("StateTree TaskRoot should enter state", Exec.Expect(TaskRoot.GetName(), EnterStateStr));
-		AITEST_TRUE("StateTree Task1 should enter state", Exec.Expect(Task1.GetName(), EnterStateStr));
-		AITEST_TRUE("StateTree Task1A should enter state", Exec.Expect(Task1A.GetName(), EnterStateStr));
-		AITEST_FALSE("StateTree TaskRoot should not tick", Exec.Expect(TaskRoot.GetName(), TickStr));
-		AITEST_FALSE("StateTree Task1 should not tick", Exec.Expect(Task1.GetName(), TickStr));
-		AITEST_FALSE("StateTree Task1A should not tick", Exec.Expect(Task1A.GetName(), TickStr));
-		AITEST_TRUE("StateTree should be running", Status == EStateTreeRunStatus::Running);
+		AITEST_TRUE(TEXT("StateTree TaskRoot should enter state"), Exec.Expect(TaskRoot.GetName(), EnterStateStr));
+		AITEST_TRUE(TEXT("StateTree Task1 should enter state"), Exec.Expect(Task1.GetName(), EnterStateStr));
+		AITEST_TRUE(TEXT("StateTree Task1A should enter state"), Exec.Expect(Task1A.GetName(), EnterStateStr));
+		AITEST_FALSE(TEXT("StateTree TaskRoot should not tick"), Exec.Expect(TaskRoot.GetName(), TickStr));
+		AITEST_FALSE(TEXT("StateTree Task1 should not tick"), Exec.Expect(Task1.GetName(), TickStr));
+		AITEST_FALSE(TEXT("StateTree Task1A should not tick"), Exec.Expect(Task1A.GetName(), TickStr));
+		AITEST_TRUE(TEXT("StateTree should be running"), Status == EStateTreeRunStatus::Running);
 		Exec.LogClear();
 
 		// Regular tick, no state selection at all.
 		Status = Exec.Tick(0.1f);
-		AITEST_TRUE("StateTree tasks should update in order", Exec.Expect(TaskRoot.GetName(), TickStr).Then(Task1.GetName(), TickStr).Then(Task1A.GetName(), TickStr));
-		AITEST_FALSE("StateTree TaskRoot should not EnterState", Exec.Expect(TaskRoot.GetName(), EnterStateStr));
-		AITEST_FALSE("StateTree Task1 should not EnterState", Exec.Expect(Task1.GetName(), EnterStateStr));
-		AITEST_FALSE("StateTree Task1A should not EnterState", Exec.Expect(Task1A.GetName(), EnterStateStr));
-		AITEST_FALSE("StateTree TaskRoot should not ExitState", Exec.Expect(TaskRoot.GetName(), ExitStateStr));
-		AITEST_FALSE("StateTree Task1 should not ExitState", Exec.Expect(Task1.GetName(), ExitStateStr));
-		AITEST_FALSE("StateTree Task1A should not ExitState", Exec.Expect(Task1A.GetName(), ExitStateStr));
-		AITEST_TRUE("StateTree should be running", Status == EStateTreeRunStatus::Running);
+		AITEST_TRUE(TEXT("StateTree tasks should update in order"), Exec.Expect(TaskRoot.GetName(), TickStr).Then(Task1.GetName(), TickStr).Then(Task1A.GetName(), TickStr));
+		AITEST_FALSE(TEXT("StateTree TaskRoot should not EnterState"), Exec.Expect(TaskRoot.GetName(), EnterStateStr));
+		AITEST_FALSE(TEXT("StateTree Task1 should not EnterState"), Exec.Expect(Task1.GetName(), EnterStateStr));
+		AITEST_FALSE(TEXT("StateTree Task1A should not EnterState"), Exec.Expect(Task1A.GetName(), EnterStateStr));
+		AITEST_FALSE(TEXT("StateTree TaskRoot should not ExitState"), Exec.Expect(TaskRoot.GetName(), ExitStateStr));
+		AITEST_FALSE(TEXT("StateTree Task1 should not ExitState"), Exec.Expect(Task1.GetName(), ExitStateStr));
+		AITEST_FALSE(TEXT("StateTree Task1A should not ExitState"), Exec.Expect(Task1A.GetName(), ExitStateStr));
+		AITEST_TRUE(TEXT("StateTree should be running"), Status == EStateTreeRunStatus::Running);
 		Exec.LogClear();
 
 		// Partial reselect, Root should not get EnterState
 		Status = Exec.Tick(0.1f);
-		AITEST_FALSE("StateTree TaskRoot should not enter state", Exec.Expect(TaskRoot.GetName(), EnterStateStr));
-		AITEST_TRUE("StateTree Task1 should tick, exit state, and enter state", Exec.Expect(Task1.GetName(), TickStr).Then(Task1.GetName(), ExitStateStr).Then(Task1.GetName(), EnterStateStr));
-		AITEST_TRUE("StateTree Task1A should tick, exit state, and enter state", Exec.Expect(Task1A.GetName(), TickStr).Then(Task1A.GetName(), ExitStateStr).Then(Task1A.GetName(), EnterStateStr));
-		AITEST_TRUE("StateTree should be running", Status == EStateTreeRunStatus::Running);
+		AITEST_FALSE(TEXT("StateTree TaskRoot should not enter state"), Exec.Expect(TaskRoot.GetName(), EnterStateStr));
+		AITEST_TRUE(TEXT("StateTree Task1 should tick, exit state, and enter state"), Exec.Expect(Task1.GetName(), TickStr).Then(Task1.GetName(), ExitStateStr).Then(Task1.GetName(), EnterStateStr));
+		AITEST_TRUE(TEXT("StateTree Task1A should tick, exit state, and enter state"), Exec.Expect(Task1A.GetName(), TickStr).Then(Task1A.GetName(), ExitStateStr).Then(Task1A.GetName(), EnterStateStr));
+		AITEST_TRUE(TEXT("StateTree should be running"), Status == EStateTreeRunStatus::Running);
 		Exec.LogClear();
 
 		Exec.Stop();
@@ -273,37 +273,37 @@ struct FStateTreeTest_FailEnterState : FStateTreeTestBase
 		FStateTreeCompilerLog Log;
 		FStateTreeCompiler Compiler(Log);
 		const bool bResult = Compiler.Compile(StateTree);
-		AITEST_TRUE("StateTree should get compiled", bResult);
+		AITEST_TRUE(TEXT("StateTree should get compiled"), bResult);
 
 		EStateTreeRunStatus Status = EStateTreeRunStatus::Unset;
 		FStateTreeInstanceData InstanceData;
 		FTestStateTreeExecutionContext Exec(StateTree, StateTree, InstanceData);
 		const bool bInitSucceeded = Exec.IsValid();
-		AITEST_TRUE("StateTree should init", bInitSucceeded);
+		AITEST_TRUE(TEXT("StateTree should init"), bInitSucceeded);
 
-		const FString TickStr(TEXT("Tick"));
-		const FString EnterStateStr(TEXT("EnterState"));
-		const FString ExitStateStr(TEXT("ExitState"));
-		const FString StateCompletedStr(TEXT("StateCompleted"));
+		const TCHAR* TickStr(TEXT("Tick"));
+		const TCHAR* EnterStateStr(TEXT("EnterState"));
+		const TCHAR* ExitStateStr(TEXT("ExitState"));
+		const TCHAR* StateCompletedStr(TEXT("StateCompleted"));
 
 		// Start and enter state
 		Status = Exec.Start();
-		AITEST_TRUE("StateTree TaskRoot should enter state", Exec.Expect(TaskRoot.GetName(), EnterStateStr));
-		AITEST_TRUE("StateTree Task1 should enter state", Exec.Expect(Task1.GetName(), EnterStateStr));
-		AITEST_TRUE("StateTree Task2 should enter state", Exec.Expect(Task2.GetName(), EnterStateStr));
-		AITEST_FALSE("StateTree Task3 should not enter state", Exec.Expect(Task3.GetName(), EnterStateStr));
-		AITEST_TRUE("StateTree Should execute StateCompleted in reverse order", Exec.Expect(Task2.GetName(), StateCompletedStr).Then(Task1.GetName(), StateCompletedStr).Then(TaskRoot.GetName(), StateCompletedStr));
-		AITEST_FALSE("StateTree Task3 should not state complete", Exec.Expect(Task3.GetName(), StateCompletedStr));
-		AITEST_TRUE("StateTree exec status should be failed", Exec.GetLastTickStatus() == EStateTreeRunStatus::Failed);
+		AITEST_TRUE(TEXT("StateTree TaskRoot should enter state"), Exec.Expect(TaskRoot.GetName(), EnterStateStr));
+		AITEST_TRUE(TEXT("StateTree Task1 should enter state"), Exec.Expect(Task1.GetName(), EnterStateStr));
+		AITEST_TRUE(TEXT("StateTree Task2 should enter state"), Exec.Expect(Task2.GetName(), EnterStateStr));
+		AITEST_FALSE(TEXT("StateTree Task3 should not enter state"), Exec.Expect(Task3.GetName(), EnterStateStr));
+		AITEST_TRUE(TEXT("StateTree Should execute StateCompleted in reverse order"), Exec.Expect(Task2.GetName(), StateCompletedStr).Then(Task1.GetName(), StateCompletedStr).Then(TaskRoot.GetName(), StateCompletedStr));
+		AITEST_FALSE(TEXT("StateTree Task3 should not state complete"), Exec.Expect(Task3.GetName(), StateCompletedStr));
+		AITEST_TRUE(TEXT("StateTree exec status should be failed"), Exec.GetLastTickStatus() == EStateTreeRunStatus::Failed);
 		Exec.LogClear();
 
 		// Stop and exit state
 		Status = Exec.Stop();
-		AITEST_TRUE("StateTree TaskRoot should exit state", Exec.Expect(TaskRoot.GetName(), ExitStateStr));
-		AITEST_TRUE("StateTree Task1 should exit state", Exec.Expect(Task1.GetName(), ExitStateStr));
-		AITEST_TRUE("StateTree Task2 should exit state", Exec.Expect(Task2.GetName(), ExitStateStr));
-		AITEST_FALSE("StateTree Task3 should not exit state", Exec.Expect(Task3.GetName(), ExitStateStr));
-		AITEST_TRUE("StateTree status should be stopped", Status == EStateTreeRunStatus::Stopped);
+		AITEST_TRUE(TEXT("StateTree TaskRoot should exit state"), Exec.Expect(TaskRoot.GetName(), ExitStateStr));
+		AITEST_TRUE(TEXT("StateTree Task1 should exit state"), Exec.Expect(Task1.GetName(), ExitStateStr));
+		AITEST_TRUE(TEXT("StateTree Task2 should exit state"), Exec.Expect(Task2.GetName(), ExitStateStr));
+		AITEST_FALSE(TEXT("StateTree Task3 should not exit state"), Exec.Expect(Task3.GetName(), ExitStateStr));
+		AITEST_TRUE(TEXT("StateTree status should be stopped"), Status == EStateTreeRunStatus::Stopped);
 		Exec.LogClear();
 
 		return true;
@@ -336,27 +336,27 @@ struct FStateTreeTest_Restart : FStateTreeTestBase
 		const bool bInitSucceeded = Exec.IsValid();
 		AITEST_TRUE("StateTree should init", bInitSucceeded);
 
-		const FString TickStr(TEXT("Tick"));
-		const FString EnterStateStr(TEXT("EnterState"));
-		const FString ExitStateStr(TEXT("ExitState"));
-		const FString StateCompletedStr(TEXT("StateCompleted"));
+		const TCHAR* TickStr(TEXT("Tick"));
+		const TCHAR* EnterStateStr(TEXT("EnterState"));
+		const TCHAR* ExitStateStr(TEXT("ExitState"));
+		const TCHAR* StateCompletedStr(TEXT("StateCompleted"));
 
 		// Start and enter state
 		Status = Exec.Start();
-		AITEST_TRUE("StateTree Task1 should enter state", Exec.Expect(Task1.GetName(), EnterStateStr));
-		AITEST_TRUE("StateTree exec status should be running", Exec.GetLastTickStatus() == EStateTreeRunStatus::Running);
+		AITEST_TRUE(TEXT("StateTree Task1 should enter state"), Exec.Expect(Task1.GetName(), EnterStateStr));
+		AITEST_TRUE(TEXT("StateTree exec status should be running"), Exec.GetLastTickStatus() == EStateTreeRunStatus::Running);
 		Exec.LogClear();
 
 		// Tick
 		Status = Exec.Tick(0.1f);
-		AITEST_TRUE("StateTree exec status should be running", Exec.GetLastTickStatus() == EStateTreeRunStatus::Running);
+		AITEST_TRUE(TEXT("StateTree exec status should be running"), Exec.GetLastTickStatus() == EStateTreeRunStatus::Running);
 		Exec.LogClear();
 
 		// Call Start again, should stop and start the tree.
 		Status = Exec.Start();
-		AITEST_TRUE("StateTree Task1 should exit state", Exec.Expect(Task1.GetName(), ExitStateStr));
-		AITEST_TRUE("StateTree Task1 should enter state", Exec.Expect(Task1.GetName(), EnterStateStr));
-		AITEST_TRUE("StateTree exec status should be running", Exec.GetLastTickStatus() == EStateTreeRunStatus::Running);
+		AITEST_TRUE(TEXT("StateTree Task1 should exit state"), Exec.Expect(Task1.GetName(), ExitStateStr));
+		AITEST_TRUE(TEXT("StateTree Task1 should enter state"), Exec.Expect(Task1.GetName(), EnterStateStr));
+		AITEST_TRUE(TEXT("StateTree exec status should be running"), Exec.GetLastTickStatus() == EStateTreeRunStatus::Running);
 		Exec.LogClear();
 
 		Exec.Stop();
@@ -396,45 +396,45 @@ struct FStateTreeTest_SubTree_ActiveTasks : FStateTreeTestBase
 		FStateTreeCompilerLog Log;
 		FStateTreeCompiler Compiler(Log);
 		const bool bResult = Compiler.Compile(StateTree);
-		AITEST_TRUE("StateTree should get compiled", bResult);
+		AITEST_TRUE(TEXT("StateTree should get compiled"), bResult);
 
 		EStateTreeRunStatus Status = EStateTreeRunStatus::Unset;
 		FStateTreeInstanceData InstanceData;
 		FTestStateTreeExecutionContext Exec(StateTree, StateTree, InstanceData);
 		const bool bInitSucceeded = Exec.IsValid();
-		AITEST_TRUE("StateTree should init", bInitSucceeded);
+		AITEST_TRUE(TEXT("StateTree should init"), bInitSucceeded);
 
-		const FString TickStr(TEXT("Tick"));
-		const FString EnterStateStr(TEXT("EnterState"));
-		const FString ExitStateStr(TEXT("ExitState"));
-		const FString StateCompletedStr(TEXT("StateCompleted"));
+		const TCHAR* TickStr(TEXT("Tick"));
+		const TCHAR* EnterStateStr(TEXT("EnterState"));
+		const TCHAR* ExitStateStr(TEXT("ExitState"));
+		const TCHAR* StateCompletedStr(TEXT("StateCompleted"));
 
 		// Start and enter state
 		Status = Exec.Start();
 
-		AITEST_TRUE("StateTree Active States should be in Root/State1/State3/State3A", Exec.ExpectInActiveStates(Root.Name, State1.Name, State3.Name, State3A.Name));
-		AITEST_FALSE("StateTree Task2 should not enter state", Exec.Expect(Task2.GetName(), EnterStateStr));
-		AITEST_TRUE("StateTree Task3A should enter state", Exec.Expect(Task3A.GetName(), EnterStateStr));
-		AITEST_TRUE("StateTree should be running", Status == EStateTreeRunStatus::Running);
+		AITEST_TRUE(TEXT("StateTree Active States should be in Root/State1/State3/State3A"), Exec.ExpectInActiveStates(Root.Name, State1.Name, State3.Name, State3A.Name));
+		AITEST_FALSE(TEXT("StateTree Task2 should not enter state"), Exec.Expect(Task2.GetName(), EnterStateStr));
+		AITEST_TRUE(TEXT("StateTree Task3A should enter state"), Exec.Expect(Task3A.GetName(), EnterStateStr));
+		AITEST_TRUE(TEXT("StateTree should be running"), Status == EStateTreeRunStatus::Running);
 		Exec.LogClear();
 
 		// Transition within subtree
 		Status = Exec.Tick(0.1f);
-		AITEST_TRUE("StateTree Active States should be in Root/State1/State3/State3B", Exec.ExpectInActiveStates(Root.Name, State1.Name, State3.Name, State3B.Name));
-		AITEST_TRUE("StateTree Task3B should enter state", Exec.Expect(Task3B.GetName(), EnterStateStr));
-		AITEST_TRUE("StateTree should be running", Status == EStateTreeRunStatus::Running);
+		AITEST_TRUE(TEXT("StateTree Active States should be in Root/State1/State3/State3B"), Exec.ExpectInActiveStates(Root.Name, State1.Name, State3.Name, State3B.Name));
+		AITEST_TRUE(TEXT("StateTree Task3B should enter state"), Exec.Expect(Task3B.GetName(), EnterStateStr));
+		AITEST_TRUE(TEXT("StateTree should be running"), Status == EStateTreeRunStatus::Running);
 		Exec.LogClear();
 
 		// Complete subtree
 		Status = Exec.Tick(0.1f);
-		AITEST_TRUE("StateTree Active States should be in Root/State2", Exec.ExpectInActiveStates(Root.Name, State2.Name));
-		AITEST_TRUE("StateTree Task2 should enter state", Exec.Expect(Task2.GetName(), EnterStateStr));
-		AITEST_TRUE("StateTree should be running", Status == EStateTreeRunStatus::Running);
+		AITEST_TRUE(TEXT("StateTree Active States should be in Root/State2"), Exec.ExpectInActiveStates(Root.Name, State2.Name));
+		AITEST_TRUE(TEXT("StateTree Task2 should enter state"), Exec.Expect(Task2.GetName(), EnterStateStr));
+		AITEST_TRUE(TEXT("StateTree should be running"), Status == EStateTreeRunStatus::Running);
 		Exec.LogClear();
 
 		// Complete the whole tree
 		Status = Exec.Tick(0.1f);
-		AITEST_TRUE("StateTree should complete in succeeded", Status == EStateTreeRunStatus::Succeeded);
+		AITEST_TRUE(TEXT("StateTree should complete in succeeded"), Status == EStateTreeRunStatus::Succeeded);
 		Exec.LogClear();
 
 		Exec.Stop();
@@ -483,39 +483,39 @@ struct FStateTreeTest_SubTree_NoActiveTasks : FStateTreeTestBase
 		FStateTreeCompilerLog Log;
 		FStateTreeCompiler Compiler(Log);
 		const bool bResult = Compiler.Compile(StateTree);
-		AITEST_TRUE("StateTree should get compiled", bResult);
+		AITEST_TRUE(TEXT("StateTree should get compiled"), bResult);
 
 		EStateTreeRunStatus Status = EStateTreeRunStatus::Unset;
 		FStateTreeInstanceData InstanceData;
 		FTestStateTreeExecutionContext Exec(StateTree, StateTree, InstanceData);
 		const bool bInitSucceeded = Exec.IsValid();
-		AITEST_TRUE("StateTree should init", bInitSucceeded);
+		AITEST_TRUE(TEXT("StateTree should init"), bInitSucceeded);
 
 		// Start and enter state
 		Status = Exec.Start();
 
-		AITEST_TRUE("StateTree Active States should be in RootA/StateA/SubTree/StateC", 
+		AITEST_TRUE(TEXT("StateTree Active States should be in RootA/StateA/SubTree/StateC"),
 					Exec.ExpectInActiveStates(RootA.Name, StateA.Name, SubTree.Name, StateC.Name));
-		AITEST_TRUE("StateTree should be running", Status == EStateTreeRunStatus::Running);
+		AITEST_TRUE(TEXT("StateTree should be running"), Status == EStateTreeRunStatus::Running);
 		Exec.LogClear();
 
 		// Transition from the subtree frame. The parent frame and the disabled task should be ignored.
 		Status = Exec.Tick(0.1f);
-		AITEST_TRUE("StateTree Active States should be in RootB", 
+		AITEST_TRUE(TEXT("StateTree Active States should be in RootB"),
 					Exec.ExpectInActiveStates(RootB.Name));
-		AITEST_TRUE("StateTree should be running", Status == EStateTreeRunStatus::Running);
+		AITEST_TRUE(TEXT("StateTree should be running"), Status == EStateTreeRunStatus::Running);
 		Exec.LogClear();
 
 		Status = Exec.Tick(0.1f);
-		AITEST_TRUE("StateTree Active States should be in RootA/StateB",
+		AITEST_TRUE(TEXT("StateTree Active States should be in RootA/StateB"),
 					Exec.ExpectInActiveStates(RootA.Name, StateB.Name));
-		AITEST_TRUE("StateTree should be running", Status == EStateTreeRunStatus::Running);
+		AITEST_TRUE(TEXT("StateTree should be running"), Status == EStateTreeRunStatus::Running);
 		Exec.LogClear();
 
 		Status = Exec.Tick(0.1f);
-		AITEST_TRUE("StateTree Active States should be in RootA/StateA/SubTree/StateC",
+		AITEST_TRUE(TEXT("StateTree Active States should be in RootA/StateA/SubTree/StateC"),
 					Exec.ExpectInActiveStates(RootA.Name, StateA.Name, SubTree.Name, StateC.Name));
-		AITEST_TRUE("StateTree should be running", Status == EStateTreeRunStatus::Running);
+		AITEST_TRUE(TEXT("StateTree should be running"), Status == EStateTreeRunStatus::Running);
 		Exec.LogClear();
 
 		Exec.Stop();
@@ -549,7 +549,7 @@ struct FStateTreeTest_SubTree_Condition : FStateTreeTestBase
 
 		Linked.SetLinkedState(SubTree.GetLinkToState());
 
-		Linked.AddTransition(EStateTreeTransitionTrigger::OnStateCompleted, EStateTreeTransitionType::GotoState, &Root);
+		Linked.AddTransition(EStateTreeTransitionTrigger::OnStateCompleted, EStateTreeTransitionType::GotoState, &Linked);
 
 		// SubTask should not complete during the test.
 		TStateTreeEditorNode<FTestTask_Stand>& SubTask = SubTree.AddTask<FTestTask_Stand>(FName(TEXT("SubTask")));
@@ -575,29 +575,29 @@ struct FStateTreeTest_SubTree_Condition : FStateTreeTestBase
 		FStateTreeInstanceData InstanceData;
 		FTestStateTreeExecutionContext Exec(StateTree, StateTree, InstanceData);
 		const bool bInitSucceeded = Exec.IsValid();
-		AITEST_TRUE("StateTree should init", bInitSucceeded);
+		AITEST_TRUE(TEXT("StateTree should init"), bInitSucceeded);
 
-		const FString TickStr(TEXT("Tick"));
-		const FString EnterStateStr(TEXT("EnterState"));
-		const FString ExitStateStr(TEXT("ExitState"));
-		const FString StateCompletedStr(TEXT("StateCompleted"));
+		const TCHAR* TickStr(TEXT("Tick"));
+		const TCHAR* EnterStateStr(TEXT("EnterState"));
+		const TCHAR* ExitStateStr(TEXT("ExitState"));
+		const TCHAR* StateCompletedStr(TEXT("StateCompleted"));
 
 		// Start and enter state
 		Status = Exec.Start();
 
-		AITEST_TRUE("StateTree Active States should be in Root/Linked/SubTree/State2", Exec.ExpectInActiveStates(Root.Name, Linked.Name, SubTree.Name, State2.Name));
-		AITEST_FALSE("StateTree State1 should not be active", Exec.ExpectInActiveStates(State1.Name)); // Enter condition should prevent to enter State1
-		AITEST_TRUE("StateTree SubTask should enter state", Exec.Expect(SubTask.GetName(), EnterStateStr));
-		AITEST_TRUE("StateTree Task2 should enter state", Exec.Expect(Task2.GetName(), EnterStateStr));
-		AITEST_TRUE("StateTree should be running", Status == EStateTreeRunStatus::Running);
+		AITEST_TRUE(TEXT("StateTree Active States should be in Root/Linked/SubTree/State2"), Exec.ExpectInActiveStates(Root.Name, Linked.Name, SubTree.Name, State2.Name));
+		AITEST_FALSE(TEXT("StateTree State1 should not be active"), Exec.ExpectInActiveStates(State1.Name)); // Enter condition should prevent to enter State1
+		AITEST_TRUE(TEXT("StateTree SubTask should enter state"), Exec.Expect(SubTask.GetName(), EnterStateStr));
+		AITEST_TRUE(TEXT("StateTree Task2 should enter state"), Exec.Expect(Task2.GetName(), EnterStateStr));
+		AITEST_TRUE(TEXT("StateTree should be running"), Status == EStateTreeRunStatus::Running);
 		Exec.LogClear();
 
 		// Task1 completes, and we should enter State1 since the enter condition now passes.
 		Status = Exec.Tick(0.1f);
-		AITEST_TRUE("StateTree Active States should be in Root/Linked/SubTree/State1", Exec.ExpectInActiveStates(Root.Name, Linked.Name, SubTree.Name, State1.Name));
-		AITEST_FALSE("StateTree State2 should not be active", Exec.ExpectInActiveStates(State2.Name));
-		AITEST_TRUE("StateTree Task1 should enter state", Exec.Expect(Task1.GetName(), EnterStateStr));
-		AITEST_TRUE("StateTree should be running", Status == EStateTreeRunStatus::Running);
+		AITEST_TRUE(TEXT("StateTree Active States should be in Root/Linked/SubTree/State1"), Exec.ExpectInActiveStates(Root.Name, Linked.Name, SubTree.Name, State1.Name));
+		AITEST_FALSE(TEXT("StateTree State2 should not be active"), Exec.ExpectInActiveStates(State2.Name));
+		AITEST_TRUE(TEXT("StateTree Task1 should enter state"), Exec.Expect(Task1.GetName(), EnterStateStr));
+		AITEST_TRUE(TEXT("StateTree should be running"), Status == EStateTreeRunStatus::Running);
 		Exec.LogClear();
 
 		Exec.Stop();
@@ -648,30 +648,30 @@ struct FStateTreeTest_SubTree_CascadedSucceeded : FStateTreeTestBase
 		FStateTreeCompilerLog Log;
 		FStateTreeCompiler Compiler(Log);
 		const bool bResult = Compiler.Compile(StateTree);
-		AITEST_TRUE("StateTree should get compiled", bResult);
+		AITEST_TRUE(TEXT("StateTree should get compiled"), bResult);
 
 		EStateTreeRunStatus Status = EStateTreeRunStatus::Unset;
 		FStateTreeInstanceData InstanceData;
 		FTestStateTreeExecutionContext Exec(StateTree, StateTree, InstanceData);
 		const bool bInitSucceeded = Exec.IsValid();
-		AITEST_TRUE("StateTree should init", bInitSucceeded);
+		AITEST_TRUE(TEXT("StateTree should init"), bInitSucceeded);
 
-		const FString TickStr(TEXT("Tick"));
-		const FString EnterStateStr(TEXT("EnterState"));
-		const FString ExitStateStr(TEXT("ExitState"));
-		const FString StateCompletedStr(TEXT("StateCompleted"));
+		const TCHAR* TickStr(TEXT("Tick"));
+		const TCHAR* EnterStateStr(TEXT("EnterState"));
+		const TCHAR* ExitStateStr(TEXT("ExitState"));
+		const TCHAR* StateCompletedStr(TEXT("StateCompleted"));
 
 		// Start and enter state
 		Status = Exec.Start();
-		AITEST_TRUE("StateTree Active States should be in Root/Linked/SubTreeState", Exec.ExpectInActiveStates(Root.Name, LinkedState.Name, SubTreeState.Name, SubLinkedState.Name, SubSubTreeState.Name, SubSubLeaf.Name));
-		AITEST_TRUE("TaskA,B,C should enter state", Exec.Expect(TaskA.GetName(), EnterStateStr).Then(TaskB.GetName(), EnterStateStr).Then(TaskC.GetName(), EnterStateStr));
-		AITEST_TRUE("StateTree should be running", Status == EStateTreeRunStatus::Running);
+		AITEST_TRUE(TEXT("StateTree Active States should be in Root/Linked/SubTreeState"), Exec.ExpectInActiveStates(Root.Name, LinkedState.Name, SubTreeState.Name, SubLinkedState.Name, SubSubTreeState.Name, SubSubLeaf.Name));
+		AITEST_TRUE(TEXT("TaskA,B,C should enter state"), Exec.Expect(TaskA.GetName(), EnterStateStr).Then(TaskB.GetName(), EnterStateStr).Then(TaskC.GetName(), EnterStateStr));
+		AITEST_TRUE(TEXT("StateTree should be running"), Status == EStateTreeRunStatus::Running);
 		Exec.LogClear();
 
 		// Subtrees completes, and it completes the whole tree too.
 		// There's no good way to observe this externally. We switch the return along the way to make sure the transition does not happen directly from the leaf to failed.
 		Status = Exec.Tick(0.1f);
-		AITEST_TRUE("StateTree should be Failed", Status == EStateTreeRunStatus::Failed);
+		AITEST_TRUE(TEXT("StateTree should be Failed"), Status == EStateTreeRunStatus::Failed);
 		Exec.LogClear();
 
 		Exec.Stop();
@@ -699,7 +699,7 @@ struct FStateTreeTest_SharedInstanceData : FStateTreeTestBase
 		FStateTreeCompilerLog Log;
 		FStateTreeCompiler Compiler(Log);
 		const bool bResult = Compiler.Compile(StateTree);
-		AITEST_TRUE("StateTree should get compiled", bResult);
+		AITEST_TRUE(TEXT("StateTree should get compiled"), bResult);
 
 		// Init, nothing should access the shared data.
 		constexpr int32 NumConcurrent = 100;
@@ -714,8 +714,8 @@ struct FStateTreeTest_SharedInstanceData : FStateTreeTestBase
 			FTestStateTreeExecutionContext Exec(StateTree, StateTree, InstanceDatas[Index]);
 			bInitSucceeded &= Exec.IsValid();
 		}
-		AITEST_TRUE("All StateTree contexts should init", bInitSucceeded);
-		AITEST_EQUAL("Test condition global counter should be 0", (int32)FStateTreeTestConditionInstanceData::GlobalCounter, 0);
+		AITEST_TRUE(TEXT("All StateTree contexts should init"), bInitSucceeded);
+		AITEST_EQUAL(TEXT("Test condition global counter should be 0"), (int32)FStateTreeTestConditionInstanceData::GlobalCounter, 0);
 		
 		// Start in parallel
 		// This should create shared data per thread.
@@ -741,8 +741,8 @@ struct FStateTreeTest_SharedInstanceData : FStateTreeTestBase
 		{
 			StartTotalRunning += RunContext.Count;
 		}
-		AITEST_EQUAL("All StateTree contexts should be running after Start", StartTotalRunning, NumConcurrent);
-		AITEST_EQUAL("Test condition global counter should equal context count after Start", (int32)FStateTreeTestConditionInstanceData::GlobalCounter, InstanceDatas.Num());
+		AITEST_EQUAL(TEXT("All StateTree contexts should be running after Start"), StartTotalRunning, NumConcurrent);
+		AITEST_EQUAL(TEXT("Test condition global counter should equal context count after Start"), (int32)FStateTreeTestConditionInstanceData::GlobalCounter, InstanceDatas.Num());
 
 		
 		// Tick in parallel
@@ -771,8 +771,8 @@ struct FStateTreeTest_SharedInstanceData : FStateTreeTestBase
 		{
 			TickTotalRunning += RunContext.Count;
 		}
-		AITEST_EQUAL("All StateTree contexts should be running after Tick", TickTotalRunning, NumConcurrent);
-		AITEST_EQUAL("Test condition global counter should equal context count after Tick", (int32)FStateTreeTestConditionInstanceData::GlobalCounter, InstanceDatas.Num());
+		AITEST_EQUAL(TEXT("All StateTree contexts should be running after Tick"), TickTotalRunning, NumConcurrent);
+		AITEST_EQUAL(TEXT("Test condition global counter should equal context count after Tick"), (int32)FStateTreeTestConditionInstanceData::GlobalCounter, InstanceDatas.Num());
 
 		ParallelForWithTaskContext(
 			RunContexts,
@@ -811,30 +811,30 @@ struct FStateTreeTest_LastConditionWithIndent : FStateTreeTestBase
 		FStateTreeCompilerLog Log;
 		FStateTreeCompiler Compiler(Log);
 		const bool bResult = Compiler.Compile(StateTree);
-		AITEST_TRUE("StateTree should get compiled", bResult);
+		AITEST_TRUE(TEXT("StateTree should get compiled"), bResult);
 
 		EStateTreeRunStatus Status = EStateTreeRunStatus::Unset;
 		FStateTreeInstanceData InstanceData;
 		FTestStateTreeExecutionContext Exec(StateTree, StateTree, InstanceData);
 		const bool bInitSucceeded = Exec.IsValid();
-		AITEST_TRUE("StateTree should init", bInitSucceeded);
+		AITEST_TRUE(TEXT("StateTree should init"), bInitSucceeded);
 
-		const FString TickStr(TEXT("Tick"));
-		const FString EnterStateStr(TEXT("EnterState"));
-		const FString ExitStateStr(TEXT("ExitState"));
+		const TCHAR* TickStr(TEXT("Tick"));
+		const TCHAR* EnterStateStr(TEXT("EnterState"));
+		const TCHAR* ExitStateStr(TEXT("ExitState"));
 
 		Status = Exec.Start();
-		AITEST_TRUE("StateTree Task1 should enter state", Exec.Expect(Task1.GetName(), EnterStateStr));
-		AITEST_FALSE("StateTree Task1 should not tick", Exec.Expect(Task1.GetName(), TickStr));
+		AITEST_TRUE(TEXT("StateTree Task1 should enter state"), Exec.Expect(Task1.GetName(), EnterStateStr));
+		AITEST_FALSE(TEXT("StateTree Task1 should not tick"), Exec.Expect(Task1.GetName(), TickStr));
 		Exec.LogClear();
 
 		Status = Exec.Tick(0.1f);
-		AITEST_TRUE("StateTree Task1 should tick, and exit state", Exec.Expect(Task1.GetName(), TickStr).Then(Task1.GetName(), ExitStateStr));
-		AITEST_TRUE("StateTree should be completed", Status == EStateTreeRunStatus::Succeeded);
+		AITEST_TRUE(TEXT("StateTree Task1 should tick, and exit state"), Exec.Expect(Task1.GetName(), TickStr).Then(Task1.GetName(), ExitStateStr));
+		AITEST_TRUE(TEXT("StateTree should be completed"), Status == EStateTreeRunStatus::Succeeded);
 		Exec.LogClear();
 
 		Status = Exec.Tick(0.1f);
-		AITEST_FALSE("StateTree Task1 should not tick", Exec.Expect(Task1.GetName(), TickStr));
+		AITEST_FALSE(TEXT("StateTree Task1 should not tick"), Exec.Expect(Task1.GetName(), TickStr));
 		Exec.LogClear();
 
 		Exec.Stop();
@@ -901,22 +901,22 @@ struct FStateTreeTest_StateRequiringEvent : FStateTreeTestBase
 		FStateTreeCompiler Compiler(Log);
 		const bool bResult = Compiler.Compile(StateTree);
 
-		AITEST_TRUE("StateTree should get compiled", bResult);
+		AITEST_TRUE(TEXT("StateTree should get compiled"), bResult);
 
 		EStateTreeRunStatus Status = EStateTreeRunStatus::Unset;
 		FStateTreeInstanceData InstanceData;
 
-		const FString EnterStateStr(TEXT("EnterState"));
+		const TCHAR* EnterStateStr(TEXT("EnterState"));
 
 		auto Test=[&](FTestStateTreeExecutionContext& Exec)
-		{
-			AITEST_FALSE("StateTree TaskA should not enter state", Exec.Expect(TaskA.GetName(), EnterStateStr));
-			AITEST_FALSE("StateTree TaskB should not enter state", Exec.Expect(TaskB.GetName(), EnterStateStr));
-			AITEST_FALSE("StateTree TaskC should not enter state", Exec.Expect(TaskC.GetName(), EnterStateStr));
-			AITEST_FALSE("StateTree TaskD should not enter state", Exec.Expect(TaskD.GetName(), EnterStateStr));
-			AITEST_TRUE("StateTree TaskE should enter state", Exec.Expect(TaskE.GetName(), EnterStateStr));
-			return true;
-		};
+			{
+				AITEST_FALSE(TEXT("StateTree TaskA should not enter state"), Exec.Expect(TaskA.GetName(), EnterStateStr));
+				AITEST_FALSE(TEXT("StateTree TaskB should not enter state"), Exec.Expect(TaskB.GetName(), EnterStateStr));
+				AITEST_FALSE(TEXT("StateTree TaskC should not enter state"), Exec.Expect(TaskC.GetName(), EnterStateStr));
+				AITEST_FALSE(TEXT("StateTree TaskD should not enter state"), Exec.Expect(TaskD.GetName(), EnterStateStr));
+				AITEST_TRUE(TEXT("StateTree TaskE should enter state"), Exec.Expect(TaskE.GetName(), EnterStateStr));
+				return true;
+			};
 
 		{
 			FTestStateTreeExecutionContext Exec(StateTree, StateTree, InstanceData);
@@ -924,7 +924,7 @@ struct FStateTreeTest_StateRequiringEvent : FStateTreeTestBase
 			AITEST_TRUE("StateTree should init", bInitSucceeded);
 
 			Status = Exec.Start();
-			AITEST_TRUE("StateTree TaskInitial should enter state", Exec.Expect(TaskInitial.GetName(), EnterStateStr));
+			AITEST_TRUE(TEXT("StateTree TaskInitial should enter state"), Exec.Expect(TaskInitial.GetName(), EnterStateStr));
 			Exec.LogClear();
 
 			Exec.SendEvent(ValidTag, FConstStructView::Make(FValidPayload()));
@@ -942,10 +942,10 @@ struct FStateTreeTest_StateRequiringEvent : FStateTreeTestBase
 		{
 			FTestStateTreeExecutionContext Exec(StateTree, StateTree, InstanceData);
 			const bool bInitSucceeded = Exec.IsValid();
-			AITEST_TRUE("StateTree should init", bInitSucceeded);
+			AITEST_TRUE(TEXT("StateTree should init"), bInitSucceeded);
 
 			Status = Exec.Start();
-			AITEST_TRUE("StateTree TaskInitial should enter state", Exec.Expect(TaskInitial.GetName(), EnterStateStr));
+			AITEST_TRUE(TEXT("StateTree TaskInitial should enter state"), Exec.Expect(TaskInitial.GetName(), EnterStateStr));
 			Exec.LogClear();
 
 			FStateTreeWeakExecutionContext WeakExec = Exec.MakeWeakExecutionContext();
@@ -966,10 +966,10 @@ struct FStateTreeTest_StateRequiringEvent : FStateTreeTestBase
 			{
 				FTestStateTreeExecutionContext Exec(StateTree, StateTree, InstanceData);
 				const bool bInitSucceeded = Exec.IsValid();
-				AITEST_TRUE("StateTree should init", bInitSucceeded);
+				AITEST_TRUE(TEXT("StateTree should init"), bInitSucceeded);
 
 				Status = Exec.Start();
-				AITEST_TRUE("StateTree TaskInitial should enter state", Exec.Expect(TaskInitial.GetName(), EnterStateStr));
+				AITEST_TRUE(TEXT("StateTree TaskInitial should enter state"), Exec.Expect(TaskInitial.GetName(), EnterStateStr));
 				Exec.LogClear();
 
 				WeakExec = Exec.MakeWeakExecutionContext();
@@ -1034,28 +1034,28 @@ struct FStateTreeTest_Stop : FStateTreeTestBase
 		FStateTreeCompiler Compiler(Log);
 		const bool bResult = Compiler.Compile(StateTree);
 
-		AITEST_TRUE("StateTree should get compiled", bResult);
+		AITEST_TRUE(TEXT("StateTree should get compiled"), bResult);
 
 		EStateTreeRunStatus Status = EStateTreeRunStatus::Unset;
 		FStateTreeInstanceData InstanceData;
 		FTestStateTreeExecutionContext Exec(StateTree, StateTree, InstanceData);
 		const bool bInitSucceeded = Exec.IsValid();
-		AITEST_TRUE("StateTree should init", bInitSucceeded);
+		AITEST_TRUE(TEXT("StateTree should init"), bInitSucceeded);
 
-		const FString TickStr(TEXT("Tick"));
-		const FString EnterStateStr(TEXT("EnterState"));
-		const FString ExitStateStr(TEXT("ExitState"));
+		const TCHAR* TickStr(TEXT("Tick"));
+		const TCHAR* EnterStateStr(TEXT("EnterState"));
+		const TCHAR* ExitStateStr(TEXT("ExitState"));
 
 		Status = Exec.Start();
-		AITEST_EQUAL("Start should be running", Status, EStateTreeRunStatus::Running);
-		AITEST_TRUE("StateTree GlobalTask should enter state", Exec.Expect(GlobalTaskName, EnterStateStr));
-		AITEST_TRUE("StateTree TaskA should enter state", Exec.Expect(TaskAName, EnterStateStr));
+		AITEST_EQUAL(TEXT("Start should be running"), Status, EStateTreeRunStatus::Running);
+		AITEST_TRUE(TEXT("StateTree GlobalTask should enter state"), Exec.Expect(GlobalTaskName, EnterStateStr));
+		AITEST_TRUE(TEXT("StateTree TaskA should enter state"), Exec.Expect(TaskAName, EnterStateStr));
 		Exec.LogClear();
 
 		Status = Exec.Tick(0.1f);
-		AITEST_EQUAL("Tree should end expectedly", Status, ExpectedStatusAfterTick);
-		AITEST_TRUE("StateTree GlobalTask should get exit state expectedly", Exec.Expect(GlobalTaskName, ExpectedExitStatusStr));
-		AITEST_TRUE("StateTree TaskA should get exit state expectedly", Exec.Expect(TaskAName, ExpectedExitStatusStr));
+		AITEST_EQUAL(TEXT("Tree should end expectedly"), Status, ExpectedStatusAfterTick);
+		AITEST_TRUE(TEXT("StateTree GlobalTask should get exit state expectedly"), Exec.Expect(GlobalTaskName, ExpectedExitStatusStr));
+		AITEST_TRUE(TEXT("StateTree TaskA should get exit state expectedly"), Exec.Expect(TaskAName, ExpectedExitStatusStr));
 		Exec.LogClear();
 		
 		Exec.Stop();
@@ -1184,30 +1184,30 @@ struct FStateTreeTest_StopEnterNormal : FStateTreeTest_Stop
 		FStateTreeCompiler Compiler(Log);
 		const bool bResult = Compiler.Compile(StateTree);
 
-		AITEST_TRUE("StateTree should get compiled", bResult);
+		AITEST_TRUE(TEXT("StateTree should get compiled"), bResult);
 
 		EStateTreeRunStatus Status = EStateTreeRunStatus::Unset;
 		FStateTreeInstanceData InstanceData;
 		FTestStateTreeExecutionContext Exec(StateTree, StateTree, InstanceData);
 		const bool bInitSucceeded = Exec.IsValid();
-		AITEST_TRUE("StateTree should init", bInitSucceeded);
+		AITEST_TRUE(TEXT("StateTree should init"), bInitSucceeded);
 
-		const FString TickStr(TEXT("Tick"));
-		const FString EnterStateStr(TEXT("EnterState"));
-		const FString ExitStateStr(TEXT("ExitState"));
+		const TCHAR* TickStr(TEXT("Tick"));
+		const TCHAR* EnterStateStr(TEXT("EnterState"));
+		const TCHAR* ExitStateStr(TEXT("ExitState"));
 
 		// If a normal task fails at start, the last tick status will be failed, but transition handling (and final execution status) will take place next tick. 
 		Status = Exec.Start();
-		AITEST_EQUAL("Tree should be running after start", Status, EStateTreeRunStatus::Running);
-		AITEST_EQUAL("Last execution status should be expected value", Exec.GetLastTickStatus(), ExpectedStatusAfterStart);
+		AITEST_EQUAL(TEXT("Tree should be running after start"), Status, EStateTreeRunStatus::Running);
+		AITEST_EQUAL(TEXT("Last execution status should be expected value"), Exec.GetLastTickStatus(), ExpectedStatusAfterStart);
 
 		// Handles any transitions from failed transition
 		Status = Exec.Tick(0.1f);
-		AITEST_EQUAL("Start should be expected value", Status, ExpectedStatusAfterStart);
-		AITEST_TRUE("StateTree GlobalTask should get exit state expectedly", Exec.Expect(GlobalTaskName, ExpectedExitStatusStr));
+		AITEST_EQUAL(TEXT("Start should be expected value"), Status, ExpectedStatusAfterStart);
+		AITEST_TRUE(TEXT("StateTree GlobalTask should get exit state expectedly"), Exec.Expect(GlobalTaskName, ExpectedExitStatusStr));
 
-		AITEST_TRUE("StateTree TaskA should enter state", Exec.Expect(TaskAName, EnterStateStr));
-		AITEST_TRUE("StateTree TaskA should report exit status", Exec.Expect(TaskAName, ExpectedExitStatusStr));
+		AITEST_TRUE(TEXT("StateTree TaskA should enter state"), Exec.Expect(TaskAName, EnterStateStr));
+		AITEST_TRUE(TEXT("StateTree TaskA should report exit status"), Exec.Expect(TaskAName, ExpectedExitStatusStr));
 
 		Exec.Stop();
 		
@@ -1276,25 +1276,25 @@ struct FStateTreeTest_StopEnterGlobal : FStateTreeTest_Stop
 		FStateTreeCompiler Compiler(Log);
 		const bool bResult = Compiler.Compile(StateTree);
 
-		AITEST_TRUE("StateTree should get compiled", bResult);
+		AITEST_TRUE(TEXT("StateTree should get compiled"), bResult);
 
 		EStateTreeRunStatus Status = EStateTreeRunStatus::Unset;
 		FStateTreeInstanceData InstanceData;
 		FTestStateTreeExecutionContext Exec(StateTree, StateTree, InstanceData);
 		const bool bInitSucceeded = Exec.IsValid();
-		AITEST_TRUE("StateTree should init", bInitSucceeded);
+		AITEST_TRUE(TEXT("StateTree should init"), bInitSucceeded);
 
-		const FString TickStr(TEXT("Tick"));
-		const FString EnterStateStr(TEXT("EnterState"));
-		const FString ExitStateStr(TEXT("ExitState"));
+		const TCHAR* TickStr(TEXT("Tick"));
+		const TCHAR* EnterStateStr(TEXT("EnterState"));
+		const TCHAR* ExitStateStr(TEXT("ExitState"));
 
 		Status = Exec.Start();
-		AITEST_EQUAL("Start should be expected value", Status, ExpectedStatusAfterStart);
-		AITEST_TRUE("StateTree GlobalTask should get exit state expectedly", Exec.Expect(GlobalTaskName, ExpectedExitStatusStr));
+		AITEST_EQUAL(TEXT("Start should be expected value"), Status, ExpectedStatusAfterStart);
+		AITEST_TRUE(TEXT("StateTree GlobalTask should get exit state expectedly"), Exec.Expect(GlobalTaskName, ExpectedExitStatusStr));
 
 		// Normal tasks should not run
-		AITEST_FALSE("StateTree TaskA should not enter state", Exec.Expect(TaskAName, EnterStateStr));
-		AITEST_FALSE("StateTree TaskA should not report exit status", Exec.Expect(TaskAName, ExpectedExitStatusStr));
+		AITEST_FALSE(TEXT("StateTree TaskA should not enter state"), Exec.Expect(TaskAName, EnterStateStr));
+		AITEST_FALSE(TEXT("StateTree TaskA should not report exit status"), Exec.Expect(TaskAName, ExpectedExitStatusStr));
 		Exec.LogClear();
 
 		Exec.Stop();
@@ -1378,34 +1378,34 @@ struct FStateTreeTest_Stop_ExternalStop : FStateTreeTest_Stop
 		FStateTreeCompiler Compiler(Log);
 		const bool bResult = Compiler.Compile(StateTree);
 
-		AITEST_TRUE("StateTree should get compiled", bResult);
+		AITEST_TRUE(TEXT("StateTree should get compiled"), bResult);
 
 		EStateTreeRunStatus Status = EStateTreeRunStatus::Unset;
 		FStateTreeInstanceData InstanceData;
 		FTestStateTreeExecutionContext Exec(StateTree, StateTree, InstanceData);
 		const bool bInitSucceeded = Exec.IsValid();
-		AITEST_TRUE("StateTree should init", bInitSucceeded);
+		AITEST_TRUE(TEXT("StateTree should init"), bInitSucceeded);
 
-		const FString TickStr(TEXT("Tick"));
-		const FString EnterStateStr(TEXT("EnterState"));
-		const FString ExitStateStr(TEXT("ExitState"));
+		const TCHAR* TickStr(TEXT("Tick"));
+		const TCHAR* EnterStateStr(TEXT("EnterState"));
+		const TCHAR* ExitStateStr(TEXT("ExitState"));
 
 		Status = Exec.Start();
-		AITEST_EQUAL("Start should be running", Status, EStateTreeRunStatus::Running);
-		AITEST_TRUE("StateTree GlobalTask should enter state", Exec.Expect(GlobalTaskName, EnterStateStr));
-		AITEST_TRUE("StateTree TaskA should enter state", Exec.Expect(TaskAName, EnterStateStr));
+		AITEST_EQUAL(TEXT("Start should be running"), Status, EStateTreeRunStatus::Running);
+		AITEST_TRUE(TEXT("StateTree GlobalTask should enter state"), Exec.Expect(GlobalTaskName, EnterStateStr));
+		AITEST_TRUE(TEXT("StateTree TaskA should enter state"), Exec.Expect(TaskAName, EnterStateStr));
 		Exec.LogClear();
 
 		Status = Exec.Tick(0.1f);
-		AITEST_EQUAL("Tree should end expectedly", Status, ExpectedStatusAfterTick);
+		AITEST_EQUAL(TEXT("Tree should end expectedly"), Status, ExpectedStatusAfterTick);
 		Exec.LogClear();
 
 		Status = Exec.Stop(EStateTreeRunStatus::Stopped);
-		AITEST_EQUAL("Start should be running", Status, ExpectedStatusAfterStop);
+		AITEST_EQUAL(TEXT("Start should be running"), Status, ExpectedStatusAfterStop);
 		if (!ExpectedExitStatusStr.IsEmpty())
 		{
-			AITEST_TRUE("StateTree GlobalTask should get exit state expectedly", Exec.Expect(GlobalTaskName, ExpectedExitStatusStr));
-			AITEST_TRUE("StateTree TaskA should get exit state expectedly", Exec.Expect(TaskAName, ExpectedExitStatusStr));
+			AITEST_TRUE(TEXT("StateTree GlobalTask should get exit state expectedly"), Exec.Expect(GlobalTaskName, ExpectedExitStatusStr));
+			AITEST_TRUE(TEXT("StateTree TaskA should get exit state expectedly"), Exec.Expect(TaskAName, ExpectedExitStatusStr));
 		}
 		
 		return true;
@@ -1475,12 +1475,12 @@ struct FStateTreeTest_DeferredStop : FStateTreeTestBase
 		FStateTreeCompiler Compiler(Log);
 		const bool bResult = Compiler.Compile(StateTree);
 
-		AITEST_TRUE("StateTree should get compiled", bResult);
+		AITEST_TRUE(TEXT("StateTree should get compiled"), bResult);
 
 		FStateTreeInstanceData InstanceData;
 		FTestStateTreeExecutionContext Exec(StateTree, StateTree, InstanceData);
 		const bool bInitSucceeded = Exec.IsValid();
-		AITEST_TRUE("StateTree should init", bInitSucceeded);
+		AITEST_TRUE(TEXT("StateTree should init"), bInitSucceeded);
 
 		return RunDerivedTest(Exec);
 	}
@@ -1499,7 +1499,7 @@ struct FStateTreeTest_DeferredStop_EnterGlobalTask : FStateTreeTest_DeferredStop
 		EStateTreeRunStatus Status = EStateTreeRunStatus::Unset;
 
 		Status = Exec.Start();
-		AITEST_EQUAL("Tree should be stopped", Status, EStateTreeRunStatus::Stopped);
+		AITEST_EQUAL(TEXT("Tree should be stopped"), Status, EStateTreeRunStatus::Stopped);
 
 		Exec.Stop();
 
@@ -1516,10 +1516,10 @@ struct FStateTreeTest_DeferredStop_TickGlobalTask : FStateTreeTest_DeferredStop
 		EStateTreeRunStatus Status = EStateTreeRunStatus::Unset;
 
 		Status = Exec.Start();
-		AITEST_EQUAL("Tree should be running", Status, EStateTreeRunStatus::Running);
+		AITEST_EQUAL(TEXT("Tree should be running"), Status, EStateTreeRunStatus::Running);
 
 		Status = Exec.Tick(0.1f);
-		AITEST_EQUAL("Tree should be stopped", Status, EStateTreeRunStatus::Stopped);
+		AITEST_EQUAL(TEXT("Tree should be stopped"), Status, EStateTreeRunStatus::Stopped);
 
 		Exec.Stop();
 
@@ -1536,13 +1536,13 @@ struct FStateTreeTest_DeferredStop_ExitGlobalTask : FStateTreeTest_DeferredStop
 		EStateTreeRunStatus Status = EStateTreeRunStatus::Unset;
 
 		Status = Exec.Start();
-		AITEST_EQUAL("Tree should be running", Status, EStateTreeRunStatus::Running);
+		AITEST_EQUAL(TEXT("Tree should be running"), Status, EStateTreeRunStatus::Running);
 
 		Status = Exec.Tick(0.1f);
-		AITEST_EQUAL("Tree should be running", Status, EStateTreeRunStatus::Running);
+		AITEST_EQUAL(TEXT("Tree should be running"), Status, EStateTreeRunStatus::Running);
 
 		Status = Exec.Stop();
-		AITEST_EQUAL("Tree should be stopped", Status, EStateTreeRunStatus::Stopped);
+		AITEST_EQUAL(TEXT("Tree should be stopped"), Status, EStateTreeRunStatus::Stopped);
 
 		return true;
 	}
@@ -1557,7 +1557,7 @@ struct FStateTreeTest_DeferredStop_EnterTask : FStateTreeTest_DeferredStop
 		EStateTreeRunStatus Status = EStateTreeRunStatus::Unset;
 
 		Status = Exec.Start();
-		AITEST_EQUAL("Tree should be running", Status, EStateTreeRunStatus::Stopped);
+		AITEST_EQUAL(TEXT("Tree should be running"), Status, EStateTreeRunStatus::Stopped);
 
 		Exec.Stop();
 
@@ -1574,10 +1574,10 @@ struct FStateTreeTest_DeferredStop_TickTask : FStateTreeTest_DeferredStop
 		EStateTreeRunStatus Status = EStateTreeRunStatus::Unset;
 
 		Status = Exec.Start();
-		AITEST_EQUAL("Tree should be running", Status, EStateTreeRunStatus::Running);
+		AITEST_EQUAL(TEXT("Tree should be running"), Status, EStateTreeRunStatus::Running);
 
 		Status = Exec.Tick(0.1f);
-		AITEST_EQUAL("Tree should be stopped", Status, EStateTreeRunStatus::Stopped);
+		AITEST_EQUAL(TEXT("Tree should be stopped"), Status, EStateTreeRunStatus::Stopped);
 
 		Exec.Stop();
 
@@ -1594,13 +1594,13 @@ struct FStateTreeTest_DeferredStop_ExitTask : FStateTreeTest_DeferredStop
 		EStateTreeRunStatus Status = EStateTreeRunStatus::Unset;
 
 		Status = Exec.Start();
-		AITEST_EQUAL("Tree should be running", Status, EStateTreeRunStatus::Running);
+		AITEST_EQUAL(TEXT("Tree should be running"), Status, EStateTreeRunStatus::Running);
 
 		Status = Exec.Tick(0.1f);
-		AITEST_EQUAL("Tree should be running", Status, EStateTreeRunStatus::Running);
+		AITEST_EQUAL(TEXT("Tree should be running"), Status, EStateTreeRunStatus::Running);
 
 		Status = Exec.Stop();
-		AITEST_EQUAL("Tree should be stopped", Status, EStateTreeRunStatus::Stopped);
+		AITEST_EQUAL(TEXT("Tree should be stopped"), Status, EStateTreeRunStatus::Stopped);
 		
 		return true;
 	}
@@ -1644,32 +1644,31 @@ struct FStateTreeTest_FinishTasks : FStateTreeTestBase
 				FStateTreeCompilerLog Log;
 				FStateTreeCompiler Compiler(Log);
 				const bool bResult = Compiler.Compile(StateTree);
-				AITEST_TRUE("StateTree should get compiled", bResult);
+				AITEST_TRUE(TEXT("StateTree should get compiled"), bResult);
 			}
 			{
 				EStateTreeRunStatus Status = EStateTreeRunStatus::Unset;
 				FStateTreeInstanceData InstanceData;
 				FTestStateTreeExecutionContext Exec(StateTree, StateTree, InstanceData);
 				const bool bInitSucceeded = Exec.IsValid();
-				AITEST_TRUE("StateTree should init", bInitSucceeded);
+				AITEST_TRUE(TEXT("StateTree should init"), bInitSucceeded);
 
 				// Start and enter state
 				Status = Exec.Start();
 
-				AITEST_TRUE("StateTree Active States should be in RootA/StateA",
+				AITEST_TRUE(TEXT("StateTree Active States should be in RootA/StateA"),
 					Exec.ExpectInActiveStates(RootA.Name, StateA.Name));
-				AITEST_TRUE("StateTree should be running", Status == EStateTreeRunStatus::Running);
+				AITEST_TRUE(TEXT("StateTree should be running"), Status == EStateTreeRunStatus::Running);
 				Exec.LogClear();
 
 				// One FinishTask, it should loop to StateA
 				Status = Exec.Tick(0.1f);
-				AITEST_TRUE("StateTree Active States should be in RootA/StateA",
+				AITEST_TRUE(TEXT("StateTree Active States should be in RootA/StateA"),
 					Exec.ExpectInActiveStates(RootA.Name, StateA.Name));
-				AITEST_TRUE("StateTree should be running", Status == EStateTreeRunStatus::Running);
-				AITEST_TRUE("Expect the output tasks", Exec.Expect("StateATaskA", TEXT("Tick101"))
+				AITEST_TRUE(TEXT("StateTree should be running"), Status == EStateTreeRunStatus::Running);
+				AITEST_TRUE(TEXT("Expect the output tasks"), Exec.Expect("StateATaskA", TEXT("Tick101"))
 					.Then("StateATaskA", TEXT("Exitstate102"))
-					//@TODO: instance data should be reset
-					//.Then("StateATaskA", TEXT("EnterState101"))
+					.Then("StateATaskA", TEXT("EnterState101"))
 					);
 				Exec.LogClear();
 
@@ -1690,29 +1689,29 @@ struct FStateTreeTest_FinishTasks : FStateTreeTestBase
 				FStateTreeCompilerLog Log;
 				FStateTreeCompiler Compiler(Log);
 				const bool bResult = Compiler.Compile(StateTree);
-				AITEST_TRUE("StateTree should get compiled", bResult);
+				AITEST_TRUE(TEXT("StateTree should get compiled"), bResult);
 			}
 			{
 				EStateTreeRunStatus Status = EStateTreeRunStatus::Unset;
 				FStateTreeInstanceData InstanceData;
 				FTestStateTreeExecutionContext Exec(StateTree, StateTree, InstanceData);
 				const bool bInitSucceeded = Exec.IsValid();
-				AITEST_TRUE("StateTree should init", bInitSucceeded);
+				AITEST_TRUE(TEXT("StateTree should init"), bInitSucceeded);
 
 				// Start and enter state
 				Status = Exec.Start();
 
-				AITEST_TRUE("StateTree Active States should be in RootA/StateA",
+				AITEST_TRUE(TEXT("StateTree Active States should be in RootA/StateA"),
 					Exec.ExpectInActiveStates(RootA.Name, StateA.Name));
-				AITEST_TRUE("StateTree should be running", Status == EStateTreeRunStatus::Running);
+				AITEST_TRUE(TEXT("StateTree should be running"), Status == EStateTreeRunStatus::Running);
 				Exec.LogClear();
 
 				// One FinishTask, it should loop to StateA
 				Status = Exec.Tick(0.1f);
-				AITEST_TRUE("StateTree Active States should be in RootA/StateA",
+				AITEST_TRUE(TEXT("StateTree Active States should be in RootA/StateA"),
 					Exec.ExpectInActiveStates(RootA.Name, StateA.Name));
-				AITEST_TRUE("StateTree should be running", Status == EStateTreeRunStatus::Running);
-				AITEST_TRUE("Expect the output tasks", Exec.Expect("StateATaskA", TEXT("Tick101"))
+				AITEST_TRUE(TEXT("StateTree should be running"), Status == EStateTreeRunStatus::Running);
+				AITEST_TRUE(TEXT("Expect the output tasks"), Exec.Expect("StateATaskA", TEXT("Tick101"))
 					.Then("StateATaskA", TEXT("Exitstate102"))
 					//@TODO: instance data should be reset
 					//.Then("StateATaskA", TEXT("EnterState101"))
@@ -1741,34 +1740,33 @@ struct FStateTreeTest_FinishTasks : FStateTreeTestBase
 				FStateTreeCompilerLog Log;
 				FStateTreeCompiler Compiler(Log);
 				const bool bResult = Compiler.Compile(StateTree);
-				AITEST_TRUE("StateTree should get compiled", bResult);
+				AITEST_TRUE(TEXT("StateTree should get compiled"), bResult);
 			}
 			{
 				EStateTreeRunStatus Status = EStateTreeRunStatus::Unset;
 				FStateTreeInstanceData InstanceData;
 				FTestStateTreeExecutionContext Exec(StateTree, StateTree, InstanceData);
 				const bool bInitSucceeded = Exec.IsValid();
-				AITEST_TRUE("StateTree should init", bInitSucceeded);
+				AITEST_TRUE(TEXT("StateTree should init"), bInitSucceeded);
 
 				// Start and enter state
 				Status = Exec.Start();
 
-				AITEST_TRUE("StateTree Active States should be in RootA/StateA",
+				AITEST_TRUE(TEXT("StateTree Active States should be in RootA/StateA"),
 					Exec.ExpectInActiveStates(RootA.Name, StateA.Name));
-				AITEST_TRUE("StateTree should be running", Status == EStateTreeRunStatus::Running);
+				AITEST_TRUE(TEXT("StateTree should be running"), Status == EStateTreeRunStatus::Running);
 				Exec.LogClear();
 
 				// One FinishTask in exit but should not close StateA again. It should loop to StateA
 				Status = Exec.Tick(0.1f);
-				AITEST_TRUE("StateTree Active States should be in RootA/StateA",
+				AITEST_TRUE(TEXT("StateTree Active States should be in RootA/StateA"),
 					Exec.ExpectInActiveStates(RootA.Name, StateA.Name));
-				AITEST_TRUE("StateTree should be running", Status == EStateTreeRunStatus::Running);
-				AITEST_TRUE("Expect the output tasks", Exec.Expect("StateATaskA", TEXT("Tick101"))
+				AITEST_TRUE(TEXT("StateTree should be running"), Status == EStateTreeRunStatus::Running);
+				AITEST_TRUE(TEXT("Expect the output tasks"), Exec.Expect("StateATaskA", TEXT("Tick101"))
 					.Then("StateATaskA", TEXT("ExitState102"))
-					//@TODO: instance data should be reset
-					//.Then("StateATaskA", TEXT("EnterState101"))
+					.Then("StateATaskA", TEXT("EnterState101"))
 				);
-				AITEST_FALSE("Expect the output tasks", Exec.Expect("StateATaskA", TEXT("ExitState103")));
+				AITEST_FALSE(TEXT("Expect the output tasks"), Exec.Expect("StateATaskA", TEXT("ExitState103")));
 				Exec.LogClear();
 
 				Exec.Stop();
@@ -1814,7 +1812,7 @@ struct FStateTreeTest_NestedOverride : FStateTreeTestBase
 		{
 			FStateTreeCompiler Compiler3(Log);
 			const bool bResult3 = Compiler3.Compile(StateTree3);
-			AITEST_TRUE("StateTree3 should get compiled", bResult3);
+			AITEST_TRUE(TEXT("StateTree3 should get compiled"), bResult3);
 		}
 		// Wrong Asset 4
 		UStateTree* StateTree4 = NewObject<UStateTree>(&GetWorld());;
@@ -1831,7 +1829,7 @@ struct FStateTreeTest_NestedOverride : FStateTreeTestBase
 
 			FStateTreeCompiler Compiler4(Log);
 			const bool bResult4 = Compiler4.Compile(*StateTree4);
-			AITEST_TRUE("StateTree4 should get compiled", bResult4);
+			AITEST_TRUE(TEXT("StateTree4 should get compiled"), bResult4);
 		}
 
 		// Main asset
@@ -1848,11 +1846,11 @@ struct FStateTreeTest_NestedOverride : FStateTreeTestBase
 
 		FStateTreeCompiler Compiler(Log);
 		const bool bResult = Compiler.Compile(StateTree);
-		AITEST_TRUE("StateTree should get compiled", bResult);
+		AITEST_TRUE(TEXT("StateTree should get compiled"), bResult);
 
-		const FString TickStr(TEXT("Tick"));
-		const FString EnterStateStr(TEXT("EnterState"));
-		const FString ExitStateStr(TEXT("ExitState"));
+		const TCHAR* TickStr(TEXT("Tick"));
+		const TCHAR* EnterStateStr(TEXT("EnterState"));
+		const TCHAR* ExitStateStr(TEXT("ExitState"));
 
 		// Without overrides
 		{
@@ -1860,11 +1858,11 @@ struct FStateTreeTest_NestedOverride : FStateTreeTestBase
 			FStateTreeInstanceData InstanceData;
 			FTestStateTreeExecutionContext Exec(StateTree, StateTree, InstanceData);
 			const bool bInitSucceeded = Exec.IsValid();
-			AITEST_TRUE("StateTree should init", bInitSucceeded);
+			AITEST_TRUE(TEXT("StateTree should init"), bInitSucceeded);
 
 			Status = Exec.Start();
-			AITEST_EQUAL("Start should complete with Running", Status, EStateTreeRunStatus::Running);
-			AITEST_TRUE("StateTree should enter TaskRoot2", Exec.Expect(TaskRoot2.GetName(), EnterStateStr));
+			AITEST_EQUAL(TEXT("Start should complete with Running"), Status, EStateTreeRunStatus::Running);
+			AITEST_TRUE(TEXT("StateTree should enter TaskRoot2"), Exec.Expect(TaskRoot2.GetName(), EnterStateStr));
 
 			Exec.Stop();
 		}
@@ -1886,9 +1884,9 @@ struct FStateTreeTest_NestedOverride : FStateTreeTestBase
 			AITEST_TRUE("StateTree should init", bInitSucceeded);
 			
 			Status = Exec.Start();
-			AITEST_EQUAL("Start should complete with Running", Status, EStateTreeRunStatus::Running);
-			AITEST_TRUE("StateTree should enter TaskRoot3", Exec.Expect(TaskRoot3.GetName(), EnterStateStr));
-			AITEST_FALSE("StateTree should not enter TaskRoot2", Exec.Expect(TaskRoot2.GetName(), EnterStateStr));
+			AITEST_EQUAL(TEXT("Start should complete with Running"), Status, EStateTreeRunStatus::Running);
+			AITEST_TRUE(TEXT("StateTree should enter TaskRoot3"), Exec.Expect(TaskRoot3.GetName(), EnterStateStr));
+			AITEST_FALSE(TEXT("StateTree should not enter TaskRoot2"), Exec.Expect(TaskRoot2.GetName(), EnterStateStr));
 
 			Exec.Stop();
 		}
@@ -1908,17 +1906,17 @@ struct FStateTreeTest_NestedOverride : FStateTreeTestBase
 
 			FTestStateTreeExecutionContext Exec(StateTree, StateTree, InstanceData);
 
-			GetTestRunner().AddExpectedMessage("their schemas don't match", ELogVerbosity::Error, EAutomationExpectedMessageFlags::Contains, 1, false);
+			GetTestRunner().AddExpectedMessage(TEXT("their schemas don't match"), ELogVerbosity::Error, EAutomationExpectedMessageFlags::Contains, 1, false);
 			Exec.SetLinkedStateTreeOverrides(Overrides);
-			AITEST_TRUE("Start should complete with Running", GetTestRunner().HasMetExpectedErrors());
+			AITEST_TRUE(TEXT("Start should complete with Running"), GetTestRunner().HasMetExpectedErrors());
 
 			const bool bInitSucceeded = Exec.IsValid();
-			AITEST_TRUE("StateTree should init", bInitSucceeded);
+			AITEST_TRUE(TEXT("StateTree should init"), bInitSucceeded);
 
 			Status = Exec.Start();
-			AITEST_EQUAL("Start should complete with Running", Status, EStateTreeRunStatus::Running);
-			AITEST_TRUE("StateTree should enter TaskRoot2", Exec.Expect(TaskRoot2.GetName(), EnterStateStr));
-			AITEST_FALSE("StateTree should not enter TaskRoot3", Exec.Expect(TaskRoot3.GetName(), EnterStateStr));
+			AITEST_EQUAL(TEXT("Start should complete with Running"), Status, EStateTreeRunStatus::Running);
+			AITEST_TRUE(TEXT("StateTree should enter TaskRoot2"), Exec.Expect(TaskRoot2.GetName(), EnterStateStr));
+			AITEST_FALSE(TEXT("StateTree should not enter TaskRoot3"), Exec.Expect(TaskRoot3.GetName(), EnterStateStr));
 
 			Exec.Stop();
 		}
@@ -1951,13 +1949,13 @@ struct FStateTreeTest_RecursiveParallelTask : FStateTreeTestBase
 			FStateTreeCompilerLog Log;
 			FStateTreeCompiler Compiler(Log);
 			const bool bResult = Compiler.Compile(StateTree1);
-			AITEST_TRUE("StateTreePar should get compiled", bResult);
+			AITEST_TRUE(TEXT("StateTreePar should get compiled"), bResult);
 		}
 		{
 			FStateTreeInstanceData InstanceData;
 			FTestStateTreeExecutionContext Exec(StateTree1, StateTree1, InstanceData);
 			const bool bInitSucceeded = Exec.IsValid();
-			AITEST_TRUE("StateTree should init", bInitSucceeded);
+			AITEST_TRUE(TEXT("StateTree should init"), bInitSucceeded);
 
 			{
 				GetTestRunner().AddExpectedError(TEXT("Trying to start a new parallel tree from the same tree"), EAutomationExpectedErrorFlags::Contains, 1);
@@ -2005,7 +2003,7 @@ struct FStateTreeTest_ParallelEventPriority : FStateTreeTestBase
 		{
 			FStateTreeCompiler Compiler(Log);
 			const bool bResult = Compiler.Compile(StateTreePar);
-			AITEST_TRUE("StateTreePar should get compiled", bResult);
+			AITEST_TRUE(TEXT("StateTreePar should get compiled"), bResult);
 		}
 
 		// Main asset
@@ -2034,12 +2032,12 @@ struct FStateTreeTest_ParallelEventPriority : FStateTreeTestBase
 		{
 			FStateTreeCompiler Compiler(Log);
 			const bool bResult = Compiler.Compile(StateTree);
-			AITEST_TRUE("StateTree should get compiled", bResult);
+			AITEST_TRUE(TEXT("StateTree should get compiled"), bResult);
 		}
 
-		const FString TickStr(TEXT("Tick"));
-		const FString EnterStateStr(TEXT("EnterState"));
-		const FString ExitStateStr(TEXT("ExitState"));
+		const TCHAR* TickStr(TEXT("Tick"));
+		const TCHAR* EnterStateStr(TEXT("EnterState"));
+		const TCHAR* ExitStateStr(TEXT("ExitState"));
 
 		// Run StateTreePar in parallel with the main tree.
 		// Both trees have a transition on same event.
@@ -2048,16 +2046,16 @@ struct FStateTreeTest_ParallelEventPriority : FStateTreeTestBase
 		FStateTreeInstanceData InstanceData;
 		FTestStateTreeExecutionContext Exec(StateTree, StateTree, InstanceData);
 		const bool bInitSucceeded = Exec.IsValid();
-		AITEST_TRUE("StateTree should init", bInitSucceeded);
+		AITEST_TRUE(TEXT("StateTree should init"), bInitSucceeded);
 
 		Status = Exec.Start();
-		AITEST_EQUAL("Start should complete with Running", Status, EStateTreeRunStatus::Running);
-		AITEST_TRUE("StateTree should enter Task1, Task3", Exec.Expect(Task1.GetName(), EnterStateStr).Then(Task3.GetName(), EnterStateStr));
+		AITEST_EQUAL(TEXT("Start should complete with Running"), Status, EStateTreeRunStatus::Running);
+		AITEST_TRUE(TEXT("StateTree should enter Task1, Task3"), Exec.Expect(Task1.GetName(), EnterStateStr).Then(Task3.GetName(), EnterStateStr));
 		Exec.LogClear();
 
 		Status = Exec.Tick(0.1f);
-		AITEST_EQUAL("Tick should complete with Running", Status, EStateTreeRunStatus::Running);
-		AITEST_TRUE("StateTree should tick Task1, Task3", Exec.Expect(Task1.GetName(), TickStr).Then(Task3.GetName(), TickStr));
+		AITEST_EQUAL(TEXT("Tick should complete with Running"), Status, EStateTreeRunStatus::Running);
+		AITEST_TRUE(TEXT("StateTree should tick Task1, Task3"), Exec.Expect(Task1.GetName(), TickStr).Then(Task3.GetName(), TickStr));
 		Exec.LogClear();
 
 		Exec.SendEvent(EventTag);
@@ -2069,26 +2067,26 @@ struct FStateTreeTest_ParallelEventPriority : FStateTreeTestBase
 		{
 			// Main tree should do the transition.
 			Status = Exec.Tick(0.1f);
-			AITEST_EQUAL("Tick should complete with Running", Status, EStateTreeRunStatus::Running);
-			AITEST_TRUE("StateTree should enter Task4", Exec.Expect(Task4.GetName(), EnterStateStr));
+			AITEST_EQUAL(TEXT("Tick should complete with Running"), Status, EStateTreeRunStatus::Running);
+			AITEST_TRUE(TEXT("StateTree should enter Task4"), Exec.Expect(Task4.GetName(), EnterStateStr));
 			Exec.LogClear();
 
 			Status = Exec.Tick(0.1f);
-			AITEST_EQUAL("Tick should complete with Running", Status, EStateTreeRunStatus::Running);
-			AITEST_TRUE("StateTree should tick Task1, Task4", Exec.Expect(Task1.GetName(), TickStr).Then(Task4.GetName(), TickStr));
+			AITEST_EQUAL(TEXT("Tick should complete with Running"), Status, EStateTreeRunStatus::Running);
+			AITEST_TRUE(TEXT("StateTree should tick Task1, Task4"), Exec.Expect(Task1.GetName(), TickStr).Then(Task4.GetName(), TickStr));
 			Exec.LogClear();
 		}
 		else
 		{
 			// Parallel tree should do the transition.
 			Status = Exec.Tick(0.1f);
-			AITEST_EQUAL("Tick should complete with Running", Status, EStateTreeRunStatus::Running);
-			AITEST_TRUE("StateTree should enter Task2", Exec.Expect(Task2.GetName(), EnterStateStr));
+			AITEST_EQUAL(TEXT("Tick should complete with Running"), Status, EStateTreeRunStatus::Running);
+			AITEST_TRUE(TEXT("StateTree should enter Task2"), Exec.Expect(Task2.GetName(), EnterStateStr));
 			Exec.LogClear();
 
 			Status = Exec.Tick(0.1f);
-			AITEST_EQUAL("Tick should complete with Running", Status, EStateTreeRunStatus::Running);
-			AITEST_TRUE("StateTree should tick Task2, Task3", Exec.Expect(Task2.GetName(), TickStr).Then(Task3.GetName(), TickStr));
+			AITEST_EQUAL(TEXT("Tick should complete with Running"), Status, EStateTreeRunStatus::Running);
+			AITEST_TRUE(TEXT("StateTree should tick Task2, Task3"), Exec.Expect(Task2.GetName(), TickStr).Then(Task3.GetName(), TickStr));
 			Exec.LogClear();
 		}
 
@@ -2165,34 +2163,34 @@ struct FStateTreeTest_SubTreeTransition : FStateTreeTestBase
 		FStateTreeCompilerLog Log;
 		FStateTreeCompiler Compiler(Log);
 		const bool bResult = Compiler.Compile(StateTree);
-		AITEST_TRUE("StateTree should get compiled", bResult);
+		AITEST_TRUE(TEXT("StateTree should get compiled"), bResult);
 
 		EStateTreeRunStatus Status = EStateTreeRunStatus::Unset;
 		FStateTreeInstanceData InstanceData;
 		FTestStateTreeExecutionContext Exec(StateTree, StateTree, InstanceData);
 		const bool bInitSucceeded = Exec.IsValid();
-		AITEST_TRUE("StateTree should init", bInitSucceeded);
+		AITEST_TRUE(TEXT("StateTree should init"), bInitSucceeded);
 
-		const FString TickStr(TEXT("Tick"));
-		const FString EnterStateStr(TEXT("EnterState"));
-		const FString ExitStateStr(TEXT("ExitState"));
-		const FString StateCompletedStr(TEXT("StateCompleted"));
+		const TCHAR* TickStr(TEXT("Tick"));
+		const TCHAR* EnterStateStr(TEXT("EnterState"));
+		const TCHAR* ExitStateStr(TEXT("ExitState"));
+		const TCHAR* StateCompletedStr(TEXT("StateCompleted"));
 
 		// Start and enter state
 		Status = Exec.Start();
 
-		AITEST_TRUE("StateTree Active States should be in Root/PreLastStand/BusinessAsUsual", Exec.ExpectInActiveStates(Root.Name, PreLastStand.Name, BusinessAsUsual.Name));
-		AITEST_TRUE("StateTree Task1 should enter state", Exec.Expect(Task1.GetName(), EnterStateStr));
-		AITEST_TRUE("StateTree Task2 should enter state", Exec.Expect(Task2.GetName(), EnterStateStr));
-		AITEST_TRUE("StateTree should be running", Status == EStateTreeRunStatus::Running);
+		AITEST_TRUE(TEXT("StateTree Active States should be in Root/PreLastStand/BusinessAsUsual"), Exec.ExpectInActiveStates(Root.Name, PreLastStand.Name, BusinessAsUsual.Name));
+		AITEST_TRUE(TEXT("StateTree Task1 should enter state"), Exec.Expect(Task1.GetName(), EnterStateStr));
+		AITEST_TRUE(TEXT("StateTree Task2 should enter state"), Exec.Expect(Task2.GetName(), EnterStateStr));
+		AITEST_TRUE(TEXT("StateTree should be running"), Status == EStateTreeRunStatus::Running);
 		Exec.LogClear();
 
 		// Transition to Reinforcements
 		Status = Exec.Tick(0.1f);
-		AITEST_TRUE("StateTree Active States should be in Root/LastStand/Reinforcements/TimeoutChecker/RemainingCount", Exec.ExpectInActiveStates(Root.Name, LastStand.Name, Reinforcements.Name, TimeoutChecker.Name, RemainingCount.Name));
-		AITEST_TRUE("StateTree Task3 should enter state", Exec.Expect(Task3.GetName(), EnterStateStr));
-		AITEST_TRUE("StateTree Task4 should enter state", Exec.Expect(Task4.GetName(), EnterStateStr));
-		AITEST_TRUE("StateTree should be running", Status == EStateTreeRunStatus::Running);
+		AITEST_TRUE(TEXT("StateTree Active States should be in Root/LastStand/Reinforcements/TimeoutChecker/RemainingCount"), Exec.ExpectInActiveStates(Root.Name, LastStand.Name, Reinforcements.Name, TimeoutChecker.Name, RemainingCount.Name));
+		AITEST_TRUE(TEXT("StateTree Task3 should enter state"), Exec.Expect(Task3.GetName(), EnterStateStr));
+		AITEST_TRUE(TEXT("StateTree Task4 should enter state"), Exec.Expect(Task4.GetName(), EnterStateStr));
+		AITEST_TRUE(TEXT("StateTree should be running"), Status == EStateTreeRunStatus::Running);
 		Exec.LogClear();
 
 		Exec.Stop();
diff --git a/Engine/Plugins/Runtime/StateTree/Source/StateTreeTestSuite/Private/StateTreeTestTypes.h b/Engine/Plugins/Runtime/StateTree/Source/StateTreeTestSuite/Private/StateTreeTestTypes.h
--- a/Engine/Plugins/Runtime/StateTree/Source/StateTreeTestSuite/Private/StateTreeTestTypes.h
+++ b/Engine/Plugins/Runtime/StateTree/Source/StateTreeTestSuite/Private/StateTreeTestTypes.h
@@ -90,7 +90,7 @@ struct FTestStateTreeExecutionContext : public FStateTreeExecutionContext
 			return FLogOrder(*Context, NextIndex);
 		}
 
-		FLogOrder Then(const FName& Name, const FString& Message) const
+		FLogOrder Then(const FName& Name, const FStringView Message) const
 		{
 			int32 NextIndex = Index;
 			while (NextIndex < Context->Log.Items.Num())
@@ -120,7 +120,7 @@ struct FTestStateTreeExecutionContext : public FStateTreeExecutionContext
 		return FLogOrder(*this, 0).Then(Name);
 	}
 
-	FLogOrder Expect(const FName& Name, const FString& Message) const
+	FLogOrder Expect(const FName& Name, const FStringView Message) const
 	{
 		return FLogOrder(*this, 0).Then(Name, Message);
 	}
@@ -281,7 +281,7 @@ struct FTestTask_PrintValue : public FStateTreeTaskBase
 	{
 		FStateTreeTestLog& TestLog = Context.GetExternalData(LogHandle);
 		const FInstanceDataType& InstanceData = Context.GetInstanceData(*this);
-		TestLog.Log(Name,  FString::Printf(TEXT("EnterState%d"), InstanceData.Value));
+		TestLog.Log(Name, FString::Printf(TEXT("EnterState%d"), InstanceData.Value));
 		TestLog.Log(Name, FString::Printf(TEXT("EnterState:%s"), *InstanceData.GetArrayDescription()));
 		if (CustomEnterStateFunc)
 		{
@@ -294,7 +294,7 @@ struct FTestTask_PrintValue : public FStateTreeTaskBase
 	{
 		FStateTreeTestLog& TestLog = Context.GetExternalData(LogHandle);
 		const FInstanceDataType& InstanceData = Context.GetInstanceData(*this);
-		TestLog.Log(Name,  FString::Printf(TEXT("ExitState%d"), InstanceData.Value));
+		TestLog.Log(Name, FString::Printf(TEXT("ExitState%d"), InstanceData.Value));
 		TestLog.Log(Name, FString::Printf(TEXT("ExitState:%s"), *InstanceData.GetArrayDescription()));
 		if (CustomExitStateFunc)
 		{
@@ -306,8 +306,8 @@ struct FTestTask_PrintValue : public FStateTreeTaskBase
 	{
 		FStateTreeTestLog& TestLog = Context.GetExternalData(LogHandle);
 		const FInstanceDataType& InstanceData = Context.GetInstanceData(*this);
-		TestLog.Log(Name,  FString::Printf(TEXT("Tick%d"), InstanceData.Value));
-		TestLog.Log(Name,  FString::Printf(TEXT("Tick:%s"), *InstanceData.GetArrayDescription()));
+		TestLog.Log(Name, FString::Printf(TEXT("Tick%d"), InstanceData.Value));
+		TestLog.Log(Name, FString::Printf(TEXT("Tick:%s"), *InstanceData.GetArrayDescription()));
 		if (CustomTickFunc)
 		{
 			CustomTickFunc(Context, this);
@@ -498,7 +498,7 @@ struct FTestTask_Stand : public FStateTreeTaskBase
 	virtual EStateTreeRunStatus EnterState(FStateTreeExecutionContext& Context, const FStateTreeTransitionResult& Transition) const override
 	{
 		FStateTreeTestLog& TestLog = Context.GetExternalData(LogHandle);
-		TestLog.Log(Name, TEXT("EnterState"));
+			TestLog.Log(Name, TEXT("EnterState"));
 
 		FInstanceDataType& InstanceData = Context.GetInstanceData(*this);
 		
@@ -525,9 +525,9 @@ struct FTestTask_Stand : public FStateTreeTaskBase
 		{
 			TestLog.Log(Name, TEXT("ExitStopped"));
 		}
-		
-		TestLog.Log(Name, TEXT("ExitState"));
-	}
+
+			TestLog.Log(Name, TEXT("ExitState"));
+		}
 
 	virtual void StateCompleted(FStateTreeExecutionContext& Context, const EStateTreeRunStatus CompletionStatus, const FStateTreeActiveStates& CompletedActiveStates) const override
 	{
@@ -1046,4 +1046,4 @@ struct FTestTask_CustomFuncOnDelegate : public FStateTreeTaskBase
 	TFunction<void(const FStateTreeWeakExecutionContext&, FStateTreeDelegateListener)> CustomFunc;
 
 	bool bRemoveOnExit = true;
-};
\ No newline at end of file
+};
diff --git a/Engine/Plugins/Runtime/StateTree/Source/StateTreeTestSuite/Private/StateTreeTransitionTest.cpp b/Engine/Plugins/Runtime/StateTree/Source/StateTreeTestSuite/Private/StateTreeTransitionTest.cpp
--- a/Engine/Plugins/Runtime/StateTree/Source/StateTreeTestSuite/Private/StateTreeTransitionTest.cpp
+++ b/Engine/Plugins/Runtime/StateTree/Source/StateTreeTestSuite/Private/StateTreeTransitionTest.cpp
@@ -66,27 +66,27 @@ struct FStateTreeTest_TransitionPriority : FStateTreeTestBase
 		const bool bInitSucceeded = Exec.IsValid();
 		AITEST_TRUE("StateTree should init", bInitSucceeded);
 
-		const FString TickStr(TEXT("Tick"));
-		const FString EnterStateStr(TEXT("EnterState"));
-		const FString ExitStateStr(TEXT("ExitState"));
-		const FString StateCompletedStr(TEXT("StateCompleted"));
+		const TCHAR* TickStr(TEXT("Tick"));
+		const TCHAR* EnterStateStr(TEXT("EnterState"));
+		const TCHAR* ExitStateStr(TEXT("ExitState"));
+		const TCHAR* StateCompletedStr(TEXT("StateCompleted"));
 
 		// Start and enter state
 		Status = Exec.Start();
-		AITEST_TRUE("StateTree Task1 should enter state", Exec.Expect(Task1.GetName(), EnterStateStr));
-		AITEST_TRUE("StateTree Task1A should enter state", Exec.Expect(Task1A.GetName(), EnterStateStr));
+		AITEST_TRUE(TEXT("StateTree Task1 should enter state"), Exec.Expect(Task1.GetName(), EnterStateStr));
+		AITEST_TRUE(TEXT("StateTree Task1A should enter state"), Exec.Expect(Task1A.GetName(), EnterStateStr));
 		Exec.LogClear();
 
 		// Transition from Task1A to Task1B
 		Status = Exec.Tick(0.1f);
-		AITEST_TRUE("StateTree Task1A should complete", Exec.Expect(Task1A.GetName(), StateCompletedStr));
-		AITEST_TRUE("StateTree Task1B should enter state", Exec.Expect(Task1B.GetName(), EnterStateStr));
+		AITEST_TRUE(TEXT("StateTree Task1A should complete"), Exec.Expect(Task1A.GetName(), StateCompletedStr));
+		AITEST_TRUE(TEXT("StateTree Task1B should enter state"), Exec.Expect(Task1B.GetName(), EnterStateStr));
 		Exec.LogClear();
 
 		// Task1 completes, and we should take State1 transition. 
 		Status = Exec.Tick(0.1f);
 		AITEST_TRUE("StateTree Task1 should complete", Exec.Expect(Task1.GetName(), StateCompletedStr));
-		AITEST_EQUAL("Tree execution should stop on success", Status, EStateTreeRunStatus::Succeeded);
+		AITEST_EQUAL(TEXT("Tree execution should stop on success"), Status, EStateTreeRunStatus::Succeeded);
 		Exec.LogClear();
 
 		Exec.Stop();
@@ -138,21 +138,21 @@ struct FStateTreeTest_TransitionPriorityEnterState : FStateTreeTestBase
 		const bool bInitSucceeded = Exec.IsValid();
 		AITEST_TRUE("StateTree should init", bInitSucceeded);
 
-		const FString TickStr(TEXT("Tick"));
-		const FString EnterStateStr(TEXT("EnterState"));
-		const FString ExitStateStr(TEXT("ExitState"));
-		const FString StateCompletedStr(TEXT("StateCompleted"));
+		const TCHAR* TickStr(TEXT("Tick"));
+		const TCHAR* EnterStateStr(TEXT("EnterState"));
+		const TCHAR* ExitStateStr(TEXT("ExitState"));
+		const TCHAR* StateCompletedStr(TEXT("StateCompleted"));
 
 		// Start and enter state
 		Status = Exec.Start();
-		AITEST_TRUE("StateTree Task0 should enter state", Exec.Expect(Task0.GetName(), EnterStateStr));
+		AITEST_TRUE(TEXT("StateTree Task0 should enter state"), Exec.Expect(Task0.GetName(), EnterStateStr));
 		Exec.LogClear();
 
 		// Transition from State0 to State1, it should fail (Task1), and the transition on State1->State2 (and not State1A->State3)
 		Status = Exec.Tick(0.1f);
-		AITEST_TRUE("StateTree Task0 should complete", Exec.Expect(Task0.GetName(), StateCompletedStr));
-		AITEST_TRUE("StateTree Task2 should enter state", Exec.Expect(Task2.GetName(), EnterStateStr));
-		AITEST_FALSE("StateTree Task3 should not enter state", Exec.Expect(Task3.GetName(), EnterStateStr));
+		AITEST_TRUE(TEXT("StateTree Task0 should complete"), Exec.Expect(Task0.GetName(), StateCompletedStr));
+		AITEST_TRUE(TEXT("StateTree Task2 should enter state"), Exec.Expect(Task2.GetName(), EnterStateStr));
+		AITEST_FALSE(TEXT("StateTree Task3 should not enter state"), Exec.Expect(Task3.GetName(), EnterStateStr));
 		Exec.LogClear();
 
 		Exec.Stop();
@@ -198,33 +198,33 @@ struct FStateTreeTest_TransitionNextSelectableState : FStateTreeTestBase
 		FStateTreeCompilerLog Log;
 		FStateTreeCompiler Compiler(Log);
 		const bool bResult = Compiler.Compile(StateTree);
-		AITEST_TRUE("StateTree should get compiled", bResult);
+		AITEST_TRUE(TEXT("StateTree should get compiled"), bResult);
 
 		FStateTreeInstanceData InstanceData;
 		FTestStateTreeExecutionContext Exec(StateTree, StateTree, InstanceData);
 		const bool bInitSucceeded = Exec.IsValid();
-		AITEST_TRUE("StateTree should init", bInitSucceeded);
+		AITEST_TRUE(TEXT("StateTree should init"), bInitSucceeded);
 
-		const FString TickStr(TEXT("Tick"));
-		const FString EnterStateStr(TEXT("EnterState"));
-		const FString ExitStateStr(TEXT("ExitState"));
-		const FString StateCompletedStr(TEXT("StateCompleted"));
+		const TCHAR* TickStr = TEXT("Tick");
+		const TCHAR* EnterStateStr = TEXT("EnterState");
+		const TCHAR* ExitStateStr = TEXT("ExitState");
+		const TCHAR* StateCompletedStr = TEXT("StateCompleted");
 
 		// Start and enter state
 		Exec.Start();
-		AITEST_TRUE("StateTree Task0 should enter state", Exec.Expect(Task0.GetName(), EnterStateStr));
+		AITEST_TRUE(TEXT("StateTree Task0 should enter state"), Exec.Expect(Task0.GetName(), EnterStateStr));
 		Exec.LogClear();
 
-		// Transition from State0 and tries to select State1. It should fail (Task1) and because transition is set to "Next Selectable", it should now select Task 2 and Enter State
+		// Transition from State0 and tries to select State1. It should fail (Task1 and Tasks2) and because transition is set to "Next Selectable", it should now select Task 3 and Enter State
 		Exec.Tick(0.1f);
-		AITEST_TRUE("StateTree Task0 should complete", Exec.Expect(Task0.GetName(), StateCompletedStr));
-		AITEST_FALSE("StateTree Task1 should not enter state", Exec.Expect(Task1.GetName(), EnterStateStr));
-		AITEST_TRUE("StateTree Task2 should enter state", Exec.Expect(Task2.GetName(), EnterStateStr));
+		AITEST_TRUE(TEXT("StateTree Task0 should complete"), Exec.Expect(Task0.GetName(), StateCompletedStr));
+		AITEST_FALSE(TEXT("StateTree Task1 should not enter state"), Exec.Expect(Task1.GetName(), EnterStateStr));
+		AITEST_TRUE(TEXT("StateTree Task2 should enter state"), Exec.Expect(Task2.GetName(), EnterStateStr));
 		Exec.LogClear();
 
 		// Complete Task2
 		Exec.Tick(0.1f);
-		AITEST_TRUE("StateTree Task2 should complete", Exec.Expect(Task2.GetName(), StateCompletedStr));
+		AITEST_TRUE(TEXT("StateTree Task2 should complete"), Exec.Expect(Task2.GetName(), StateCompletedStr));
 		Exec.LogClear();
 
 		Exec.Stop();
@@ -268,21 +268,21 @@ struct FStateTreeTest_TransitionNextWithParentData : FStateTreeTestBase
 		const bool bInitSucceeded = Exec.IsValid();
 		AITEST_TRUE("StateTree should init", bInitSucceeded);
 
-		const FString TickStr(TEXT("Tick"));
-		const FString EnterStateStr(TEXT("EnterState"));
-		const FString ExitStateStr(TEXT("ExitState"));
-		const FString StateCompletedStr(TEXT("StateCompleted"));
+		const TCHAR* TickStr(TEXT("Tick"));
+		const TCHAR* EnterStateStr(TEXT("EnterState"));
+		const TCHAR* ExitStateStr(TEXT("ExitState"));
+		const TCHAR* StateCompletedStr(TEXT("StateCompleted"));
 
 		// Start and enter state
 		Exec.Start();
-		AITEST_TRUE("StateTree Task0 should enter state", Exec.Expect(Task0.GetName(), EnterStateStr));
+		AITEST_TRUE(TEXT("StateTree Task0 should enter state"), Exec.Expect(Task0.GetName(), EnterStateStr));
 		Exec.LogClear();
 
 		// Transition from State0 and tries to select State1.
 		// This tests that data from current shared active states (Root) is available during state selection.
 		Exec.Tick(0.1f);
-		AITEST_TRUE("StateTree Task0 should complete", Exec.Expect(Task0.GetName(), StateCompletedStr));
-		AITEST_TRUE("StateTree Task1A should enter state", Exec.Expect(Task1A.GetName(), EnterStateStr));
+		AITEST_TRUE(TEXT("StateTree Task0 should complete"), Exec.Expect(Task0.GetName(), StateCompletedStr));
+		AITEST_TRUE(TEXT("StateTree Task1A should enter state"), Exec.Expect(Task1A.GetName(), EnterStateStr));
 		Exec.LogClear();
 
 		Exec.Stop();
@@ -323,27 +323,27 @@ struct FStateTreeTest_TransitionGlobalDataView : FStateTreeTestBase
 		FStateTreeCompiler Compiler(Log);
 		const bool bResult = Compiler.Compile(StateTree);
 
-		AITEST_TRUE("StateTree should get compiled", bResult);
+		AITEST_TRUE(TEXT("StateTree should get compiled"), bResult);
 
 		EStateTreeRunStatus Status = EStateTreeRunStatus::Unset;
 		FStateTreeInstanceData InstanceData;
 		FTestStateTreeExecutionContext Exec(StateTree, StateTree, InstanceData);
 		const bool bInitSucceeded = Exec.IsValid();
-		AITEST_TRUE("StateTree should init", bInitSucceeded);
+		AITEST_TRUE(TEXT("StateTree should init"), bInitSucceeded);
 
-		const FString EnterStateStr(TEXT("EnterState"));
-		const FString EnterState42Str(TEXT("EnterState42"));
-		const FString EnterState123Str(TEXT("EnterState123"));
+		const TCHAR* EnterStateStr(TEXT("EnterState"));
+		const TCHAR* EnterState42Str(TEXT("EnterState42"));
+		const TCHAR* EnterState123Str(TEXT("EnterState123"));
 
 		// Start and enter state
 		Status = Exec.Start();
-		AITEST_TRUE("StateTree Task0 should enter state", Exec.Expect(Task0.GetName(), EnterStateStr));
+		AITEST_TRUE(TEXT("StateTree Task0 should enter state"), Exec.Expect(Task0.GetName(), EnterStateStr));
 		Exec.LogClear();
 
 		// Transition from StateA to StateB, Task0 should enter state with evaluator value copied.
 		Status = Exec.Tick(0.1f);
-		AITEST_TRUE("StateTree Task0 should enter state with value 42", Exec.Expect(Task1.GetName(), EnterState42Str));
-		AITEST_TRUE("StateTree Task1 should enter state with value 123", Exec.Expect(Task2.GetName(), EnterState123Str));
+		AITEST_TRUE(TEXT("StateTree Task0 should enter state with value 42"), Exec.Expect(Task1.GetName(), EnterState42Str));
+		AITEST_TRUE(TEXT("StateTree Task1 should enter state with value 123"), Exec.Expect(Task2.GetName(), EnterState123Str));
 		Exec.LogClear();
 
 		Exec.Stop();
@@ -383,50 +383,50 @@ struct FStateTreeTest_TransitionDelay : FStateTreeTestBase
 		FStateTreeCompiler Compiler(Log);
 		const bool bResult = Compiler.Compile(StateTree);
 
-		AITEST_TRUE("StateTree should get compiled", bResult);
+		AITEST_TRUE(TEXT("StateTree should get compiled"), bResult);
 
 		EStateTreeRunStatus Status = EStateTreeRunStatus::Unset;
 		FStateTreeInstanceData InstanceData;
 		FTestStateTreeExecutionContext Exec(StateTree, StateTree, InstanceData);
 		const bool bInitSucceeded = Exec.IsValid();
-		AITEST_TRUE("StateTree should init", bInitSucceeded);
+		AITEST_TRUE(TEXT("StateTree should init"), bInitSucceeded);
 
-		const FString TickStr(TEXT("Tick"));
-		const FString EnterStateStr(TEXT("EnterState"));
-		const FString ExitStateStr(TEXT("ExitState"));
-		const FString StateCompletedStr(TEXT("StateCompleted"));
+		const TCHAR* TickStr(TEXT("Tick"));
+		const TCHAR* EnterStateStr(TEXT("EnterState"));
+		const TCHAR* ExitStateStr(TEXT("ExitState"));
+		const TCHAR* StateCompletedStr(TEXT("StateCompleted"));
 
 		// Start and enter state
 		Status = Exec.Start();
-		AITEST_TRUE("StateTree Task0 should enter state", Exec.Expect(Task0.GetName(), EnterStateStr));
+		AITEST_TRUE(TEXT("StateTree Task0 should enter state"), Exec.Expect(Task0.GetName(), EnterStateStr));
 		Exec.LogClear();
 
 		// This should cause delayed transition.
 		Exec.SendEvent(Tag);
 		
 		Status = Exec.Tick(0.1f);
-		AITEST_TRUE("StateTree Task0 should tick", Exec.Expect(Task0.GetName(), TickStr));
+		AITEST_TRUE(TEXT("StateTree Task0 should tick"), Exec.Expect(Task0.GetName(), TickStr));
 		Exec.LogClear();
 
 		// Should have execution frames
-		AITEST_TRUE("Should have active frames", InstanceData.GetExecutionState()->ActiveFrames.Num() > 0);
+		AITEST_TRUE(TEXT("Should have active frames"), InstanceData.GetExecutionState()->ActiveFrames.Num() > 0);
 
 		// Should have delayed transitions
 		const int32 NumDelayedTransitions0 = InstanceData.GetExecutionState()->DelayedTransitions.Num();
-		AITEST_EQUAL("Should have a delayed transition", NumDelayedTransitions0, 1);
+		AITEST_EQUAL(TEXT("Should have a delayed transition"), NumDelayedTransitions0, 1);
 
 		// Tick and expect a delayed transition. 
 		Status = Exec.Tick(0.1f);
-		AITEST_TRUE("StateTree Task0 should tick", Exec.Expect(Task0.GetName(), TickStr));
+		AITEST_TRUE(TEXT("StateTree Task0 should tick"), Exec.Expect(Task0.GetName(), TickStr));
 		Exec.LogClear();
 
 		const int32 NumDelayedTransitions1 = InstanceData.GetExecutionState()->DelayedTransitions.Num();
-		AITEST_EQUAL("Should have a delayed transition", NumDelayedTransitions1, 1);
+		AITEST_EQUAL(TEXT("Should have a delayed transition"), NumDelayedTransitions1, 1);
 
 		// Should complete delayed transition.
 		Status = Exec.Tick(0.1f);
-		AITEST_TRUE("StateTree Task0 should exit state", Exec.Expect(Task0.GetName(), ExitStateStr));
-		AITEST_TRUE("StateTree Task1 should enter state", Exec.Expect(Task1.GetName(), EnterStateStr));
+		AITEST_TRUE(TEXT("StateTree Task0 should exit state"), Exec.Expect(Task0.GetName(), ExitStateStr));
+		AITEST_TRUE(TEXT("StateTree Task1 should enter state"), Exec.Expect(Task1.GetName(), EnterStateStr));
 		Exec.LogClear();
 
 		Exec.Stop();
@@ -472,24 +472,24 @@ struct FStateTreeTest_TransitionDelayZero : FStateTreeTestBase
 		FStateTreeInstanceData InstanceData;
 		FTestStateTreeExecutionContext Exec(StateTree, StateTree, InstanceData);
 		const bool bInitSucceeded = Exec.IsValid();
-		AITEST_TRUE("StateTree should init", bInitSucceeded);
+		AITEST_TRUE(TEXT("StateTree should init"), bInitSucceeded);
 
-		const FString TickStr(TEXT("Tick"));
-		const FString EnterStateStr(TEXT("EnterState"));
-		const FString ExitStateStr(TEXT("ExitState"));
-		const FString StateCompletedStr(TEXT("StateCompleted"));
+		const TCHAR* TickStr(TEXT("Tick"));
+		const TCHAR* EnterStateStr(TEXT("EnterState"));
+		const TCHAR* ExitStateStr(TEXT("ExitState"));
+		const TCHAR* StateCompletedStr(TEXT("StateCompleted"));
 
 		// Start and enter state
 		Status = Exec.Start();
-		AITEST_TRUE("StateTree Task0 should enter state", Exec.Expect(Task0.GetName(), EnterStateStr));
+		AITEST_TRUE(TEXT("StateTree Task0 should enter state"), Exec.Expect(Task0.GetName(), EnterStateStr));
 		Exec.LogClear();
 
 		// This should cause delayed transition. Because the time is 0, it should happen immediately.
 		Exec.SendEvent(Tag);
 		
 		Status = Exec.Tick(0.1f);
-		AITEST_TRUE("StateTree Task0 should exit state", Exec.Expect(Task0.GetName(), ExitStateStr));
-		AITEST_TRUE("StateTree Task1 should enter state", Exec.Expect(Task1.GetName(), EnterStateStr));
+		AITEST_TRUE(TEXT("StateTree Task0 should exit state"), Exec.Expect(Task0.GetName(), ExitStateStr));
+		AITEST_TRUE(TEXT("StateTree Task1 should enter state"), Exec.Expect(Task1.GetName(), EnterStateStr));
 		Exec.LogClear();
 
 		Exec.Stop();
@@ -512,12 +512,12 @@ struct FStateTreeTest_PassingTransitionEventToStateSelection : FStateTreeTestBas
 		{
 			const bool bParseResult = PathToPayloadMember.FromString(TEXT("Payload.A"));
 
-			AITEST_TRUE("Parsing path should succeeed", bParseResult);
+			AITEST_TRUE(TEXT("Parsing path should succeeed"), bParseResult);
 
 			FStateTreeEvent EventWithPayload;
 			EventWithPayload.Payload = FInstancedStruct::Make<FStateTreeTest_PropertyStructA>();
 			const bool bUpdateSegments = PathToPayloadMember.UpdateSegmentsFromValue(FStateTreeDataView(FStructView::Make(EventWithPayload)));
-			AITEST_TRUE("Updating segments should succeeed", bUpdateSegments);
+			AITEST_TRUE(TEXT("Updating segments should succeeed"), bUpdateSegments);
 		}
 
 		// This state shouldn't be selected, because transition's condition and state's enter condition exlude each other.
@@ -572,18 +572,18 @@ struct FStateTreeTest_PassingTransitionEventToStateSelection : FStateTreeTestBas
 		FStateTreeCompiler Compiler(Log);
 		const bool bResult = Compiler.Compile(StateTree);
 
-		AITEST_TRUE("StateTree should get compiled", bResult);
+		AITEST_TRUE(TEXT("StateTree should get compiled"), bResult);
 
 		EStateTreeRunStatus Status = EStateTreeRunStatus::Unset;
 		FStateTreeInstanceData InstanceData;
 		FTestStateTreeExecutionContext Exec(StateTree, StateTree, InstanceData);
 		const bool bInitSucceeded = Exec.IsValid();
-		AITEST_TRUE("StateTree should init", bInitSucceeded);
+		AITEST_TRUE(TEXT("StateTree should init"), bInitSucceeded);
 
-		const FString EnterStateStr(TEXT("EnterState"));
+		const TCHAR* EnterStateStr(TEXT("EnterState"));
 
 		Status = Exec.Start();
-		AITEST_TRUE("StateTree TaskInitial should enter state", Exec.Expect(TaskInitial.GetName(), EnterStateStr));
+		AITEST_TRUE(TEXT("StateTree TaskInitial should enter state"), Exec.Expect(TaskInitial.GetName(), EnterStateStr));
 		Exec.LogClear();
 
 		// The conditions test for payload Value=1, the first event should not trigger transition. 
@@ -591,8 +591,8 @@ struct FStateTreeTest_PassingTransitionEventToStateSelection : FStateTreeTestBas
 		Exec.SendEvent(GetTestTag1(), FConstStructView::Make(FStateTreeTest_PropertyStructA{1}));
 		Status = Exec.Tick(0.1f);
 
-		AITEST_FALSE("StateTree TaskA should not enter state", Exec.Expect(TaskA.GetName(), EnterStateStr));
-		AITEST_TRUE("StateTree TaskB should enter state", Exec.Expect(TaskB.GetName(), TEXT("EnterState1"))); // TaskB decorates "EnterState" with value from the payload.
+		AITEST_FALSE(TEXT("StateTree TaskA should not enter state"), Exec.Expect(TaskA.GetName(), EnterStateStr));
+		AITEST_TRUE(TEXT("StateTree TaskB should enter state"), Exec.Expect(TaskB.GetName(), TEXT("EnterState1"))); // TaskB decorates "EnterState" with value from the payload.
 		Exec.LogClear();
 
 		Exec.Stop();
@@ -665,22 +665,22 @@ struct FStateTreeTest_FollowTransitions : FStateTreeTestBase
 		FStateTreeCompiler Compiler(Log);
 		const bool bResult = Compiler.Compile(StateTree);
 
-		AITEST_TRUE("StateTree should get compiled", bResult);
+		AITEST_TRUE(TEXT("StateTree should get compiled"), bResult);
 
 		EStateTreeRunStatus Status = EStateTreeRunStatus::Unset;
 		FStateTreeInstanceData InstanceData;
 		FTestStateTreeExecutionContext Exec(StateTree, StateTree, InstanceData);
 		const bool bInitSucceeded = Exec.IsValid();
-		AITEST_TRUE("StateTree should init", bInitSucceeded);
+		AITEST_TRUE(TEXT("StateTree should init"), bInitSucceeded);
 
-		const FString TickStr(TEXT("Tick"));
-		const FString EnterStateStr(TEXT("EnterState"));
-		const FString ExitStateStr(TEXT("ExitState"));
+		const TCHAR* TickStr = TEXT("Tick");
+		const TCHAR* EnterStateStr = TEXT("EnterState");
+		const TCHAR* ExitStateStr = TEXT("ExitState");
 
 		Status = Exec.Start();
-		AITEST_FALSE("StateTree TaskA should not enter state", Exec.Expect(TaskA.GetName(), EnterStateStr));
-		AITEST_FALSE("StateTree TaskB should not enter state", Exec.Expect(TaskB.GetName(), EnterStateStr));
-		AITEST_TRUE("StateTree TaskC should enter state", Exec.Expect(TaskC.GetName(), EnterStateStr));
+		AITEST_FALSE(TEXT("StateTree TaskA should not enter state"), Exec.Expect(TaskA.GetName(), EnterStateStr));
+		AITEST_FALSE(TEXT("StateTree TaskB should not enter state"), Exec.Expect(TaskB.GetName(), EnterStateStr));
+		AITEST_TRUE(TEXT("StateTree TaskC should enter state"), Exec.Expect(TaskC.GetName(), EnterStateStr));
 		Exec.LogClear();
 
 		Exec.Stop();
@@ -741,23 +741,23 @@ struct FStateTreeTest_InfiniteLoop : FStateTreeTestBase
 		FStateTreeCompiler Compiler(Log);
 		const bool bResult = Compiler.Compile(StateTree);
 
-		AITEST_TRUE("StateTree should get compiled", bResult);
+		AITEST_TRUE(TEXT("StateTree should get compiled"), bResult);
 
 		EStateTreeRunStatus Status = EStateTreeRunStatus::Unset;
 		FStateTreeInstanceData InstanceData;
 		FTestStateTreeExecutionContext Exec(StateTree, StateTree, InstanceData);
 		const bool bInitSucceeded = Exec.IsValid();
-		AITEST_TRUE("StateTree should init", bInitSucceeded);
+		AITEST_TRUE(TEXT("StateTree should init"), bInitSucceeded);
 
-		const FString TickStr(TEXT("Tick"));
-		const FString EnterStateStr(TEXT("EnterState"));
-		const FString ExitStateStr(TEXT("ExitState"));
+		const TCHAR* TickStr = TEXT("Tick");
+		const TCHAR* EnterStateStr = TEXT("EnterState");
+		const TCHAR* ExitStateStr = TEXT("ExitState");
 
 		GetTestRunner().AddExpectedError(TEXT("Loop detected when trying to select state"), EAutomationExpectedErrorFlags::Contains, 1);
 		GetTestRunner().AddExpectedError(TEXT("Failed to select initial state"), EAutomationExpectedErrorFlags::Contains, 1);
 		
 		Status = Exec.Start();
-		AITEST_EQUAL("Start should fail", Status, EStateTreeRunStatus::Failed);
+		AITEST_EQUAL(TEXT("Start should fail"), Status, EStateTreeRunStatus::Failed);
 		Exec.LogClear();
 
 		Exec.Stop();
@@ -833,93 +833,93 @@ struct FStateTreeTest_RegularTransitions : FStateTreeTestBase
 
 			FStateTreeCompiler Compiler(Log);
 			const bool bResult = Compiler.Compile(StateTree);
-			AITEST_TRUE("StateTree should get compiled", bResult);
+			AITEST_TRUE(TEXT("StateTree should get compiled"), bResult);
 		}
 
-		const FString EnterStateStr(TEXT("EnterState"));
-		const FString ExitStateStr(TEXT("ExitState"));
+		const TCHAR* EnterStateStr(TEXT("EnterState"));
+		const TCHAR* ExitStateStr(TEXT("ExitState"));
 
 		{
 			EStateTreeRunStatus Status = EStateTreeRunStatus::Unset;
 			FStateTreeInstanceData InstanceData;
 			FTestStateTreeExecutionContext Exec(StateTree, StateTree, InstanceData);
 			const bool bInitSucceeded = Exec.IsValid();
-			AITEST_TRUE("StateTree should init", bInitSucceeded);
+			AITEST_TRUE(TEXT("StateTree should init"), bInitSucceeded);
 
 			FInstancedPropertyBag Parameters;
 			Parameters.MigrateToNewBagInstance(StateTree.GetDefaultParameters());
 			Parameters.SetValueInt32("Value", 111);
 
 			Status = Exec.Start(&Parameters);
-			AITEST_EQUAL("Start should complete with Running", Status, EStateTreeRunStatus::Running);
-			AITEST_TRUE("Start should enter Global tasks", Exec.Expect("GlobalTask", TEXT("EnterState111")));
-			AITEST_TRUE("Start should enter StateA", Exec.Expect("TaskA", TEXT("EnterState111")));
-			AITEST_TRUE("Start should enter StateB", Exec.Expect("TaskB", TEXT("EnterState1")));
+			AITEST_EQUAL(TEXT("Start should complete with Running"), Status, EStateTreeRunStatus::Running);
+			AITEST_TRUE(TEXT("Start should enter Global tasks"), Exec.Expect("GlobalTask", TEXT("EnterState111")));
+			AITEST_TRUE(TEXT("Start should enter StateA"), Exec.Expect("TaskA", TEXT("EnterState111")));
+			AITEST_TRUE(TEXT("Start should enter StateB"), Exec.Expect("TaskB", TEXT("EnterState1")));
 			Exec.LogClear();
 
 			Status = Exec.Tick(1.5f); // over tick, should trigger
-			AITEST_EQUAL("1st Tick should complete with Running", Status, EStateTreeRunStatus::Running);
-			AITEST_TRUE("1st Tick should tick Global tasks", Exec.Expect("GlobalTask", TEXT("Tick111")));
-			AITEST_TRUE("1st Tick should tick StateA", Exec.Expect("TaskA", TEXT("Tick111")));
-			AITEST_TRUE("1st Tick should tick StateB", Exec.Expect("TaskB", TEXT("Tick1")));
+			AITEST_EQUAL(TEXT("1st Tick should complete with Running"), Status, EStateTreeRunStatus::Running);
+			AITEST_TRUE(TEXT("1st Tick should tick Global tasks"), Exec.Expect("GlobalTask", TEXT("Tick111")));
+			AITEST_TRUE(TEXT("1st Tick should tick StateA"), Exec.Expect("TaskA", TEXT("Tick111")));
+			AITEST_TRUE(TEXT("1st Tick should tick StateB"), Exec.Expect("TaskB", TEXT("Tick1")));
 			Exec.LogClear();
 
 			Status = Exec.Tick(1.0f);
-			AITEST_EQUAL("2nd Tick should complete with Running", Status, EStateTreeRunStatus::Running);
-			AITEST_TRUE("2nd Tick should tick Global tasks", Exec.Expect("GlobalTask", TEXT("Tick111")));
-			AITEST_TRUE("2nd Tick should tick StateA", Exec.Expect("TaskA", TEXT("Tick111")));
-			AITEST_TRUE("2nd Tick should tick the StateB", Exec.Expect("TaskB", TEXT("Tick1")));
-			AITEST_TRUE("2nd Tick should exit the StateB", Exec.Expect("TaskB", TEXT("ExitState1")));
-			AITEST_TRUE("2nd Tick should enter the StateC", Exec.Expect("TaskC", TEXT("EnterState2")));
+			AITEST_EQUAL(TEXT("2nd Tick should complete with Running"), Status, EStateTreeRunStatus::Running);
+			AITEST_TRUE(TEXT("2nd Tick should tick Global tasks"), Exec.Expect("GlobalTask", TEXT("Tick111")));
+			AITEST_TRUE(TEXT("2nd Tick should tick StateA"), Exec.Expect("TaskA", TEXT("Tick111")));
+			AITEST_TRUE(TEXT("2nd Tick should tick the StateB"), Exec.Expect("TaskB", TEXT("Tick1")));
+			AITEST_TRUE(TEXT("2nd Tick should exit the StateB"), Exec.Expect("TaskB", TEXT("ExitState1")));
+			AITEST_TRUE(TEXT("2nd Tick should enter the StateC"), Exec.Expect("TaskC", TEXT("EnterState2")));
 			Exec.LogClear();
 
 			Status = Exec.Tick(1.0f);
-			AITEST_EQUAL("3rd Tick should complete with Running", Status, EStateTreeRunStatus::Running);
-			AITEST_TRUE("3rd Tick should tick Global tasks", Exec.Expect("GlobalTask", TEXT("Tick111")));
-			AITEST_TRUE("3rd Tick should tick StateA", Exec.Expect("TaskA", TEXT("Tick111")));
-			AITEST_TRUE("3rd Tick should tick StateC", Exec.Expect("TaskC", TEXT("Tick2")));
+			AITEST_EQUAL(TEXT("3rd Tick should complete with Running"), Status, EStateTreeRunStatus::Running);
+			AITEST_TRUE(TEXT("3rd Tick should tick Global tasks"), Exec.Expect("GlobalTask", TEXT("Tick111")));
+			AITEST_TRUE(TEXT("3rd Tick should tick StateA"), Exec.Expect("TaskA", TEXT("Tick111")));
+			AITEST_TRUE(TEXT("3rd Tick should tick StateC"), Exec.Expect("TaskC", TEXT("Tick2")));
 			Exec.LogClear();
 
 			Status = Exec.Tick(1.0f);
-			AITEST_EQUAL("4th Tick should complete with Running", Status, EStateTreeRunStatus::Running);
-			AITEST_TRUE("4th Tick should tick Global tasks", Exec.Expect("GlobalTask", TEXT("Tick111")));
-			AITEST_TRUE("244ththnd Tick should tick StateA", Exec.Expect("TaskA", TEXT("Tick111")));
-			AITEST_TRUE("4th Tick should tick the StateC", Exec.Expect("TaskC", TEXT("Tick2")));
-			AITEST_TRUE("4th Tick should exit the StateC", Exec.Expect("TaskC", TEXT("ExitState2")));
-			AITEST_TRUE("4th Tick should enter the StateD", Exec.Expect("TaskD", TEXT("EnterState3")));
+			AITEST_EQUAL(TEXT("4th Tick should complete with Running"), Status, EStateTreeRunStatus::Running);
+			AITEST_TRUE(TEXT("4th Tick should tick Global tasks"), Exec.Expect("GlobalTask", TEXT("Tick111")));
+			AITEST_TRUE(TEXT("4th Tick should tick StateA"), Exec.Expect("TaskA", TEXT("Tick111")));
+			AITEST_TRUE(TEXT("4th Tick should tick the StateC"), Exec.Expect("TaskC", TEXT("Tick2")));
+			AITEST_TRUE(TEXT("4th Tick should exit the StateC"), Exec.Expect("TaskC", TEXT("ExitState2")));
+			AITEST_TRUE(TEXT("4th Tick should enter the StateD"), Exec.Expect("TaskD", TEXT("EnterState3")));
 			Exec.LogClear();
 
 			Status = Exec.Tick(0.001f);
-			AITEST_EQUAL("5th Tick should complete with Running", Status, EStateTreeRunStatus::Running);
-			AITEST_TRUE("5th Tick should tick Global tasks", Exec.Expect("GlobalTask", TEXT("Tick111")));
-			AITEST_TRUE("5th Tick should tick StateA", Exec.Expect("TaskA", TEXT("Tick111")));
-			AITEST_TRUE("5th Tick should tick StateD", Exec.Expect("TaskD", TEXT("Tick3")));
+			AITEST_EQUAL(TEXT("5th Tick should complete with Running"), Status, EStateTreeRunStatus::Running);
+			AITEST_TRUE(TEXT("5th Tick should tick Global tasks"), Exec.Expect("GlobalTask", TEXT("Tick111")));
+			AITEST_TRUE(TEXT("5th Tick should tick StateA"), Exec.Expect("TaskA", TEXT("Tick111")));
+			AITEST_TRUE(TEXT("5th Tick should tick StateD"), Exec.Expect("TaskD", TEXT("Tick3")));
 			Exec.LogClear();
 
 			Status = Exec.Tick(1.0f);
-			AITEST_EQUAL("6th Tick should complete with Running", Status, EStateTreeRunStatus::Running);
-			AITEST_TRUE("6th Tick should tick Global tasks", Exec.Expect("GlobalTask", TEXT("Tick111")));
-			AITEST_TRUE("6th Tick should tick StateA", Exec.Expect("TaskA", TEXT("Tick111")));
-			AITEST_TRUE("6th Tick should tick StateD", Exec.Expect("TaskD", TEXT("Tick3")));
-			AITEST_TRUE("6th Tick should exit the StateD", Exec.Expect("TaskD", TEXT("ExitState3")));
-			//AITEST_FALSE("6th Tick should not exit the Global tasks", Exec.Expect("GlobalTask", TEXT("ExitState111")));
-			//AITEST_FALSE("6th Tick should not enter the Global tasks", Exec.Expect("GlobalTask", TEXT("EnterState111")));
-			//AITEST_FALSE("6th Tick should not exit the StateA", Exec.Expect("TaskA", TEXT("ExitState111")));
-			//AITEST_FALSE("6th Tick should not enter the StateA", Exec.Expect("TaskA", TEXT("EnterState111")));
-			AITEST_TRUE("6th Tick should enter the StateB", Exec.Expect("TaskB", TEXT("EnterState1")));
+			AITEST_EQUAL(TEXT("6th Tick should complete with Running"), Status, EStateTreeRunStatus::Running);
+			AITEST_TRUE(TEXT("6th Tick should tick Global tasks"), Exec.Expect("GlobalTask", TEXT("Tick111")));
+			AITEST_TRUE(TEXT("6th Tick should tick StateA"), Exec.Expect("TaskA", TEXT("Tick111")));
+			AITEST_TRUE(TEXT("6th Tick should tick StateD"), Exec.Expect("TaskD", TEXT("Tick3")));
+			AITEST_TRUE(TEXT("6th Tick should exit the StateD"), Exec.Expect("TaskD", TEXT("ExitState3")));
+			//AITEST_FALSE(TEXT("6th Tick should not exit the Global tasks"), Exec.Expect("GlobalTask", TEXT("ExitState111")));
+			//AITEST_FALSE(TEXT("6th Tick should not enter the Global tasks"), Exec.Expect("GlobalTask", TEXT("EnterState111")));
+			//AITEST_FALSE(TEXT("6th Tick should not exit the StateA"), Exec.Expect("TaskA", TEXT("ExitState111")));
+			//AITEST_FALSE(TEXT("6th Tick should not enter the StateA"), Exec.Expect("TaskA", TEXT("EnterState111")));
+			AITEST_TRUE(TEXT("6th Tick should enter the StateB"), Exec.Expect("TaskB", TEXT("EnterState1")));
 			Exec.LogClear();
 
 			Status = Exec.Tick(1.0f);
-			AITEST_EQUAL("7th Tick should complete with Running", Status, EStateTreeRunStatus::Running);
-			AITEST_TRUE("7th Tick should tick Global tasks", Exec.Expect("GlobalTask", TEXT("Tick111")));
-			AITEST_TRUE("7th Tick should tick StateA", Exec.Expect("TaskA", TEXT("Tick111")));
-			AITEST_TRUE("7th Tick should tick StateB", Exec.Expect("TaskB", TEXT("Tick1")));
+			AITEST_EQUAL(TEXT("7th Tick should complete with Running"), Status, EStateTreeRunStatus::Running);
+			AITEST_TRUE(TEXT("7th Tick should tick Global tasks"), Exec.Expect("GlobalTask", TEXT("Tick111")));
+			AITEST_TRUE(TEXT("7th Tick should tick StateA"), Exec.Expect("TaskA", TEXT("Tick111")));
+			AITEST_TRUE(TEXT("7th Tick should tick StateB"), Exec.Expect("TaskB", TEXT("Tick1")));
 			Exec.LogClear();
 
 			Exec.Stop();
-			AITEST_TRUE("Stop Tick should exit the StateB", Exec.Expect("TaskB", TEXT("ExitState1")));
-			AITEST_TRUE("Stop Tick should exit the StateA", Exec.Expect("TaskA", TEXT("ExitState111")));
-			AITEST_TRUE("Stop should tick Global tasks", Exec.Expect("GlobalTask", TEXT("ExitState111")));
+			AITEST_TRUE(TEXT("Stop Tick should exit the StateB"), Exec.Expect("TaskB", TEXT("ExitState1")));
+			AITEST_TRUE(TEXT("Stop Tick should exit the StateA"), Exec.Expect("TaskA", TEXT("ExitState111")));
+			AITEST_TRUE(TEXT("Stop should tick Global tasks"), Exec.Expect("GlobalTask", TEXT("ExitState111")));
 			Exec.LogClear();
 		}
 
diff --git a/Engine/Plugins/Runtime/StateTree/Source/StateTreeTestSuite/Private/StateTreeWeakContextTest.cpp b/Engine/Plugins/Runtime/StateTree/Source/StateTreeTestSuite/Private/StateTreeWeakContextTest.cpp
--- a/Engine/Plugins/Runtime/StateTree/Source/StateTreeTestSuite/Private/StateTreeWeakContextTest.cpp
+++ b/Engine/Plugins/Runtime/StateTree/Source/StateTreeTestSuite/Private/StateTreeWeakContextTest.cpp
@@ -101,7 +101,7 @@ struct FStateTreeTest_WeakContext_FinishTask : FStateTreeTestBase
 			FStateTreeCompilerLog Log;
 			FStateTreeCompiler Compiler(Log);
 			const bool bResult = Compiler.Compile(StateTree1);
-			AITEST_TRUE("StateTree1 should get compiled", bResult);
+			AITEST_TRUE(TEXT("StateTree1 should get compiled"), bResult);
 		}
 
 		// Create context
@@ -110,12 +110,12 @@ struct FStateTreeTest_WeakContext_FinishTask : FStateTreeTestBase
 			{
 				FTestStateTreeExecutionContext Exec(StateTree1, StateTree1, InstanceData);
 				const bool bInitSucceeded = Exec.IsValid();
-				AITEST_TRUE("StateTree should init", bInitSucceeded);
+				AITEST_TRUE(TEXT("StateTree should init"), bInitSucceeded);
 			}
 			{
 				FTestStateTreeExecutionContext Exec(StateTree1, StateTree1, InstanceData);
 				EStateTreeRunStatus Status = Exec.Start();
-				AITEST_EQUAL("Start should complete with Running", Status, EStateTreeRunStatus::Running);
+				AITEST_EQUAL(TEXT("Start should complete with Running"), Status, EStateTreeRunStatus::Running);
 				AITEST_TRUE(TEXT("In correct states"), Exec.ExpectInActiveStates("Tree1StateRoot", "Tree1State1", "Tree1State2"));
 
 				AITEST_TRUE(TEXT("Start should EnterState"), Exec.Expect("Tree1GlobalTask", TEXT("EnterState0"))
@@ -129,7 +129,7 @@ struct FStateTreeTest_WeakContext_FinishTask : FStateTreeTestBase
 			{
 				FTestStateTreeExecutionContext Exec(StateTree1, StateTree1, InstanceData);
 				EStateTreeRunStatus Status = Exec.Tick(0.1f);
-				AITEST_EQUAL("Tick should complete with Running", Status, EStateTreeRunStatus::Running);
+				AITEST_EQUAL(TEXT("Tick should complete with Running"), Status, EStateTreeRunStatus::Running);
 				AITEST_TRUE(TEXT("In correct states"), Exec.ExpectInActiveStates("Tree1StateRoot", "Tree1State1", "Tree1State2"));
 
 				AITEST_TRUE(TEXT("Tick should Tick"), Exec.Expect("Tree1GlobalTask", TEXT("Tick0"))
@@ -144,7 +144,7 @@ struct FStateTreeTest_WeakContext_FinishTask : FStateTreeTestBase
 				WeakContext.bGlobalFinishTaskSuccessOnTick = true;
 				FTestStateTreeExecutionContext Exec(StateTree1, StateTree1, InstanceData);
 				EStateTreeRunStatus Status = Exec.Tick(0.1f);
-				AITEST_EQUAL("Tick should complete with Succeeded", Status, EStateTreeRunStatus::Succeeded);
+				AITEST_EQUAL(TEXT("Tick should complete with Succeeded"), Status, EStateTreeRunStatus::Succeeded);
 				AITEST_TRUE(TEXT("Tick should Tick"), Exec.Expect("Tree1GlobalTask", TEXT("Tick0"))
 							.Then("Tree1State2Task", TEXT("ExitState0"))
 							.Then("Tree1State1Task", TEXT("ExitState0"))
@@ -164,7 +164,7 @@ struct FStateTreeTest_WeakContext_FinishTask : FStateTreeTestBase
 				FTestStateTreeExecutionContext Exec(StateTree1, StateTree1, InstanceData);
 				WeakContext.ContextTree1GlobalTask.FinishTask(EStateTreeFinishTaskType::Failed);
 				EStateTreeRunStatus Status = Exec.Tick(0.1f);
-				AITEST_EQUAL("Tick should complete with Failed", Status, EStateTreeRunStatus::Failed);
+				AITEST_EQUAL(TEXT("Tick should complete with Failed"), Status, EStateTreeRunStatus::Failed);
 				AITEST_FALSE(TEXT("Tick should not Tick"), Exec.Expect("Tree1GlobalTask", TEXT("Tick0")));
 				AITEST_FALSE(TEXT("Tick should not Tick"), Exec.Expect("Tree1RootTask", TEXT("Tick0")));
 				AITEST_TRUE(TEXT("Tick should Tick"), Exec.Expect("Tree1State2Task", TEXT("ExitState0"))
@@ -184,7 +184,7 @@ struct FStateTreeTest_WeakContext_FinishTask : FStateTreeTestBase
 				WeakContext.bState1lFinishTaskFailOnTick = true;
 				FTestStateTreeExecutionContext Exec(StateTree1, StateTree1, InstanceData);
 				EStateTreeRunStatus Status = Exec.Tick(0.1f);
-				AITEST_EQUAL("Tick should complete with Running", Status, EStateTreeRunStatus::Running);
+				AITEST_EQUAL(TEXT("Tick should complete with Running"), Status, EStateTreeRunStatus::Running);
 				AITEST_TRUE(TEXT("Tick should Tick"), Exec.Expect("Tree1GlobalTask", TEXT("Tick0"))
 							.Then("Tree1RootTask", TEXT("Tick0"))
 							.Then("Tree1State1Task", TEXT("Tick0"))
@@ -200,7 +200,7 @@ struct FStateTreeTest_WeakContext_FinishTask : FStateTreeTestBase
 				FTestStateTreeExecutionContext Exec(StateTree1, StateTree1, InstanceData);
 				WeakContext.ContextTree1State1Task.FinishTask(EStateTreeFinishTaskType::Succeeded);
 				EStateTreeRunStatus Status = Exec.Tick(0.1f);
-				AITEST_EQUAL("Tick should complete with Running", Status, EStateTreeRunStatus::Running);
+				AITEST_EQUAL(TEXT("Tick should complete with Running"), Status, EStateTreeRunStatus::Running);
 				AITEST_TRUE(TEXT("Tick should Tick"), Exec.Expect("Tree1GlobalTask", TEXT("Tick0"))
 							.Then("Tree1RootTask", TEXT("Tick0"))
 							.Then("Tree1State2Task", TEXT("Tick0"))
