StateTree: Update the K2Node for PropertyRef to use the byte type (as blueprint does) when the property is an enum.

diff --git a/Engine/Plugins/Runtime/StateTree/Source/StateTreeEditorModule/Private/K2Node_StateTreeBlueprintPropertyRef.cpp b/Engine/Plugins/Runtime/StateTree/Source/StateTreeEditorModule/Private/K2Node_StateTreeBlueprintPropertyRef.cpp
--- a/Engine/Plugins/Runtime/StateTree/Source/StateTreeEditorModule/Private/K2Node_StateTreeBlueprintPropertyRef.cpp
+++ b/Engine/Plugins/Runtime/StateTree/Source/StateTreeEditorModule/Private/K2Node_StateTreeBlueprintPropertyRef.cpp
@@ -91,32 +91,33 @@ void UK2Node_StateTreeBlueprintPropertyRef::ValidateNodeDuringCompilation(FCompi
 void UK2Node_StateTreeBlueprintPropertyRef::UpdateOutputPin() const
 {
 	UEdGraphPin& OutValuePin = *FindPinChecked(UE::StateTree::BlueprintPropertyRef::OutReferencePinName);
-    const UEdGraphPin& PropertyRefPin = *FindPinChecked(UE::StateTree::BlueprintPropertyRef::PropertyRefPinName);
+	const UEdGraphPin& PropertyRefPin = *FindPinChecked(UE::StateTree::BlueprintPropertyRef::PropertyRefPinName);
 
-    const UBlueprint* Blueprint = FBlueprintEditorUtils::FindBlueprintForNodeChecked(this);
-    if (Blueprint && Blueprint->GeneratedClass && PropertyRefPin.LinkedTo.Num() == 1)
-    {
-        UK2Node_VariableGet* VariableGet = Cast<UK2Node_VariableGet>(PropertyRefPin.LinkedTo[0]->GetOwningNode());
-        if (VariableGet)
-        {
-            UObject* CDO = Blueprint->GeneratedClass->GetDefaultObject<UObject>();
-            FProperty* Property = Blueprint->GeneratedClass->FindPropertyByName(VariableGet->VariableReference.GetMemberName());
-            FStructProperty* StructProperty = CastField<FStructProperty>(Property);
-            if (StructProperty && StructProperty->Struct == FStateTreeBlueprintPropertyRef::StaticStruct())
-            {
-                FStateTreeBlueprintPropertyRef PropertyRef;
-                Property->GetValue_InContainer(CDO, &PropertyRef);
+	const UBlueprint* Blueprint = FBlueprintEditorUtils::FindBlueprintForNodeChecked(this);
+	if (Blueprint && Blueprint->GeneratedClass && PropertyRefPin.LinkedTo.Num() == 1)
+	{
+		UK2Node_VariableGet* VariableGet = Cast<UK2Node_VariableGet>(PropertyRefPin.LinkedTo[0]->GetOwningNode());
+		if (VariableGet)
+		{
+			UObject* CDO = Blueprint->GeneratedClass->GetDefaultObject<UObject>();
+			FProperty* Property = Blueprint->GeneratedClass->FindPropertyByName(VariableGet->VariableReference.GetMemberName());
+			FStructProperty* StructProperty = CastField<FStructProperty>(Property);
+			if (StructProperty && StructProperty->Struct == FStateTreeBlueprintPropertyRef::StaticStruct())
+			{
+				FStateTreeBlueprintPropertyRef PropertyRef;
+				Property->GetValue_InContainer(CDO, &PropertyRef);
 				FEdGraphPinType PinType = UE::StateTree::PropertyRefHelpers::GetBlueprintPropertyRefInternalTypeAsPin(PropertyRef);
 
 				PinType.bIsReference = true;
 				OutValuePin.PinType = PinType;
-                return;
-            }
-        }
-    }
+				GetDefault<UEdGraphSchema_K2>()->SetPinAutogeneratedDefaultValueBasedOnType(&OutValuePin);
+				return;
+			}
+		}
+	}
 
-    FCreatePinParams PinParams{};
-    OutValuePin.PinType = FEdGraphPinType(UEdGraphSchema_K2::PC_Wildcard, NAME_None, nullptr, PinParams.ContainerType, PinParams.bIsReference, PinParams.ValueTerminalType);
+	FCreatePinParams PinParams{};
+	OutValuePin.PinType = FEdGraphPinType(UEdGraphSchema_K2::PC_Wildcard, NAME_None, nullptr, PinParams.ContainerType, PinParams.bIsReference, PinParams.ValueTerminalType);
 }
 
 void FKCHandler_StateTreeBlueprintPropertyRefGet::RegisterNets(FKismetFunctionContext& Context, UEdGraphNode* Node)
diff --git a/Engine/Plugins/Runtime/StateTree/Source/StateTreeModule/Private/StateTreePropertyRefHelpers.cpp b/Engine/Plugins/Runtime/StateTree/Source/StateTreeModule/Private/StateTreePropertyRefHelpers.cpp
--- a/Engine/Plugins/Runtime/StateTree/Source/StateTreeModule/Private/StateTreePropertyRefHelpers.cpp
+++ b/Engine/Plugins/Runtime/StateTree/Source/StateTreeModule/Private/StateTreePropertyRefHelpers.cpp
@@ -11,24 +11,25 @@
 #include "EdGraphSchema_K2.h"
 #include "IPropertyAccessEditor.h"
 #include "EdGraph/EdGraphPin.h"
+#include <limits>
 #endif
 
 namespace UE::StateTree::PropertyRefHelpers
 {
 #if WITH_EDITOR
-	static const FName BoolName = TEXT("bool");
-	static const FName ByteName = TEXT("byte");
-	static const FName Int32Name = TEXT("int32");
-	static const FName Int64Name = TEXT("int64");
-	static const FName FloatName = TEXT("float");
-	static const FName DoubleName = TEXT("double");
-	static const FName NameName = TEXT("Name");
-	static const FName StringName = TEXT("String");
-	static const FName TextName = TEXT("Text");
-	const FName IsRefToArrayName = TEXT("IsRefToArray");
-	const FName CanRefToArrayName = TEXT("CanRefToArray");
-	const FName RefTypeName = TEXT("RefType");
-	static const FName IsOptionalName = TEXT("Optional");
+	static const FLazyName BoolName = "bool";
+	static const FLazyName ByteName = "byte";
+	static const FLazyName Int32Name = "int32";
+	static const FLazyName Int64Name = "int64";
+	static const FLazyName FloatName = "float";
+	static const FLazyName DoubleName = "double";
+	static const FLazyName NameName = "Name";
+	static const FLazyName StringName = "String";
+	static const FLazyName TextName = "Text";
+	const FName IsRefToArrayName = "IsRefToArray";
+	const FName CanRefToArrayName = "CanRefToArray";
+	const FName RefTypeName = "RefType";
+	static const FLazyName IsOptionalName = "Optional";
 
 	bool ArePropertyRefsCompatible(const FProperty& TargetRefProperty, const FProperty& SourceRefProperty, const void* TargetRefAddress, const void* SourceRefAddress)
 	{
@@ -344,6 +345,18 @@ PRAGMA_ENABLE_DEPRECATION_WARNINGS
 		case EStateTreePropertyRefType::Enum:
 			PinType.PinCategory = UEdGraphSchema_K2::PC_Enum;
 			PinType.PinSubCategoryObject = PropertyRef.GetTypeObject();
+			if (UEnum* Enum = Cast<UEnum>(PinType.PinSubCategoryObject))
+			{
+				if (Enum->GetMaxEnumValue() <= (int64)std::numeric_limits<uint8>::max())
+				{
+					// Use byte for BP. It will use the correct picker and enum k2 node.
+					PinType.PinCategory = UEdGraphSchema_K2::PC_Byte;
+				}
+			}
+			else
+			{
+				UE_LOG(LogStateTree, Warning, TEXT("The property ref of type enum has an invalid enum. %s"), *GetFullNameSafe(PinType.PinSubCategoryObject.Get()));
+			}
 			break;
 
 		case EStateTreePropertyRefType::Struct:
@@ -605,6 +618,11 @@ PRAGMA_ENABLE_DEPRECATION_WARNINGS
 				{
 					PinType.PinCategory = UEdGraphSchema_K2::PC_Enum;
 					PinType.PinSubCategoryObject = Enum;
+					if (Enum->GetMaxEnumValue() <= (int64)std::numeric_limits<uint8>::max())
+					{
+						// Use byte for BP. It will use the correct picker and enum k2 node.
+						PinType.PinCategory = UEdGraphSchema_K2::PC_Byte;
+					}
 				}
 				else
 				{
