[State Tree] Deprecated properties in STNodeBlueprintBase. Started using StrongExecutionContext.

diff --git a/Engine/Plugins/Runtime/StateTree/Source/StateTreeModule/Private/Blueprint/StateTreeNodeBlueprintBase.cpp b/Engine/Plugins/Runtime/StateTree/Source/StateTreeModule/Private/Blueprint/StateTreeNodeBlueprintBase.cpp
--- a/Engine/Plugins/Runtime/StateTree/Source/StateTreeModule/Private/Blueprint/StateTreeNodeBlueprintBase.cpp
+++ b/Engine/Plugins/Runtime/StateTree/Source/StateTreeModule/Private/Blueprint/StateTreeNodeBlueprintBase.cpp
@@ -49,51 +49,36 @@ AActor* UStateTreeNodeBlueprintBase::GetOwnerActor(const FStateTreeExecutionCont
 void UStateTreeNodeBlueprintBase::SetCachedInstanceDataFromContext(const FStateTreeExecutionContext& Context) const
 {
 	WeakExecutionContext = Context.MakeWeakExecutionContext();
-	if (FStateTreeInstanceData* InstanceData = Context.GetMutableInstanceData())
-	{
-		WeakInstanceStorage = InstanceData->GetWeakMutableStorage();
-	}
-
-	const FStateTreeExecutionFrame* CurrentlyProcessedFrame = Context.GetCurrentlyProcessedFrame();
-	check(CurrentlyProcessedFrame);
-
-	CachedFrameStateTree = CurrentlyProcessedFrame->StateTree;
-	CachedFrameRootState = CurrentlyProcessedFrame->RootState;
 }
 
 void UStateTreeNodeBlueprintBase::ClearCachedInstanceData() const
 {
 	WeakExecutionContext = FStateTreeWeakExecutionContext();
-	WeakInstanceStorage = nullptr;
-	CachedFrameStateTree = nullptr;
-	CachedFrameRootState = FStateTreeStateHandle::Invalid;
 }
 
 void* UStateTreeNodeBlueprintBase::GetMutablePtrToProperty(const FStateTreeBlueprintPropertyRef& PropertyRef, FProperty*& OutSourceProperty) const
 {
-	TSharedPtr<FStateTreeInstanceStorage> InstanceStorage = WeakInstanceStorage.Pin();
-	if (InstanceStorage == nullptr)
+	using namespace UE::StateTree;
+	const FStateTreeStrongExecutionContext StrongContext = WeakExecutionContext.MakeStrongExecutionContext();
+
+	const Async::FActivePathInfo ActivePath = StrongContext.GetActivePathInfo();
+	if (!ActivePath.IsValid())
 	{
 		UE_VLOG_UELOG(this, LogStateTree, Error, TEXT("Trying to GetMutablePtrToProperty while node is not active."));
 		return nullptr;
 	}
 
-	check(CachedFrameStateTree && CachedFrameRootState.IsValid());
-
-	TConstArrayView<FStateTreeExecutionFrame> ActiveFrames = InstanceStorage->GetExecutionState().ActiveFrames;
-	const FStateTreeExecutionFrame* CurrentParentFrame = nullptr;
-	const FStateTreeExecutionFrame* CurrentFrame = FStateTreeExecutionContext::FindFrame(CachedFrameStateTree, CachedFrameRootState, ActiveFrames, CurrentParentFrame);
-	check(CurrentFrame);
-
+	check(StrongContext.IsValid());
 	const FProperty* SourceProperty = nullptr;
-	void* PropertyAddress = UE::StateTree::PropertyRefHelpers::GetMutablePtrToProperty<void>(PropertyRef, *InstanceStorage, *CurrentFrame, CurrentParentFrame, &SourceProperty);
-	if (PropertyAddress && UE::StateTree::PropertyRefHelpers::IsBlueprintPropertyRefCompatibleWithProperty(*SourceProperty, &PropertyRef))
+	void* PropertyAddress = PropertyRefHelpers::GetMutablePtrToProperty<void>(PropertyRef, *const_cast<FStateTreeInstanceStorage*>(StrongContext.GetStorage().Get()), *ActivePath.Frame, ActivePath.ParentFrame, &SourceProperty);
+	if (PropertyAddress && PropertyRefHelpers::IsBlueprintPropertyRefCompatibleWithProperty(*SourceProperty, &PropertyRef))
 	{
 		OutSourceProperty = const_cast<FProperty*>(SourceProperty);
 		return PropertyAddress;
 	}
 
 	OutSourceProperty = nullptr;
+	
 	return nullptr;
 }
 
diff --git a/Engine/Plugins/Runtime/StateTree/Source/StateTreeModule/Public/Blueprint/StateTreeNodeBlueprintBase.h b/Engine/Plugins/Runtime/StateTree/Source/StateTreeModule/Public/Blueprint/StateTreeNodeBlueprintBase.h
--- a/Engine/Plugins/Runtime/StateTree/Source/StateTreeModule/Public/Blueprint/StateTreeNodeBlueprintBase.h
+++ b/Engine/Plugins/Runtime/StateTree/Source/StateTreeModule/Public/Blueprint/StateTreeNodeBlueprintBase.h
@@ -103,17 +103,21 @@ private:
 	/** Cached execution context while the node is active for async nodes. */
 	mutable FStateTreeWeakExecutionContext WeakExecutionContext;
 
+#if WITH_EDITORONLY_DATA
+
+	UE_DEPRECATED(5.6, "WeakInstanceStorage is deprecated. Use WeakExecutionContext")
 	/** Cached instance data while the node is active for async nodes. */
 	mutable TWeakPtr<FStateTreeInstanceStorage> WeakInstanceStorage;
 
+	UE_DEPRECATED(5.6, "CachedFrameStateTree is deprecated.")
 	/** Cached State Tree of owning execution frame. */
 	UPROPERTY()
 	mutable TObjectPtr<const UStateTree> CachedFrameStateTree = nullptr;
 
+	UE_DEPRECATED(5.6, "CachedFrameRootState is deprecated.")
 	/** Cached root state of owning execution frame. */
 	mutable FStateTreeStateHandle CachedFrameRootState;
 
-#if WITH_EDITORONLY_DATA
 	/** Description of the node. */
 	UPROPERTY(EditDefaultsOnly, Category="Description")
 	FText Description;
