[State Tree] fixed rvalue destroyed mistake

diff --git a/Engine/Plugins/Runtime/StateTree/Source/StateTreeEditorModule/Private/StateTreeCompiler.cpp b/Engine/Plugins/Runtime/StateTree/Source/StateTreeEditorModule/Private/StateTreeCompiler.cpp
--- a/Engine/Plugins/Runtime/StateTree/Source/StateTreeEditorModule/Private/StateTreeCompiler.cpp
+++ b/Engine/Plugins/Runtime/StateTree/Source/StateTreeEditorModule/Private/StateTreeCompiler.cpp
@@ -1991,7 +1991,8 @@ bool FStateTreeCompiler::CompileAndValidateNode(const UStateTreeState* SourceSta
 		return false;
 	}
 
-	UE::StateTree::FCompileNodeContext CompileContext(InstanceData, NodeDesc, FStateTreeBindingLookup(EditorData));
+	const FStateTreeBindingLookup& BindingLookup = FStateTreeBindingLookup(EditorData);
+	UE::StateTree::FCompileNodeContext CompileContext(InstanceData, NodeDesc, BindingLookup);
 	const EDataValidationResult Result = Node.Compile(CompileContext);
 
 	if (Result == EDataValidationResult::Invalid && CompileContext.ValidationErrors.IsEmpty())
