StateTree: Fix crash when the flag "task" is disabled in the editor view.
#rnx

diff --git a/Engine/Plugins/Runtime/StateTree/Source/StateTreeEditorModule/Private/SStateTreeViewRow.cpp b/Engine/Plugins/Runtime/StateTree/Source/StateTreeEditorModule/Private/SStateTreeViewRow.cpp
--- a/Engine/Plugins/Runtime/StateTree/Source/StateTreeEditorModule/Private/SStateTreeViewRow.cpp
+++ b/Engine/Plugins/Runtime/StateTree/Source/StateTreeEditorModule/Private/SStateTreeViewRow.cpp
@@ -834,7 +834,10 @@ TSharedRef<SWidget> SStateTreeViewRow::MakeConditionsWidget(const TSharedPtr<SSc
 void SStateTreeViewRow::MakeTransitionsWidget()
 {
 	SHorizontalBox* TransitionsContainerPtr = TransitionsContainer.Get();
-	check(TransitionsContainerPtr != nullptr)
+	if (TransitionsContainerPtr == nullptr)
+	{
+		return;
+	}
 
 	TransitionsContainerPtr->ClearChildren();
 
