StateTree: Fix runtime validation for global tasks. During state selection, you can enter a new linked state with the same tree as a previously running linked state. The EnterState on the new global task is executed before the previous task receives ExitState.

#rnx

diff --git a/Engine/Plugins/Runtime/StateTree/Source/StateTreeModule/Private/Debugger/StateTreeDebug.cpp b/Engine/Plugins/Runtime/StateTree/Source/StateTreeModule/Private/Debugger/StateTreeDebug.cpp
--- a/Engine/Plugins/Runtime/StateTree/Source/StateTreeModule/Private/Debugger/StateTreeDebug.cpp
+++ b/Engine/Plugins/Runtime/StateTree/Source/StateTreeModule/Private/Debugger/StateTreeDebug.cpp
@@ -53,7 +53,8 @@ void NodeEnter(const FStateTreeExecutionContext& ExecutionContext, const FNodeRe
 	{
 		if (FRuntimeValidationInstanceData* RuntimeValidation = InstanceData->GetRuntimeValidation().GetInstanceData())
 		{
-			RuntimeValidation->NodeEnterState(NodeId);
+			const FActiveFrameID FrameID = ExecutionContext.GetCurrentlyProcessedFrame() ? ExecutionContext.GetCurrentlyProcessedFrame()->FrameID : FActiveFrameID();
+			RuntimeValidation->NodeEnterState(NodeId, FrameID);
 		}
 		Delegate.Broadcast(ExecutionContext, FNodeDelegateArgs{ .Node = Node, .NodeId = NodeId });
 	}
@@ -67,7 +68,8 @@ void NodeExit(const FStateTreeExecutionContext& ExecutionContext, const FNodeRef
 	{
 		if (FRuntimeValidationInstanceData* RuntimeValidation = InstanceData->GetRuntimeValidation().GetInstanceData())
 		{
-			RuntimeValidation->NodeExitState(NodeId);
+			const FActiveFrameID FrameID = ExecutionContext.GetCurrentlyProcessedFrame() ? ExecutionContext.GetCurrentlyProcessedFrame()->FrameID : FActiveFrameID();
+			RuntimeValidation->NodeExitState(NodeId, FrameID);
 		}
 		Delegate.Broadcast(ExecutionContext, FNodeDelegateArgs{ .Node = Node, .NodeId = NodeId });
 	}
diff --git a/Engine/Plugins/Runtime/StateTree/Source/StateTreeModule/Private/Debugger/StateTreeRuntimeValidation.cpp b/Engine/Plugins/Runtime/StateTree/Source/StateTreeModule/Private/Debugger/StateTreeRuntimeValidation.cpp
--- a/Engine/Plugins/Runtime/StateTree/Source/StateTreeModule/Private/Debugger/StateTreeRuntimeValidation.cpp
+++ b/Engine/Plugins/Runtime/StateTree/Source/StateTreeModule/Private/Debugger/StateTreeRuntimeValidation.cpp
@@ -57,7 +57,7 @@ FRuntimeValidationInstanceData::~FRuntimeValidationInstanceData()
 		{
 			if (EnumHasAllFlags(Pair.State, EState::BetweenEnterExitState))
 			{
-				ensureAlwaysMsgf(false, TEXT("Tree exited. Missing ExitState on %s."), *Private::NodeToString(StateTree.Get(), Pair.Id));
+				ensureAlwaysMsgf(false, TEXT("Tree exited. Missing ExitState on %s."), *Private::NodeToString(StateTree.Get(), Pair.NodeID));
 				Private::bRuntimeValidationEnterExitState = false;
 			}
 		}
@@ -94,11 +94,11 @@ void FRuntimeValidationInstanceData::SetContext(const UObject* InNewOwner, const
 	Owner = NewOwner;
 }
 
-void FRuntimeValidationInstanceData::NodeEnterState(FGuid Id)
+void FRuntimeValidationInstanceData::NodeEnterState(FGuid NodeID, FActiveFrameID FrameID)
 {
-	FNodeStatePair* Found = NodeStates.FindByPredicate([Id](const FNodeStatePair& Other)
+	FNodeStatePair* Found = NodeStates.FindByPredicate([&NodeID, &FrameID](const FNodeStatePair& Other)
 	{
-		return Other.Id == Id;
+		return Other.NodeID == NodeID && Other.FrameID == FrameID;
 	});
 	if (Found)
 	{
@@ -106,22 +106,22 @@ void FRuntimeValidationInstanceData::NodeEnterState(FGuid Id)
 		{
 			ensureAlwaysMsgf(false, TEXT("StateTree runtime check failed: EnterState executed on node %s without an ExitState.\n"
 				"Auto deactivate Runtime check StateTree.RuntimeValidation.EnterExitState to prevent reporting the same error multiple times.")
-				, *Private::NodeToString(Owner.Get(), Id));
+				, *Private::NodeToString(Owner.Get(), NodeID));
 			Private::bRuntimeValidationEnterExitState = false;
 		}
 		EnumAddFlags(Found->State, EState::BetweenEnterExitState);
 	}
 	else
 	{
-		NodeStates.Add(FNodeStatePair{.Id = Id, .State = EState::BetweenEnterExitState });
+		NodeStates.Add(FNodeStatePair{.NodeID = NodeID, .FrameID = FrameID, .State = EState::BetweenEnterExitState });
 	}
 }
 
-void FRuntimeValidationInstanceData::NodeExitState(FGuid Id)
+void FRuntimeValidationInstanceData::NodeExitState(FGuid NodeID, FActiveFrameID FrameID)
 {
-	FNodeStatePair* Found = NodeStates.FindByPredicate([Id](const FNodeStatePair& Other)
+	FNodeStatePair* Found = NodeStates.FindByPredicate([&NodeID, &FrameID](const FNodeStatePair& Other)
 	{
-		return Other.Id == Id;
+		return Other.NodeID == NodeID && Other.FrameID == FrameID;
 	});
 	if (Found)
 	{
@@ -129,7 +129,7 @@ void FRuntimeValidationInstanceData::NodeExitState(FGuid Id)
 		{
 			ensureAlwaysMsgf(false, TEXT("StateTree runtime check failed: ExitState executed on node %s without an EnterState.\n"
 				"Auto deactivate Runtime check StateTree.RuntimeValidation.EnterExitState to prevent reporting the same error multiple times.")
-				, *Private::NodeToString(Owner.Get(), Id));
+				, *Private::NodeToString(Owner.Get(), NodeID));
 			Private::bRuntimeValidationEnterExitState = false;
 		}
 		EnumRemoveFlags(Found->State, EState::BetweenEnterExitState);
@@ -138,7 +138,7 @@ void FRuntimeValidationInstanceData::NodeExitState(FGuid Id)
 	{
 		ensureAlwaysMsgf(false, TEXT("StateTree runtime check failed: ExitState executed on node %s without an EnterState.\n"
 			"Auto deactivate Runtime check StateTree.RuntimeValidation.EnterExitState to prevent reporting the same error multiple times.")
-			, *Private::NodeToString(Owner.Get(), Id));
+			, *Private::NodeToString(Owner.Get(), NodeID));
 		Private::bRuntimeValidationEnterExitState = false;
 	}
 }
diff --git a/Engine/Plugins/Runtime/StateTree/Source/StateTreeModule/Private/Debugger/StateTreeRuntimeValidationInstanceData.h b/Engine/Plugins/Runtime/StateTree/Source/StateTreeModule/Private/Debugger/StateTreeRuntimeValidationInstanceData.h
--- a/Engine/Plugins/Runtime/StateTree/Source/StateTreeModule/Private/Debugger/StateTreeRuntimeValidationInstanceData.h
+++ b/Engine/Plugins/Runtime/StateTree/Source/StateTreeModule/Private/Debugger/StateTreeRuntimeValidationInstanceData.h
@@ -2,7 +2,7 @@
 
 #pragma once
 
-#include "StateTreeTypes.h"
+#include "StateTreeStatePath.h"
 
 class UStateTree;
 class UObject;
@@ -22,8 +22,8 @@ public:
 
 	void SetContext(const UObject* Owner, const UStateTree* StateTree);
 
-	void NodeEnterState(FGuid Id);
-	void NodeExitState(FGuid Id);
+	void NodeEnterState(FGuid NodeID, FActiveFrameID FrameID);
+	void NodeExitState(FGuid NodeID, FActiveFrameID FrameID);
 
 private:
 	void ValidatesTreeNodes(const UStateTree* InNewStateTree) const;
@@ -38,7 +38,8 @@ private:
 
 	struct FNodeStatePair
 	{
-		FGuid Id;
+		FGuid NodeID;
+		FActiveFrameID FrameID;
 		EState State = EState::None;
 	};
 	TArray<FNodeStatePair> NodeStates;
