[State Tree]
- change interface functions from UE pure virtual to be c++ pure virtual
- pass FName by ref

diff --git a/Engine/Plugins/Runtime/StateTree/Source/StateTreeEditorModule/Private/StateTreeCompiler.cpp b/Engine/Plugins/Runtime/StateTree/Source/StateTreeEditorModule/Private/StateTreeCompiler.cpp
--- a/Engine/Plugins/Runtime/StateTree/Source/StateTreeEditorModule/Private/StateTreeCompiler.cpp
+++ b/Engine/Plugins/Runtime/StateTree/Source/StateTreeEditorModule/Private/StateTreeCompiler.cpp
@@ -39,7 +39,7 @@ namespace UE::StateTree
 			return InstanceDataView;
 		}
 
-		virtual bool HasBindingForProperty(const FName& PropertyName) const override
+		virtual bool HasBindingForProperty(const FName PropertyName) const override
 		{
 			const FPropertyBindingPath& PropertyPath = FPropertyBindingPath(Desc.ID, PropertyName);
 
diff --git a/Engine/Plugins/Runtime/StateTree/Source/StateTreeModule/Public/StateTreeNodeBase.h b/Engine/Plugins/Runtime/StateTree/Source/StateTreeModule/Public/StateTreeNodeBase.h
--- a/Engine/Plugins/Runtime/StateTree/Source/StateTreeModule/Public/StateTreeNodeBase.h
+++ b/Engine/Plugins/Runtime/StateTree/Source/StateTreeModule/Public/StateTreeNodeBase.h
@@ -23,17 +23,15 @@ struct IStateTreeBindingLookup;
 #define SET_NODE_CUSTOM_TRACE_TEXT(...)
 #endif //WITH_STATETREE_TRACE
 
-#define UE_API STATETREEMODULE_API
-
 #if WITH_EDITOR
 namespace UE::StateTree
 {
 	struct ICompileNodeContext
 	{
 		virtual ~ICompileNodeContext() {}
-		UE_API virtual void AddValidationError(const FText& Message) PURE_VIRTUAL(ICompileNodeContext::AddValidationError,;);
-		UE_API virtual FStateTreeDataView GetInstanceDataView() const PURE_VIRTUAL(ICompileNodeContext::GetInstanceDataView, return FStateTreeDataView(););
-		UE_API virtual bool HasBindingForProperty(const FName& PropertyName) const PURE_VIRTUAL(ICompileNodeContext::HasBinding, return false;);
+		virtual void AddValidationError(const FText& Message) = 0;
+		virtual FStateTreeDataView GetInstanceDataView() const = 0;
+		virtual bool HasBindingForProperty(const FName PropertyName) const = 0;
 	};
 }
 #endif
@@ -192,5 +190,3 @@ PRAGMA_ENABLE_DEPRECATION_WARNINGS
 	UPROPERTY()
 	FStateTreeDataHandle InstanceDataHandle = FStateTreeDataHandle::Invalid;
 };
-
-#undef UE_API
\ No newline at end of file
diff --git a/Engine/Plugins/Runtime/StateTree/Source/StateTreeModule/Public/StateTreePropertyBindings.h b/Engine/Plugins/Runtime/StateTree/Source/StateTreeModule/Public/StateTreePropertyBindings.h
--- a/Engine/Plugins/Runtime/StateTree/Source/StateTreeModule/Public/StateTreePropertyBindings.h
+++ b/Engine/Plugins/Runtime/StateTree/Source/StateTreeModule/Public/StateTreePropertyBindings.h
@@ -599,16 +599,16 @@ struct IStateTreeBindingLookup
 	virtual ~IStateTreeBindingLookup() {} 
 
 	/** @return Source path for given target path, or null if binding does not exists. */
-	UE_API virtual const FPropertyBindingPath* GetPropertyBindingSource(const FPropertyBindingPath& InTargetPath) const PURE_VIRTUAL(IStateTreeBindingLookup::GetPropertyBindingSource, return nullptr; );
+	virtual const FPropertyBindingPath* GetPropertyBindingSource(const FPropertyBindingPath& InTargetPath) const = 0;
 
 	/** @return Display name given property path. */
-	UE_API virtual FText GetPropertyPathDisplayName(const FPropertyBindingPath& InPath, EStateTreeNodeFormatting Formatting = EStateTreeNodeFormatting::Text) const PURE_VIRTUAL(IStateTreeBindingLookup::GetPropertyPathDisplayName, return FText::GetEmpty(); );
+	virtual FText GetPropertyPathDisplayName(const FPropertyBindingPath& InPath, EStateTreeNodeFormatting Formatting = EStateTreeNodeFormatting::Text) const = 0;
 
 	/** @return Leaf property based on property path. */
-	UE_API virtual const FProperty* GetPropertyPathLeafProperty(const FPropertyBindingPath& InPath) const PURE_VIRTUAL(IStateTreeBindingLookup::GetPropertyPathLeafProperty, return nullptr; );
+	virtual const FProperty* GetPropertyPathLeafProperty(const FPropertyBindingPath& InPath) const = 0;
 
 	/** @return Display name of binding source, or empty if binding does not exists. */
-	UE_API virtual FText GetBindingSourceDisplayName(const FPropertyBindingPath& InTargetPath, EStateTreeNodeFormatting Formatting = EStateTreeNodeFormatting::Text) const PURE_VIRTUAL(IStateTreeBindingLookup::GetPropertyPathDisplayName, return FText::GetEmpty(); );
+	virtual FText GetBindingSourceDisplayName(const FPropertyBindingPath& InTargetPath, EStateTreeNodeFormatting Formatting = EStateTreeNodeFormatting::Text) const = 0;
 
 PRAGMA_DISABLE_DEPRECATION_WARNINGS
 
