StateTree: Update documentation for the different GetDataView functions.
#rnx

diff --git a/Engine/Plugins/Runtime/StateTree/Source/StateTreeModule/Public/StateTreeInstanceData.h b/Engine/Plugins/Runtime/StateTree/Source/StateTreeModule/Public/StateTreeInstanceData.h
--- a/Engine/Plugins/Runtime/StateTree/Source/StateTreeModule/Public/StateTreeInstanceData.h
+++ b/Engine/Plugins/Runtime/StateTree/Source/StateTreeModule/Public/StateTreeInstanceData.h
@@ -29,37 +29,36 @@ namespace UE::StateTree::InstanceData
 {
 	namespace Internal
 	{
-		/** @return true if handle source is valid cified handle relative to given frame. */
-		bool IsHandleSourceValid(
+		/** @return whether the handle is valid relative to the given frame. */
+		[[nodiscard]] bool IsHandleSourceValid(
 			const FStateTreeInstanceStorage& InstanceStorage,
 			const FStateTreeExecutionFrame* ParentFrame,
 			const FStateTreeExecutionFrame& CurrentFrame,
 			const FStateTreeDataHandle Handle);
 
 		/** @return data view of the specified handle in temporary instance. */
-		FStateTreeDataView GetTemporaryDataView(
+		[[nodiscard]] FStateTreeDataView GetTemporaryDataView(
 			FStateTreeInstanceStorage& InstanceStorage,
 			const FStateTreeExecutionFrame* ParentFrame,
 			const FStateTreeExecutionFrame& CurrentFrame,
 			const FStateTreeDataHandle Handle);
 	}
 
-	/** @return data view of the specified handle relative to given frame. */
-	UE_API FStateTreeDataView GetDataView(
-		FStateTreeInstanceStorage& InstanceStorage, 
-		FStateTreeInstanceStorage* SharedInstanceStorage, 
-		const FStateTreeExecutionFrame* ParentFrame, 
-		const FStateTreeExecutionFrame& CurrentFrame, 
+	/** @return data view of the specified handle relative to the given frame. */
+	[[nodiscard]] UE_API FStateTreeDataView GetDataView(
+		FStateTreeInstanceStorage& InstanceStorage,
+		FStateTreeInstanceStorage* SharedInstanceStorage,
+		const FStateTreeExecutionFrame* ParentFrame,
+		const FStateTreeExecutionFrame& CurrentFrame,
 		const FStateTreeDataHandle Handle);
 
 	/** @return data view of the specified handle relative to the given frame, or tries to find a matching temporary instance. */
-	UE_API FStateTreeDataView GetDataViewOrTemporary(
+	[[nodiscard]] UE_API FStateTreeDataView GetDataViewOrTemporary(
 		FStateTreeInstanceStorage& InstanceStorage,
 		FStateTreeInstanceStorage* SharedInstanceStorage,
 		const FStateTreeExecutionFrame* ParentFrame,
 		const FStateTreeExecutionFrame& CurrentFrame,
 		const FStateTreeDataHandle Handle);
-
 }
 
 /** Wrapper class to store an object amongst the structs. */
@@ -69,7 +68,9 @@ struct FStateTreeInstanceObjectWrapper
 	GENERATED_BODY()
 
 	FStateTreeInstanceObjectWrapper() = default;
-	FStateTreeInstanceObjectWrapper(UObject* Object) : InstanceObject(Object) {}
+	FStateTreeInstanceObjectWrapper(UObject* Object)
+		: InstanceObject(Object)
+	{}
 	
 	UPROPERTY()
 	TObjectPtr<UObject> InstanceObject = nullptr;
