[AutoRTFM] Update state-tree tracing logic for transactional safety.

diff --git a/Engine/Plugins/Runtime/StateTree/Source/StateTreeModule/Private/Debugger/StateTreeTrace.cpp b/Engine/Plugins/Runtime/StateTree/Source/StateTreeModule/Private/Debugger/StateTreeTrace.cpp
--- a/Engine/Plugins/Runtime/StateTree/Source/StateTreeModule/Private/Debugger/StateTreeTrace.cpp
+++ b/Engine/Plugins/Runtime/StateTree/Source/StateTreeModule/Private/Debugger/StateTreeTrace.cpp
@@ -5,6 +5,8 @@
 #include "Debugger/StateTreeTrace.h"
 #include "Debugger/StateTreeTraceTypes.h"
 #include "Exporters/Exporter.h"
+#include "Misc/TransactionallySafeRWLock.h"
+#include "Misc/MTTransactionallySafeAccessDetector.h"
 #include "ObjectTrace.h"
 #include "Serialization/BufferArchive.h"
 #include "StateTree.h"
@@ -276,17 +278,17 @@ private:
 
 	/** List of asset debug ids events that will be output if channel gets enabled. */
 	TArray<FAssetDebugIdEventBufferedData> AssetDebugIdMap;
-	FRWLock AssetDebugIdMapLock;
+	FTransactionallySafeRWLock AssetDebugIdMapLock;
 
 	/** List of lifetime events that will be output if channel gets enabled in the Push - Pop lifetime window of an instance. */
 	TMap<FStateTreeInstanceDebugId, FInstanceEventBufferedData> BufferedEventsMap;
-	FRWLock BufferedEventsMapLock;
+	FTransactionallySafeRWLock BufferedEventsMapLock;
 
 	/** MT access detector used to validate that AssetDebugIdEvents and InstanceLifetimeEvents are never flushed while getting accessed by worker threads. */
-	UE_MT_DECLARE_RW_ACCESS_DETECTOR(AssetDebugIdMapMTDetector);
+	UE_MT_DECLARE_TS_RW_ACCESS_DETECTOR(AssetDebugIdMapMTDetector);
 
 	/** MT access detector used to validate that AssetDebugIdEvents and InstanceLifetimeEvents are never flushed while getting accessed by worker threads. */
-	UE_MT_DECLARE_RW_ACCESS_DETECTOR(BufferedEventsMapMTDetector);
+	UE_MT_DECLARE_TS_RW_ACCESS_DETECTOR(BufferedEventsMapMTDetector);
 
 	/** Flag use to prevent reentrant calls when buffered events gets flushed when starting a new trace (i.e., OnTracesStarted) */
 	bool bFlushingLifetimeEvents = false;
