[AutoRTFM] Update the atomic instance serial number in the open.

This does not need to be rolled back even if a transaction fails.

diff --git a/Engine/Plugins/Runtime/StateTree/Source/StateTreeModule/Private/StateTreeExecutionContext.cpp b/Engine/Plugins/Runtime/StateTree/Source/StateTreeModule/Private/StateTreeExecutionContext.cpp
--- a/Engine/Plugins/Runtime/StateTree/Source/StateTreeModule/Private/StateTreeExecutionContext.cpp
+++ b/Engine/Plugins/Runtime/StateTree/Source/StateTreeModule/Private/StateTreeExecutionContext.cpp
@@ -9,6 +9,7 @@
 #include "StateTreeDelegate.h"
 #include "StateTreePropertyFunctionBase.h"
 #include "StateTreeReference.h"
+#include "AutoRTFM.h"
 #include "Containers/StaticArray.h"
 #include "Debugger/StateTreeDebug.h"
 #include "Debugger/StateTreeTrace.h"
@@ -578,12 +579,19 @@ FString FStateTreeReadOnlyExecutionContext::GetInstanceDescriptionInternal() con
 }
 
 #if WITH_STATETREE_TRACE
+UE_AUTORTFM_ALWAYS_OPEN
+static uint32 GetNextInstanceSerialNumber()
+{
+	// The instance serial number is only used to synthesize unique instance debug IDs; rollback isn't needed.
+	return ++UE::StateTree::Debug::InstanceSerialNumber;
+}
+
 FStateTreeInstanceDebugId FStateTreeReadOnlyExecutionContext::GetInstanceDebugId() const
 {
 	FStateTreeInstanceDebugId& InstanceDebugId = Storage.GetMutableExecutionState().InstanceDebugId;
 	if (!InstanceDebugId.IsValid())
 	{
-		InstanceDebugId = FStateTreeInstanceDebugId(GetTypeHash(GetInstanceDescriptionInternal()), ++UE::StateTree::Debug::InstanceSerialNumber);
+		InstanceDebugId = FStateTreeInstanceDebugId(GetTypeHash(GetInstanceDescriptionInternal()), GetNextInstanceSerialNumber());
 	}
 	return InstanceDebugId;
 }
