[AutoRTFM] Replace unnecessarily-concurrent linear array allocators with a nonconcurrent equivalent.

These arrays are only used in the context of a single thread and immediately disappear when their stack frame ends. The concurrent linear allocator was a good choice for speed, but its concurrency aspects were unhelpful and caused AutoRTFM to complain about threading hazards. These now use the equivalent nonconcurrent allocator.

diff --git a/Engine/Plugins/Runtime/StateTree/Source/StateTreeModule/Private/StateTreeExecutionContext.cpp b/Engine/Plugins/Runtime/StateTree/Source/StateTreeModule/Private/StateTreeExecutionContext.cpp
--- a/Engine/Plugins/Runtime/StateTree/Source/StateTreeModule/Private/StateTreeExecutionContext.cpp
+++ b/Engine/Plugins/Runtime/StateTree/Source/StateTreeModule/Private/StateTreeExecutionContext.cpp
@@ -1866,10 +1866,10 @@ void FStateTreeExecutionContext::UpdateInstanceData(TConstArrayView<FStateTreeEx
 		}
 	}
 	
-	TArray<FConstStructView, TConcurrentLinearArrayAllocator<FDefaultBlockAllocationTag>> InstanceStructs;
+	TArray<FConstStructView, FNonconcurrentLinearArrayAllocator> InstanceStructs;
 	InstanceStructs.Reserve(EstimatedNumStructs);
 
-	TArray<FInstancedStruct*, TConcurrentLinearArrayAllocator<FDefaultBlockAllocationTag>> TempInstanceStructs;
+	TArray<FInstancedStruct*, FNonconcurrentLinearArrayAllocator> TempInstanceStructs;
 	TempInstanceStructs.Reserve(EstimatedNumStructs);
 
 	TArray<FCompactStateTreeParameters, TFixedAllocator<FStateSelectionResult::MaxExecutionFrames>> TempParams;
@@ -2589,7 +2589,7 @@ EStateTreeRunStatus FStateTreeExecutionContext::EnterState(FStateTreeTransitionR
 	UE_STATETREE_DEBUG_ENTER_PHASE(this, EStateTreeUpdatePhase::EnterStates);
 
 	// The previous active frames are needed for state enter logic.
-	TArray<FStateTreeExecutionFrame, TConcurrentLinearArrayAllocator<FDefaultBlockAllocationTag>> PreviousActiveFrames;
+	TArray<FStateTreeExecutionFrame, FNonconcurrentLinearArrayAllocator> PreviousActiveFrames;
 	PreviousActiveFrames = Exec.ActiveFrames;
 
 	// Reset the current active frames, new ones are added one by one.
@@ -2825,7 +2825,7 @@ void FStateTreeExecutionContext::ExitState(const FStateTreeTransitionResult& Tra
 		bool bShouldCall = false; 
 	};
 
-	TArray<FExitStateCall, TConcurrentLinearArrayAllocator<FDefaultBlockAllocationTag>> ExitStateCalls;
+	TArray<FExitStateCall, FNonconcurrentLinearArrayAllocator> ExitStateCalls;
 
 	// Track any changed state (i.e., not sustained) to prevent reused subtrees from being considered as sustained states since the whole
 	// tree is reused and should be considered as a new instance.
diff --git a/Engine/Plugins/Runtime/StateTree/Source/StateTreeModule/Public/StateTreeExecutionContext.h b/Engine/Plugins/Runtime/StateTree/Source/StateTreeModule/Public/StateTreeExecutionContext.h
--- a/Engine/Plugins/Runtime/StateTree/Source/StateTreeModule/Public/StateTreeExecutionContext.h
+++ b/Engine/Plugins/Runtime/StateTree/Source/StateTreeModule/Public/StateTreeExecutionContext.h
@@ -1113,7 +1113,7 @@ PRAGMA_ENABLE_DEPRECATION_WARNINGS
 	FStateTreeReferenceOverrides LinkedAssetStateTreeOverrides;
 	
 	/** Data view of the context data. */
-	TArray<FStateTreeDataView, TConcurrentLinearArrayAllocator<FDefaultBlockAllocationTag>> ContextAndExternalDataViews;
+	TArray<FStateTreeDataView, FNonconcurrentLinearArrayAllocator> ContextAndExternalDataViews;
 
 	FOnCollectStateTreeExternalData CollectExternalDataDelegate;
 
@@ -1122,7 +1122,7 @@ PRAGMA_ENABLE_DEPRECATION_WARNINGS
 		const UStateTree* StateTree = nullptr;
 		FStateTreeIndex16 BaseIndex;
 	};
-	TArray<FCollectedExternalDataCache, TConcurrentLinearArrayAllocator<FDefaultBlockAllocationTag>> CollectedExternalCache;
+	TArray<FCollectedExternalDataCache, FNonconcurrentLinearArrayAllocator> CollectedExternalCache;
 
 	bool bActiveExternalDataCollected = false;
 
