[AutoRTFM] Update atomic counter value for testing in the open.

This atomic counter is only used in unit tests, so simple fixes are fully sufficient here.

diff --git a/Engine/Plugins/Runtime/StateTree/Source/StateTreeTestSuite/Private/StateTreeTest.cpp b/Engine/Plugins/Runtime/StateTree/Source/StateTreeTestSuite/Private/StateTreeTest.cpp
--- a/Engine/Plugins/Runtime/StateTree/Source/StateTreeTestSuite/Private/StateTreeTest.cpp
+++ b/Engine/Plugins/Runtime/StateTree/Source/StateTreeTestSuite/Private/StateTreeTest.cpp
@@ -9,6 +9,7 @@
 #include "Conditions/StateTreeCommonConditions.h"
 #include "Tasks/StateTreeRunParallelStateTreeTask.h"
 #include "StateTreeTestTypes.h"
+#include "AutoRTFM.h"
 #include "Engine/World.h"
 #include "Async/ParallelFor.h"
 #include "GameplayTagsManager.h"
@@ -703,7 +704,10 @@ struct FStateTreeTest_SharedInstanceData : FStateTreeTestBase
 
 		// Init, nothing should access the shared data.
 		constexpr int32 NumConcurrent = 100;
-		FStateTreeTestConditionInstanceData::GlobalCounter = 0;
+		UE_AUTORTFM_OPEN
+		{
+			FStateTreeTestConditionInstanceData::GlobalCounter = 0;
+		};
 
 		bool bInitSucceeded = true;
 		TArray<FStateTreeInstanceData> InstanceDatas;
@@ -715,7 +719,13 @@ struct FStateTreeTest_SharedInstanceData : FStateTreeTestBase
 			bInitSucceeded &= Exec.IsValid();
 		}
 		AITEST_TRUE(TEXT("All StateTree contexts should init"), bInitSucceeded);
-		AITEST_EQUAL(TEXT("Test condition global counter should be 0"), (int32)FStateTreeTestConditionInstanceData::GlobalCounter, 0);
+
+		int32 GlobalCounterValue;
+		UE_AUTORTFM_OPEN
+		{
+			GlobalCounterValue = FStateTreeTestConditionInstanceData::GlobalCounter;
+		};
+		AITEST_EQUAL(TEXT("Test condition global counter should be 0"), GlobalCounterValue, 0);
 		
 		// Start in parallel
 		// This should create shared data per thread.
@@ -742,8 +752,12 @@ struct FStateTreeTest_SharedInstanceData : FStateTreeTestBase
 			StartTotalRunning += RunContext.Count;
 		}
 		AITEST_EQUAL(TEXT("All StateTree contexts should be running after Start"), StartTotalRunning, NumConcurrent);
-		AITEST_EQUAL(TEXT("Test condition global counter should equal context count after Start"), (int32)FStateTreeTestConditionInstanceData::GlobalCounter, InstanceDatas.Num());
 
+		UE_AUTORTFM_OPEN
+		{
+			GlobalCounterValue = FStateTreeTestConditionInstanceData::GlobalCounter;
+		};
+		AITEST_EQUAL(TEXT("Test condition global counter should equal context count after Start"), GlobalCounterValue, InstanceDatas.Num());
 		
 		// Tick in parallel
 		// This should not recreate the data, so FStateTreeTestConditionInstanceData::GlobalCounter should stay as is.
@@ -772,7 +786,12 @@ struct FStateTreeTest_SharedInstanceData : FStateTreeTestBase
 			TickTotalRunning += RunContext.Count;
 		}
 		AITEST_EQUAL(TEXT("All StateTree contexts should be running after Tick"), TickTotalRunning, NumConcurrent);
-		AITEST_EQUAL(TEXT("Test condition global counter should equal context count after Tick"), (int32)FStateTreeTestConditionInstanceData::GlobalCounter, InstanceDatas.Num());
+
+		UE_AUTORTFM_OPEN
+		{
+			GlobalCounterValue = FStateTreeTestConditionInstanceData::GlobalCounter;
+		};
+		AITEST_EQUAL(TEXT("Test condition global counter should equal context count after Tick"), GlobalCounterValue, InstanceDatas.Num());
 
 		ParallelForWithTaskContext(
 			RunContexts,
diff --git a/Engine/Plugins/Runtime/StateTree/Source/StateTreeTestSuite/Private/StateTreeTestTypes.h b/Engine/Plugins/Runtime/StateTree/Source/StateTreeTestSuite/Private/StateTreeTestTypes.h
--- a/Engine/Plugins/Runtime/StateTree/Source/StateTreeTestSuite/Private/StateTreeTestTypes.h
+++ b/Engine/Plugins/Runtime/StateTree/Source/StateTreeTestSuite/Private/StateTreeTestTypes.h
@@ -11,6 +11,9 @@
 #include "StateTreePropertyRef.h"
 #include "StateTreePropertyFunctionBase.h"
 #include "StateTreeDelegate.h"
+#include "AutoRTFM.h"
+#include <atomic>
+
 #include "StateTreeTestTypes.generated.h"
 
 class UStateTree;
@@ -588,8 +591,15 @@ struct FStateTreeTestCondition : public FStateTreeConditionCommonBase
 	virtual const UStruct* GetInstanceDataType() const override { return FInstanceDataType::StaticStruct(); }
 	virtual bool TestCondition(FStateTreeExecutionContext& Context) const override
 	{
-		FInstanceDataType& InstanceData = Context.GetInstanceData(*this);
-		InstanceData.GlobalCounter.fetch_add(InstanceData.Count);
+		std::atomic<int32>& GlobalCounter = Context.GetInstanceData(*this).GlobalCounter;
+		UE_AUTORTFM_OPEN
+		{
+			++GlobalCounter;
+		};
+		UE_AUTORTFM_ONABORT(&GlobalCounter)
+		{
+			--GlobalCounter;
+		};
 		return bTestConditionResult;
 	}
 
