[State Tree] Dirtied editor compilation hash when changing property binding

diff --git a/Engine/Plugins/Runtime/StateTree/Source/StateTreeEditorModule/Private/StateTreeEditorMode.cpp b/Engine/Plugins/Runtime/StateTree/Source/StateTreeEditorModule/Private/StateTreeEditorMode.cpp
--- a/Engine/Plugins/Runtime/StateTree/Source/StateTreeEditorModule/Private/StateTreeEditorMode.cpp
+++ b/Engine/Plugins/Runtime/StateTree/Source/StateTreeEditorModule/Private/StateTreeEditorMode.cpp
@@ -81,7 +81,7 @@ void UStateTreeEditorMode::Enter()
 	UE::StateTree::Delegates::OnParametersChanged.AddUObject(this, &UStateTreeEditorMode::OnRefreshDetailsView);
 	UE::StateTree::Delegates::OnGlobalDataChanged.AddUObject(this, &UStateTreeEditorMode::OnRefreshDetailsView);
 	UE::StateTree::Delegates::OnStateParametersChanged.AddUObject(this, &UStateTreeEditorMode::OnStateParametersChanged);
-
+	UE::StateTree::PropertyBinding::OnStateTreePropertyBindingChanged.AddUObject(this, &UStateTreeEditorMode::OnPropertyBindingChanged);
 	OnStateTreeChanged();
 }
 
@@ -229,7 +229,7 @@ void UStateTreeEditorMode::Exit()
 	UE::StateTree::Delegates::OnParametersChanged.RemoveAll(this);
 	UE::StateTree::Delegates::OnGlobalDataChanged.RemoveAll(this);
 	UE::StateTree::Delegates::OnStateParametersChanged.RemoveAll(this);
-	
+	UE::StateTree::PropertyBinding::OnStateTreePropertyBindingChanged.RemoveAll(this);
 	Super::Exit();
 }
 
@@ -372,6 +372,11 @@ void UStateTreeEditorMode::BindToolkitCommands(const TSharedRef<FUICommandList>&
 	);
 }
 
+void UStateTreeEditorMode::OnPropertyBindingChanged(const FPropertyBindingPath& SourcePath, const FPropertyBindingPath& TargetPath)
+{
+	UpdateAsset();
+}
+
 void UStateTreeEditorMode::BindCommands()
 {
 	UEdMode::BindCommands();
diff --git a/Engine/Plugins/Runtime/StateTree/Source/StateTreeEditorModule/Public/StateTreeEditorData.h b/Engine/Plugins/Runtime/StateTree/Source/StateTreeEditorModule/Public/StateTreeEditorData.h
--- a/Engine/Plugins/Runtime/StateTree/Source/StateTreeEditorModule/Public/StateTreeEditorData.h
+++ b/Engine/Plugins/Runtime/StateTree/Source/StateTreeEditorModule/Public/StateTreeEditorData.h
@@ -367,7 +367,7 @@ public:
 	UPROPERTY(EditDefaultsOnly, Category = "Global Tasks")
 	EStateTreeTaskCompletionType GlobalTasksCompletion = EStateTreeTaskCompletionType::Any;
 
-	UPROPERTY(meta = (ExcludeFromHash))
+	UPROPERTY()
 	FStateTreeEditorPropertyBindings EditorBindings;
 
 	/** Color Options to assign to a State */
diff --git a/Engine/Plugins/Runtime/StateTree/Source/StateTreeEditorModule/Public/StateTreeEditorMode.h b/Engine/Plugins/Runtime/StateTree/Source/StateTreeEditorModule/Public/StateTreeEditorMode.h
--- a/Engine/Plugins/Runtime/StateTree/Source/StateTreeEditorModule/Public/StateTreeEditorMode.h
+++ b/Engine/Plugins/Runtime/StateTree/Source/StateTreeEditorModule/Public/StateTreeEditorMode.h
@@ -31,6 +31,7 @@ protected:
 	void OnStateTreeChanged();
 	void BindToolkitCommands(const TSharedRef<FUICommandList>& ToolkitCommands);
 
+	void OnPropertyBindingChanged(const FPropertyBindingPath& SourcePath, const FPropertyBindingPath& TargetPath);
 	void OnIdentifierChanged(const UStateTree& InStateTree);
 	void OnSchemaChanged(const UStateTree& InStateTree);
 	void ForceRefreshDetailsView() const;
