[StateTreeDebugger] fix for out of sync scrub state indices
#rnx

diff --git a/Engine/Plugins/Runtime/StateTree/Source/StateTreeModule/Private/Debugger/StateTreeDebuggerTypes.cpp b/Engine/Plugins/Runtime/StateTree/Source/StateTreeModule/Private/Debugger/StateTreeDebuggerTypes.cpp
--- a/Engine/Plugins/Runtime/StateTree/Source/StateTreeModule/Private/Debugger/StateTreeDebuggerTypes.cpp
+++ b/Engine/Plugins/Runtime/StateTree/Source/StateTreeModule/Private/Debugger/StateTreeDebuggerTypes.cpp
@@ -33,9 +33,10 @@ const FInstanceEventCollection FInstanceEventCollection::Invalid;
 //----------------------------------------------------------------//
 // FScrubState
 //----------------------------------------------------------------//
-bool FScrubState::SetScrubTime(const double NewScrubTime)
+bool FScrubState::SetScrubTime(const double NewScrubTime, const bool bForceRefresh)
 {
-	if (NewScrubTime == ScrubTime)
+	if (! bForceRefresh
+		&& NewScrubTime == ScrubTime)
 	{
 		return false;
 	}
@@ -87,10 +88,9 @@ void FScrubState::SetEventCollectionIndex(const int32 InEventCollectionIndex)
 {
 	EventCollectionIndex = InEventCollectionIndex;
 
-	// Force refresh of internal indices with current time applied on the new event collection. 
-	const double PrevScrubTime = ScrubTime;
-	ScrubTime = 0;
-	SetScrubTime(PrevScrubTime);
+	// Force refresh of internal indices with current time applied on the new event collection.
+	constexpr bool bForceRefresh = true;
+	SetScrubTime(ScrubTime, bForceRefresh);
 }
 
 void FScrubState::SetFrameSpanIndex(const int32 NewFrameSpanIndex)
diff --git a/Engine/Plugins/Runtime/StateTree/Source/StateTreeModule/Public/Debugger/StateTreeDebuggerTypes.h b/Engine/Plugins/Runtime/StateTree/Source/StateTreeModule/Public/Debugger/StateTreeDebuggerTypes.h
--- a/Engine/Plugins/Runtime/StateTree/Source/StateTreeModule/Public/Debugger/StateTreeDebuggerTypes.h
+++ b/Engine/Plugins/Runtime/StateTree/Source/StateTreeModule/Public/Debugger/StateTreeDebuggerTypes.h
@@ -191,9 +191,11 @@ struct FScrubState
 
 	/**
 	 * Updates internal indices based on the new time.
+	 * @param NewScrubTime The new scrub time to use to update all internal indices
+	 * @param bForceRefresh Whether to force update of internal indices regardless if the current stored scrub time is the same as the provided one
 	 * @return true if values were updated; false otherwise (i.e. no changes)
 	 */
-	STATETREEMODULE_API bool SetScrubTime(double NewScrubTime);
+	STATETREEMODULE_API bool SetScrubTime(double NewScrubTime, bool bForceRefresh = false);
 
 	/**
 	 * Indicates if the current scrub state points to a valid frame.
