StateTree: Make sure the CurrentState is properly set before TestAllCondition and EvaluateUtility. Make 2 versions of TestAllConditoins.  One that uses the temporary instances and one that uses the valid handles.

diff --git a/Engine/Plugins/Runtime/StateTree/Source/StateTreeModule/Private/StateTreeExecutionContext.cpp b/Engine/Plugins/Runtime/StateTree/Source/StateTreeModule/Private/StateTreeExecutionContext.cpp
--- a/Engine/Plugins/Runtime/StateTree/Source/StateTreeModule/Private/StateTreeExecutionContext.cpp
+++ b/Engine/Plugins/Runtime/StateTree/Source/StateTreeModule/Private/StateTreeExecutionContext.cpp
@@ -2305,7 +2305,6 @@ FStateTreeDataView FStateTreeExecutionContext::GetTemporaryDataView(const FState
 	default:
 		return UE::StateTree::InstanceData::Internal::GetTemporaryDataView(Storage, ParentFrame, CurrentFrame, Handle);
 	}
-
 }
 
 FStateTreeDataView FStateTreeExecutionContext::AddTemporaryInstance(const FStateTreeExecutionFrame& Frame, const FStateTreeIndex16 OwnerNodeIndex, const FStateTreeDataHandle DataHandle, FConstStructView NewInstanceData)
@@ -2364,8 +2363,12 @@ bool FStateTreeExecutionContext::CopyBatchWithValidation(const FStateTreeExecuti
 	return bSucceed;
 }
 
-
 bool FStateTreeExecutionContext::CollectActiveExternalData()
+{
+	return CollectActiveExternalData(GetExecState().ActiveFrames);
+}
+
+bool FStateTreeExecutionContext::CollectActiveExternalData(const TArrayView<FStateTreeExecutionFrame> Frames)
 {
 	if (bActiveExternalDataCollected)
 	{
@@ -2373,10 +2376,9 @@ bool FStateTreeExecutionContext::CollectActiveExternalData()
 	}
 
 	bool bAllExternalDataValid = true;
-	FStateTreeExecutionState& Exec = GetExecState();
 	const FStateTreeExecutionFrame* PrevFrame = nullptr;
 	
-	for (FStateTreeExecutionFrame& Frame : Exec.ActiveFrames)
+	for (FStateTreeExecutionFrame& Frame : Frames)
 	{
 		if (PrevFrame && PrevFrame->StateTree == Frame.StateTree)
 		{
@@ -3117,9 +3119,8 @@ void FStateTreeExecutionContext::StateCompleted()
 				const FCompactStateTreeState& State = CurrentStateTree->States[CurrentHandle.Index];
 
 				FCurrentlyProcessedStateScope StateScope(*this, CurrentHandle);
-				
-				STATETREE_LOG(Verbose, TEXT("%*sState '%s'"), (FrameIndex + StateIndex + 1)*UE::StateTree::Debug::IndentSize, TEXT(""), *GetSafeStateName(CurrentFrame, CurrentHandle));
 				UE_STATETREE_DEBUG_STATE_EVENT(this, CurrentHandle, EStateTreeTraceEventType::OnStateCompleted);
+				STATETREE_LOG(Verbose, TEXT("%*sState '%s'"), (FrameIndex + StateIndex + 1)*UE::StateTree::Debug::IndentSize, TEXT(""), *GetSafeStateName(CurrentFrame, CurrentHandle));
 
 				// Notify Tasks
 				for (int32 TaskIndex = (State.TasksBegin + State.TasksNum) - 1; TaskIndex >= State.TasksBegin; TaskIndex--)
@@ -3934,7 +3935,24 @@ FStateTreeExecutionContext::FTickTaskResult FStateTreeExecutionContext::TickTask
 	return FTickTaskResult{bShouldTickTasks};
 }
 
+// Deprecated
 bool FStateTreeExecutionContext::TestAllConditions(const FStateTreeExecutionFrame* CurrentParentFrame, const FStateTreeExecutionFrame& CurrentFrame, const int32 ConditionsOffset, const int32 ConditionsNum)
+{
+	return TestAllConditionsInternal<false>(CurrentParentFrame, CurrentFrame, FStateTreeStateHandle(), ConditionsOffset, ConditionsNum);
+}
+
+bool FStateTreeExecutionContext::TestAllConditionsOnActiveInstances(const FStateTreeExecutionFrame* CurrentParentFrame, const FStateTreeExecutionFrame& CurrentFrame, FStateTreeStateHandle CurrentStateHandle, const int32 ConditionsOffset, const int32 ConditionsNum)
+{
+	return TestAllConditionsInternal<true>(CurrentParentFrame, CurrentFrame, CurrentStateHandle, ConditionsOffset, ConditionsNum);
+}
+
+bool FStateTreeExecutionContext::TestAllConditionsWithValidation(const FStateTreeExecutionFrame* CurrentParentFrame, const FStateTreeExecutionFrame& CurrentFrame, FStateTreeStateHandle CurrentStateHandle, const int32 ConditionsOffset, const int32 ConditionsNum)
+{
+	return TestAllConditionsInternal<false>(CurrentParentFrame, CurrentFrame, CurrentStateHandle, ConditionsOffset, ConditionsNum);
+}
+
+template<bool bOnActiveInstances>
+bool FStateTreeExecutionContext::TestAllConditionsInternal(const FStateTreeExecutionFrame* CurrentParentFrame, const FStateTreeExecutionFrame& CurrentFrame, FStateTreeStateHandle CurrentStateHandle, const int32 ConditionsOffset, const int32 ConditionsNum)
 {
 	CSV_SCOPED_TIMING_STAT_EXCLUSIVE(StateTree_TestConditions);
 
@@ -3943,6 +3961,9 @@ bool FStateTreeExecutionContext::TestAllConditions(const FStateTreeExecutionFram
 		return true;
 	}
 
+	FCurrentlyProcessedFrameScope FrameScope(*this, CurrentParentFrame, CurrentFrame);
+	FCurrentlyProcessedStateScope StateScope(*this, CurrentStateHandle);
+
 	TStaticArray<EStateTreeExpressionOperand, UE::StateTree::MaxExpressionIndent + 1> Operands(InPlace, EStateTreeExpressionOperand::Copy);
 	TStaticArray<bool, UE::StateTree::MaxExpressionIndent + 1> Values(InPlace, false);
 
@@ -3952,6 +3973,7 @@ bool FStateTreeExecutionContext::TestAllConditions(const FStateTreeExecutionFram
 	{
 		const int32 ConditionIndex = ConditionsOffset + Index;
 		const FStateTreeConditionBase& Cond = CurrentFrame.StateTree->Nodes[ConditionIndex].Get<const FStateTreeConditionBase>();
+		checkf(Cond.InstanceDataHandle.GetSource() == EStateTreeDataSourceType::SharedInstanceData, TEXT("Condition only support SharedInstanceData"));
 		const FStateTreeDataView ConditionInstanceView = GetDataView(CurrentParentFrame, CurrentFrame, Cond.InstanceDataHandle);
 		FNodeInstanceDataScope DataScope(*this, &Cond, ConditionIndex, Cond.InstanceDataHandle, ConditionInstanceView);
 
@@ -3961,8 +3983,10 @@ bool FStateTreeExecutionContext::TestAllConditions(const FStateTreeExecutionFram
 			// Copy bound properties.
 			if (Cond.BindingsBatch.IsValid())
 			{
-				// Use validated copy, since we test in situations where the sources are not always valid (e.g. enter conditions may try to access inactive parent state). 
-				if (!CopyBatchWithValidation(CurrentParentFrame, CurrentFrame, ConditionInstanceView, Cond.BindingsBatch))
+				const bool bBatchCopied = bOnActiveInstances
+					? CopyBatchOnActiveInstances(CurrentParentFrame, CurrentFrame, ConditionInstanceView, Cond.BindingsBatch)
+					: CopyBatchWithValidation(CurrentParentFrame, CurrentFrame, ConditionInstanceView, Cond.BindingsBatch);
+				if (!bBatchCopied)
 				{
 					// If the source data cannot be accessed, the whole expression evaluates to false.
 					UE_STATETREE_DEBUG_CONDITION_EVENT(this, ConditionIndex, ConditionInstanceView, EStateTreeTraceEventType::InternalForcedFailure);
@@ -3971,7 +3995,7 @@ bool FStateTreeExecutionContext::TestAllConditions(const FStateTreeExecutionFram
 					break;
 				}
 			}
-			
+
 			UE_STATETREE_DEBUG_CONDITION_TEST_CONDITION(this, CurrentFrame.StateTree, Index);
 
 			bValue = Cond.TestCondition(*this);
@@ -3994,7 +4018,6 @@ bool FStateTreeExecutionContext::TestAllConditions(const FStateTreeExecutionFram
 		const int32 ClosedParens = FMath::Max(0, -DeltaIndent) + 1;
 
 		// Store the operand to apply when merging higher level down when returning to this level.
-		// @todo: remove this conditions in 5.1, needs resaving existing StateTrees.
 		const EStateTreeExpressionOperand Operand = Index == 0 ? EStateTreeExpressionOperand::Copy : Cond.Operand;
 		Operands[Level] = Operand;
 
@@ -4028,7 +4051,14 @@ bool FStateTreeExecutionContext::TestAllConditions(const FStateTreeExecutionFram
 	return Values[0];
 }
 
-float FStateTreeExecutionContext::EvaluateUtility(const FStateTreeExecutionFrame* CurrentParentFrame, const FStateTreeExecutionFrame& CurrentFrame, const int32 ConsiderationsOffset, const int32 ConsiderationsNum, const float StateWeight)
+//Deprecated
+float FStateTreeExecutionContext::EvaluateUtility(const FStateTreeExecutionFrame* CurrentParentFrame, const FStateTreeExecutionFrame& CurrentFrame, const int32 ConsiderationsBegin, const int32 ConsiderationsNum, const float StateWeight)
+{
+	const FStateTreeStateHandle StateHandle = FStateTreeStateHandle::Invalid;
+	return EvaluateUtilityWithValidation(CurrentParentFrame, CurrentFrame, StateHandle, ConsiderationsBegin, ConsiderationsNum, StateWeight);
+}
+
+float FStateTreeExecutionContext::EvaluateUtilityWithValidation(const FStateTreeExecutionFrame* CurrentParentFrame, const FStateTreeExecutionFrame& CurrentFrame, FStateTreeStateHandle CurrentStateHandle, const int32 ConsiderationsBegin, const int32 ConsiderationsNum, const float StateWeight)
 {
 	// @todo: Tracing support
 	CSV_SCOPED_TIMING_STAT_EXCLUSIVE(StateTree_EvaluateUtility);
@@ -4038,6 +4068,9 @@ float FStateTreeExecutionContext::EvaluateUtility(const FStateTreeExecutionFrame
 		return .0f;
 	}
 
+	FCurrentlyProcessedFrameScope FrameScope(*this, CurrentParentFrame, CurrentFrame);
+	FCurrentlyProcessedStateScope NextStateScope(*this, CurrentStateHandle);
+
 	TStaticArray<EStateTreeExpressionOperand, UE::StateTree::MaxExpressionIndent + 1> Operands(InPlace, EStateTreeExpressionOperand::Copy);
 	TStaticArray<float, UE::StateTree::MaxExpressionIndent + 1> Values(InPlace, false);
 
@@ -4045,8 +4078,9 @@ float FStateTreeExecutionContext::EvaluateUtility(const FStateTreeExecutionFrame
 	float Value = .0f;
 	for (int32 Index = 0; Index < ConsiderationsNum; Index++)
 	{
-		const int32 ConsiderationIndex = ConsiderationsOffset + Index;
+		const int32 ConsiderationIndex = ConsiderationsBegin + Index;
 		const FStateTreeConsiderationBase& Consideration = CurrentFrame.StateTree->Nodes[ConsiderationIndex].Get<const FStateTreeConsiderationBase>();
+		checkf(Consideration.InstanceDataHandle.GetSource() == EStateTreeDataSourceType::SharedInstanceData, TEXT("Consideration only support SharedInstanceData"));
 		const FStateTreeDataView ConsiderationInstanceView = GetDataView(CurrentParentFrame, CurrentFrame, Consideration.InstanceDataHandle);
 		FNodeInstanceDataScope DataScope(*this, &Consideration, ConsiderationIndex, Consideration.InstanceDataHandle, ConsiderationInstanceView);
 
@@ -4057,7 +4091,7 @@ float FStateTreeExecutionContext::EvaluateUtility(const FStateTreeExecutionFrame
 			if (!CopyBatchWithValidation(CurrentParentFrame, CurrentFrame, ConsiderationInstanceView, Consideration.BindingsBatch))
 			{
 				// If the source data cannot be accessed, the whole expression evaluates to zero.
-				Values[0] = .0f;
+				Values[0] = 0.0f;
 				break;
 			}
 		}
@@ -4634,7 +4668,7 @@ bool FStateTreeExecutionContext::TriggerTransitions()
 						FCurrentlyProcessedTransitionEventScope TransitionEventScope(*this, TransitionEvent ? TransitionEvent->Get() : nullptr);
 						UE_STATETREE_DEBUG_TRANSITION_EVENT(this, FStateTreeTransitionSource(CurrentFrame.StateTree, FStateTreeIndex16(TransitionIndex), Transition.State, Transition.Priority), EStateTreeTraceEventType::OnEvaluating);
 						UE_STATETREE_DEBUG_SCOPED_PHASE(this, EStateTreeUpdatePhase::TransitionConditions);
-						bPassed = TestAllConditions(CurrentParentFrame, CurrentFrame, Transition.ConditionsBegin, Transition.ConditionsNum);
+						bPassed = TestAllConditionsOnActiveInstances(CurrentParentFrame, CurrentFrame, Handler.StateHandle, Transition.ConditionsBegin, Transition.ConditionsNum);
 					}
 
 					if (bPassed)
@@ -4833,7 +4867,6 @@ PRAGMA_ENABLE_DEPRECATION_WARNINGS
 								continue;
 							}
 
-
 							const bool bTransitionAccepted = !bIsCurrentStatusStopped 
 								? EnumHasAnyFlags(Transition.Trigger, CompletionTrigger)
 								: Transition.Trigger == EStateTreeTransitionTrigger::OnStateCompleted;
@@ -4843,7 +4876,7 @@ PRAGMA_ENABLE_DEPRECATION_WARNINGS
 								{
 									UE_STATETREE_DEBUG_TRANSITION_EVENT(this, FStateTreeTransitionSource(CurrentFrame.StateTree, FStateTreeIndex16(TransitionIndex), Transition.State, Transition.Priority), EStateTreeTraceEventType::OnEvaluating);
 									UE_STATETREE_DEBUG_SCOPED_PHASE(this, EStateTreeUpdatePhase::TransitionConditions);
-									bPassed = TestAllConditions(CurrentParentFrame, CurrentFrame, Transition.ConditionsBegin, Transition.ConditionsNum);
+									bPassed = TestAllConditionsOnActiveInstances(CurrentParentFrame, CurrentFrame, StateHandle, Transition.ConditionsBegin, Transition.ConditionsNum);
 								}
 
 								if (bPassed)
@@ -5262,7 +5295,7 @@ bool FStateTreeExecutionContext::SelectStateInternal(
 	{
 		// Trying to select non-existing state.
 		STATETREE_LOG(Error, TEXT("%hs: Trying to select invalid state from '%s'.  '%s' using StateTree '%s'."),
-            __FUNCTION__, *GetStateStatusString(Exec), *GetNameSafe(&Owner), *GetFullNameSafe(CurrentFrame.StateTree));
+			__FUNCTION__, *GetStateStatusString(Exec), *GetNameSafe(&Owner), *GetFullNameSafe(CurrentFrame.StateTree));
 		return false;
 	}
 
@@ -5427,7 +5460,7 @@ bool FStateTreeExecutionContext::SelectStateInternal(
 		{
 			// Check that the state can be entered
 			UE_STATETREE_DEBUG_ENTER_PHASE(this, EStateTreeUpdatePhase::EnterConditions);
-			const bool bEnterConditionsPassed = TestAllConditions(CurrentParentFrame, CurrentFrame, NextState.EnterConditionsBegin, NextState.EnterConditionsNum);
+			const bool bEnterConditionsPassed = TestAllConditionsWithValidation(CurrentParentFrame, CurrentFrame, NextStateHandle, NextState.EnterConditionsBegin, NextState.EnterConditionsNum);
 			UE_STATETREE_DEBUG_EXIT_PHASE(this, EStateTreeUpdatePhase::EnterConditions);
 
 			if (!bEnterConditionsPassed)
@@ -5735,7 +5768,7 @@ bool FStateTreeExecutionContext::SelectStateInternal(
 
 						UE_STATETREE_DEBUG_SCOPED_PHASE(this, EStateTreeUpdatePhase::TransitionConditions);
 
-						bTransitionConditionsPassed = TestAllConditions(CurrentParentFrame, CurrentFrame, Transition.ConditionsBegin, Transition.ConditionsNum);
+						bTransitionConditionsPassed = TestAllConditionsWithValidation(CurrentParentFrame, CurrentFrame, NextStateHandle, Transition.ConditionsBegin, Transition.ConditionsNum);
 					}
 
 					if (bTransitionConditionsPassed)
@@ -5853,7 +5886,8 @@ bool FStateTreeExecutionContext::SelectStateInternal(
 					{
 						const uint16 CurrentStateIndex = NextLevelChildStates[Index];
 						const FCompactStateTreeState& CurrentState = CurrentStateTree->States[CurrentStateIndex];
-						const float Score = EvaluateUtility(CurrentParentFrame, CurrentFrame, CurrentState.UtilityConsiderationsBegin, CurrentState.UtilityConsiderationsNum, CurrentState.Weight);
+						const FStateTreeStateHandle CurrentStateHandle = FStateTreeStateHandle(CurrentStateIndex);
+						const float Score = EvaluateUtilityWithValidation(CurrentParentFrame, CurrentFrame, CurrentStateHandle, CurrentState.UtilityConsiderationsBegin, CurrentState.UtilityConsiderationsNum, CurrentState.Weight);
 						if (Score > HighestScore)
 						{
 							HighestScore = Score;
@@ -5903,7 +5937,8 @@ bool FStateTreeExecutionContext::SelectStateInternal(
 				for (uint16 CurrentStateIndex = NextState.ChildrenBegin; CurrentStateIndex < NextState.ChildrenEnd; CurrentStateIndex = CurrentStateTree->States[CurrentStateIndex].GetNextSibling())
 				{
 					const FCompactStateTreeState& CurrentState = CurrentStateTree->States[CurrentStateIndex];
-					const float CurrentStateScore = EvaluateUtility(CurrentParentFrame, CurrentFrame, CurrentState.UtilityConsiderationsBegin, CurrentState.UtilityConsiderationsNum, CurrentState.Weight);
+					const FStateTreeStateHandle CurrentStateHandle = FStateTreeStateHandle(CurrentStateIndex);
+					const float CurrentStateScore = EvaluateUtilityWithValidation(CurrentParentFrame, CurrentFrame, CurrentStateHandle, CurrentState.UtilityConsiderationsBegin, CurrentState.UtilityConsiderationsNum, CurrentState.Weight);
 					NextLevelChildStates.Emplace(CurrentStateIndex, CurrentStateScore);
 					TotalScore += CurrentStateScore;
 				}
diff --git a/Engine/Plugins/Runtime/StateTree/Source/StateTreeModule/Public/StateTreeExecutionContext.h b/Engine/Plugins/Runtime/StateTree/Source/StateTreeModule/Public/StateTreeExecutionContext.h
--- a/Engine/Plugins/Runtime/StateTree/Source/StateTreeModule/Public/StateTreeExecutionContext.h
+++ b/Engine/Plugins/Runtime/StateTree/Source/StateTreeModule/Public/StateTreeExecutionContext.h
@@ -906,22 +906,40 @@ protected:
 	/** Handles transition triggering part of the tick. */
 	UE_API void TickTriggerTransitionsInternal();
 
+private:
+	template<bool bOnActiveInstances>
+	bool TestAllConditionsInternal(const FStateTreeExecutionFrame* CurrentParentFrame, const FStateTreeExecutionFrame& CurrentFrame, FStateTreeStateHandle CurrentStateHandle, const int32 ConditionsOffset, const int32 ConditionsNum);
+
+protected:
+	/**
+	 * Checks all conditions at given range.
+	 * Should be used only on active instances, assumes valid handles and does not consider temporary instances.
+	 * @return True if all conditions pass.
+	 */
+	UE_API bool TestAllConditionsOnActiveInstances(const FStateTreeExecutionFrame* CurrentParentFrame, const FStateTreeExecutionFrame& CurrentFrame, FStateTreeStateHandle CurrentStateHandle, const int32 ConditionsOffset, const int32 ConditionsNum);
+
 	/**
-	 * Checks all conditions at given range
+	 * Checks all conditions at given range.
+	 * This version validates the data handles and looks up temporary instances.
 	 * @return True if all conditions pass.
 	 */
+	UE_API bool TestAllConditionsWithValidation(const FStateTreeExecutionFrame* CurrentParentFrame, const FStateTreeExecutionFrame& CurrentFrame, FStateTreeStateHandle CurrentStateHandle, const int32 ConditionsOffset, const int32 ConditionsNum);
+
+	UE_DEPRECATED(5.7, "Use TestAllConditionsWithValidation or TestAllConditionsOnActiveInstances. TestAllConditionsWithValidation can get the data from a temporary frame in state selection.")
 	UE_API bool TestAllConditions(const FStateTreeExecutionFrame* CurrentParentFrame, const FStateTreeExecutionFrame& CurrentFrame, const int32 ConditionsOffset, const int32 ConditionsNum);
 
 	/**
-	 * Calculate the final score of all considerations at given range
+	 * Calculate the final score of all considerations at given range.
 	 * @return the final score
 	 */
+	UE_API float EvaluateUtilityWithValidation(const FStateTreeExecutionFrame* CurrentParentFrame, const FStateTreeExecutionFrame& CurrentFrame, FStateTreeStateHandle StateHandle, const int32 ConsiderationsBegin, const int32 ConsiderationsNum, const float StateWeight);
+
+	UE_DEPRECATED(5.7, "Use EvaluateUtilityWithValidation. It can get the data from a temporary frame in state selection.")
 	UE_API float EvaluateUtility(const FStateTreeExecutionFrame* CurrentParentFrame, const FStateTreeExecutionFrame& CurrentFrame, const int32 ConsiderationsOffset, const int32 ConsiderationsNum, const float StateWeight);
 
 	/* Evaluate all function at given range. Should be used only on active instances, assumes valid handles and does not consider temporary instances. */
 	UE_API void EvaluatePropertyFunctionsOnActiveInstances(const FStateTreeExecutionFrame* CurrentParentFrame, const FStateTreeExecutionFrame& CurrentFrame, FStateTreeIndex16 FuncsBegin, uint16 FuncsNum);
 
-
 	/* Evaluate all function at given range. This version validates the data handles and looks up temporary instances. */
 	UE_API void EvaluatePropertyFunctionsWithValidation(const FStateTreeExecutionFrame* CurrentParentFrame, const FStateTreeExecutionFrame& CurrentFrame, FStateTreeIndex16 FuncsBegin, uint16 FuncsNum);
 
@@ -1067,8 +1085,15 @@ PRAGMA_ENABLE_DEPRECATION_WARNINGS
 
 	/**
 	 * Collects external data for all StateTrees in active frames.
-	 * @returns true if all external data are set successfully. */
+	 * @returns true if all external data are set successfully.
+	 */
 	UE_API bool CollectActiveExternalData();
+	
+	/**
+	 * Collects external data for all StateTrees in active frames.
+	 * @returns true if all external data are set successfully.
+	 */
+	UE_API bool CollectActiveExternalData(const TArrayView<FStateTreeExecutionFrame> Frames);
 
 	/**
 	 * Collects external data for specific State Tree asset. If the data is already collected, cached index is returned.
@@ -1152,6 +1177,8 @@ PRAGMA_ENABLE_DEPRECATION_WARNINGS
 	struct FCurrentlyProcessedFrameScope
 	{
 		FCurrentlyProcessedFrameScope(FStateTreeExecutionContext& InContext, const FStateTreeExecutionFrame* CurrentParentFrame, const FStateTreeExecutionFrame& CurrentFrame);
+		FCurrentlyProcessedFrameScope(const FCurrentlyProcessedFrameScope&) = delete;
+		FCurrentlyProcessedFrameScope& operator=(const FCurrentlyProcessedFrameScope&) = delete;
 		~FCurrentlyProcessedFrameScope();
 
 	private:
@@ -1177,7 +1204,8 @@ PRAGMA_ENABLE_DEPRECATION_WARNINGS
 			SavedState = Context.CurrentlyProcessedState;
 			Context.CurrentlyProcessedState = State;
 		}
-
+		FCurrentlyProcessedStateScope(const FCurrentlyProcessedStateScope&) = delete;
+		FCurrentlyProcessedStateScope& operator=(const FCurrentlyProcessedStateScope&) = delete;
 		~FCurrentlyProcessedStateScope()
 		{
 			Context.CurrentlyProcessedState = SavedState;
@@ -1186,6 +1214,7 @@ PRAGMA_ENABLE_DEPRECATION_WARNINGS
 	private:
 		FStateTreeExecutionContext& Context;
 		FStateTreeStateHandle SavedState = FStateTreeStateHandle::Invalid;
+		bool bChanged = false;
 	};
 
 	/** Current event we're processing in transition, or invalid if not applicable. */
