StateTree: Hide private functions. Add IsActiveInstanceHandleSourceValid to reuse the same logic from inside the execution context.
#rnx

diff --git a/Engine/Plugins/Runtime/StateTree/Source/StateTreeModule/Private/StateTreeExecutionContext.cpp b/Engine/Plugins/Runtime/StateTree/Source/StateTreeModule/Private/StateTreeExecutionContext.cpp
--- a/Engine/Plugins/Runtime/StateTree/Source/StateTreeModule/Private/StateTreeExecutionContext.cpp
+++ b/Engine/Plugins/Runtime/StateTree/Source/StateTreeModule/Private/StateTreeExecutionContext.cpp
@@ -7,6 +7,7 @@
 #include "StateTreeConditionBase.h"
 #include "StateTreeConsiderationBase.h"
 #include "StateTreeDelegate.h"
+#include "StateTreeInstanceDataHelpers.h"
 #include "StateTreePropertyFunctionBase.h"
 #include "StateTreeReference.h"
 #include "AutoRTFM.h"
@@ -2254,7 +2255,7 @@ bool FStateTreeExecutionContext::IsHandleSourceValid(const FStateTreeExecutionFr
 	}
 
 	default:
-		return UE::StateTree::InstanceData::Internal::IsHandleSourceValid(Storage, ParentFrame, CurrentFrame, Handle);
+		return UE::StateTree::InstanceData::Private::IsHandleSourceValid(Storage, ParentFrame, CurrentFrame, Handle);
 	}
 
 	return false;
@@ -2303,7 +2304,7 @@ FStateTreeDataView FStateTreeExecutionContext::GetTemporaryDataView(const FState
 		return {};
 
 	default:
-		return UE::StateTree::InstanceData::Internal::GetTemporaryDataView(Storage, ParentFrame, CurrentFrame, Handle);
+		return UE::StateTree::InstanceData::Private::GetTemporaryDataView(Storage, ParentFrame, CurrentFrame, Handle);
 	}
 }
 
diff --git a/Engine/Plugins/Runtime/StateTree/Source/StateTreeModule/Private/StateTreeInstanceData.cpp b/Engine/Plugins/Runtime/StateTree/Source/StateTreeModule/Private/StateTreeInstanceData.cpp
--- a/Engine/Plugins/Runtime/StateTree/Source/StateTreeModule/Private/StateTreeInstanceData.cpp
+++ b/Engine/Plugins/Runtime/StateTree/Source/StateTreeModule/Private/StateTreeInstanceData.cpp
@@ -1,6 +1,7 @@
 // Copyright Epic Games, Inc. All Rights Reserved.
 
 #include "StateTreeInstanceData.h"
+#include "StateTreeInstanceDataHelpers.h"
 #include "StateTreeExecutionTypes.h"
 #include "Debugger/StateTreeRuntimeValidationInstanceData.h"
 #include "Serialization/MemoryReader.h"
@@ -97,139 +98,142 @@ FStateTreeInstanceStorage::FStateTreeInstanceStorage()
 FStateTreeInstanceStorage::FStateTreeInstanceStorage() = default;
 #endif
 
-namespace UE::StateTree::InstanceData
+namespace UE::StateTree::InstanceData::Private
 {
-	namespace Internal
+	bool IsActiveInstanceHandleSourceValid(const FStateTreeInstanceStorage& Storage, const FStateTreeExecutionFrame& CurrentFrame, const FStateTreeDataHandle& Handle)
 	{
-		bool IsHandleSourceValid(
-			const FStateTreeInstanceStorage& InstanceStorage,
-			const FStateTreeExecutionFrame* ParentFrame,
-			const FStateTreeExecutionFrame& CurrentFrame,
-			const FStateTreeDataHandle Handle)
-		{
-			// Checks that the instance data is valid for specific handle types.
-			// 
-			// The CurrentFrame may not be yet properly initialized, for that reason we need to check
-			// that the path to the handle makes sense (it's part of the active states) as well as that
-			// we actually have instance data for the handle (index is valid).
-			// 
-			// The (base) indices can be invalid if the frame/state is not entered yet.
-			// For active instance data we need to check that the frame is initialized for a specific state,
-			// as well as that the instance data is initialized.
-
-			switch (Handle.GetSource())
-			{
-			case EStateTreeDataSourceType::None:
-				return true;
+		return CurrentFrame.ActiveInstanceIndexBase.IsValid()
+			&& CurrentFrame.ActiveStates.Contains(Handle.GetState(), CurrentFrame.NumCurrentlyActiveStates)
+			&& Storage.IsValidIndex(CurrentFrame.ActiveInstanceIndexBase.Get() + Handle.GetIndex());
+	}
 
-			case EStateTreeDataSourceType::GlobalInstanceData:
-			case EStateTreeDataSourceType::GlobalInstanceDataObject:
-				return CurrentFrame.GlobalInstanceIndexBase.IsValid()
-					&& InstanceStorage.IsValidIndex(CurrentFrame.GlobalInstanceIndexBase.Get() + Handle.GetIndex());
+	bool IsHandleSourceValid(
+		const FStateTreeInstanceStorage& InstanceStorage,
+		const FStateTreeExecutionFrame* ParentFrame,
+		const FStateTreeExecutionFrame& CurrentFrame,
+		const FStateTreeDataHandle& Handle)
+	{
+		// Checks that the instance data is valid for specific handle types.
+		// 
+		// The CurrentFrame may not be yet properly initialized, for that reason we need to check
+		// that the path to the handle makes sense (it's part of the active states) as well as that
+		// we actually have instance data for the handle (index is valid).
+		// 
+		// The (base) indices can be invalid if the frame/state is not entered yet.
+		// For active instance data we need to check that the frame is initialized for a specific state,
+		// as well as that the instance data is initialized.
 
-			case EStateTreeDataSourceType::ActiveInstanceData:
-			case EStateTreeDataSourceType::ActiveInstanceDataObject:
-			case EStateTreeDataSourceType::StateParameterData:
-				return CurrentFrame.ActiveInstanceIndexBase.IsValid()
-					&& CurrentFrame.ActiveStates.Contains(Handle.GetState(), CurrentFrame.NumCurrentlyActiveStates)
-					&& InstanceStorage.IsValidIndex(CurrentFrame.ActiveInstanceIndexBase.Get() + Handle.GetIndex());
+		switch (Handle.GetSource())
+		{
+		case EStateTreeDataSourceType::None:
+			return true;
 
-			case EStateTreeDataSourceType::SharedInstanceData:
-			case EStateTreeDataSourceType::SharedInstanceDataObject:
-				return true;
+		case EStateTreeDataSourceType::GlobalInstanceData:
+		case EStateTreeDataSourceType::GlobalInstanceDataObject:
+			return CurrentFrame.GlobalInstanceIndexBase.IsValid()
+				&& InstanceStorage.IsValidIndex(CurrentFrame.GlobalInstanceIndexBase.Get() + Handle.GetIndex());
 
-			case EStateTreeDataSourceType::GlobalParameterData:
-				return ParentFrame
-					? IsHandleSourceValid(InstanceStorage, nullptr, *ParentFrame, CurrentFrame.GlobalParameterDataHandle)
-					: CurrentFrame.GlobalParameterDataHandle.IsValid();
+		case EStateTreeDataSourceType::ActiveInstanceData:
+		case EStateTreeDataSourceType::ActiveInstanceDataObject:
+		case EStateTreeDataSourceType::StateParameterData:
+			return IsActiveInstanceHandleSourceValid(InstanceStorage, CurrentFrame, Handle);
+				
+		case EStateTreeDataSourceType::SharedInstanceData:
+		case EStateTreeDataSourceType::SharedInstanceDataObject:
+			return true;
 
-			case EStateTreeDataSourceType::SubtreeParameterData:
-				if (ParentFrame)
+		case EStateTreeDataSourceType::GlobalParameterData:
+			return ParentFrame
+				? IsHandleSourceValid(InstanceStorage, nullptr, *ParentFrame, CurrentFrame.GlobalParameterDataHandle)
+				: CurrentFrame.GlobalParameterDataHandle.IsValid();
+
+		case EStateTreeDataSourceType::SubtreeParameterData:
+			if (ParentFrame)
+			{
+				// If the current subtree state is not instantiated yet, we cannot assume that the parameter data is instantiated in the parent frame either. 
+				if (!CurrentFrame.ActiveInstanceIndexBase.IsValid())
 				{
-					// If the current subtree state is not instantiated yet, we cannot assume that the parameter data is instantiated in the parent frame either. 
-					if (!CurrentFrame.ActiveInstanceIndexBase.IsValid())
-					{
-						return false;
-					}
-					// Linked subtree, params defined in parent scope.
-					return IsHandleSourceValid(InstanceStorage, nullptr, *ParentFrame, CurrentFrame.StateParameterDataHandle);
+					return false;
 				}
-				// Standalone subtree, params define as state params.
-				return CurrentFrame.ActiveInstanceIndexBase.IsValid()
-					&& CurrentFrame.ActiveStates.Contains(Handle.GetState(), CurrentFrame.NumCurrentlyActiveStates)
-					&& InstanceStorage.IsValidIndex(CurrentFrame.ActiveInstanceIndexBase.Get() + Handle.GetIndex());
-
-			default:
-				checkf(false, TEXT("Unhandled case or unsupported type for InstanceDataStorage %s"), *UEnum::GetValueAsString(Handle.GetSource()));
+				// Linked subtree, params defined in parent scope.
+				return IsHandleSourceValid(InstanceStorage, nullptr, *ParentFrame, CurrentFrame.StateParameterDataHandle);
 			}
+			// Standalone subtree, params define as state params.
+			return IsActiveInstanceHandleSourceValid(InstanceStorage, CurrentFrame, Handle);
 
-			return false;
+		default:
+			checkf(false, TEXT("Unhandled case or unsupported type for InstanceDataStorage %s"), *UEnum::GetValueAsString(Handle.GetSource()));
 		}
 
-		FStateTreeDataView GetTemporaryDataView(
-			FStateTreeInstanceStorage& InstanceStorage, 
-			const FStateTreeExecutionFrame* ParentFrame, 
-			const FStateTreeExecutionFrame& CurrentFrame, 
-			const FStateTreeDataHandle Handle)
-		{
-			switch (Handle.GetSource())
-			{
-			case EStateTreeDataSourceType::GlobalInstanceData:
-			case EStateTreeDataSourceType::ActiveInstanceData:
-				return InstanceStorage.GetMutableTemporaryStruct(CurrentFrame, Handle);
+		return false;
+	}
 
-			case EStateTreeDataSourceType::GlobalInstanceDataObject:
-			case EStateTreeDataSourceType::ActiveInstanceDataObject:
-				return InstanceStorage.GetMutableTemporaryObject(CurrentFrame, Handle);
+	FStateTreeDataView GetTemporaryDataView(
+		FStateTreeInstanceStorage& InstanceStorage,
+		const FStateTreeExecutionFrame* ParentFrame,
+		const FStateTreeExecutionFrame& CurrentFrame,
+		const FStateTreeDataHandle& Handle)
+	{
+		switch (Handle.GetSource())
+		{
+		case EStateTreeDataSourceType::GlobalInstanceData:
+		case EStateTreeDataSourceType::ActiveInstanceData:
+			return InstanceStorage.GetMutableTemporaryStruct(CurrentFrame, Handle);
 
-			case EStateTreeDataSourceType::GlobalParameterData:
-				if (ParentFrame)
-				{
-					if (FCompactStateTreeParameters* Params = InstanceStorage.GetMutableTemporaryStruct(*ParentFrame, CurrentFrame.GlobalParameterDataHandle).GetPtr<FCompactStateTreeParameters>())
-					{
-						return Params->Parameters.GetMutableValue();
-					}
-				}
-				break;
+		case EStateTreeDataSourceType::GlobalInstanceDataObject:
+		case EStateTreeDataSourceType::ActiveInstanceDataObject:
+			return InstanceStorage.GetMutableTemporaryObject(CurrentFrame, Handle);
 
-			case EStateTreeDataSourceType::SubtreeParameterData:
-				if (ParentFrame)
-				{
-					// Linked subtree, params defined in parent scope.
-					if (FCompactStateTreeParameters* Params = InstanceStorage.GetMutableTemporaryStruct(*ParentFrame, CurrentFrame.StateParameterDataHandle).GetPtr<FCompactStateTreeParameters>())
-					{
-						return Params->Parameters.GetMutableValue();
-					}
-				}
-				// Standalone subtree, params define as state params.
-				if (FCompactStateTreeParameters* Params = InstanceStorage.GetMutableTemporaryStruct(CurrentFrame, Handle).GetPtr<FCompactStateTreeParameters>())
+		case EStateTreeDataSourceType::GlobalParameterData:
+			if (ParentFrame)
+			{
+				if (FCompactStateTreeParameters* Params = InstanceStorage.GetMutableTemporaryStruct(*ParentFrame, CurrentFrame.GlobalParameterDataHandle).GetPtr<FCompactStateTreeParameters>())
 				{
 					return Params->Parameters.GetMutableValue();
 				}
-				break;
+			}
+			break;
 
-			case EStateTreeDataSourceType::StateParameterData:
-				if (FCompactStateTreeParameters* Params = InstanceStorage.GetMutableTemporaryStruct(CurrentFrame, Handle).GetPtr<FCompactStateTreeParameters>())
+		case EStateTreeDataSourceType::SubtreeParameterData:
+			if (ParentFrame)
+			{
+				// Linked subtree, params defined in parent scope.
+				if (FCompactStateTreeParameters* Params = InstanceStorage.GetMutableTemporaryStruct(*ParentFrame, CurrentFrame.StateParameterDataHandle).GetPtr<FCompactStateTreeParameters>())
 				{
 					return Params->Parameters.GetMutableValue();
 				}
-				break;
+			}
+			// Standalone subtree, params define as state params.
+			if (FCompactStateTreeParameters* Params = InstanceStorage.GetMutableTemporaryStruct(CurrentFrame, Handle).GetPtr<FCompactStateTreeParameters>())
+			{
+				return Params->Parameters.GetMutableValue();
+			}
+			break;
 
-			default:
-				checkf(false, TEXT("Unhandled case or unsupported type for InstanceDataStorage %s"), *UEnum::GetValueAsString(Handle.GetSource()));
+		case EStateTreeDataSourceType::StateParameterData:
+			if (FCompactStateTreeParameters* Params = InstanceStorage.GetMutableTemporaryStruct(CurrentFrame, Handle).GetPtr<FCompactStateTreeParameters>())
+			{
+				return Params->Parameters.GetMutableValue();
 			}
+			break;
 
-			return {};
+		default:
+			checkf(false, TEXT("Unhandled case or unsupported type for InstanceDataStorage %s"), *UEnum::GetValueAsString(Handle.GetSource()));
 		}
 
+		return {};
 	}
 
+} // namespace
+
+namespace UE::StateTree::InstanceData
+{
 	FStateTreeDataView GetDataView(
-		FStateTreeInstanceStorage& InstanceStorage, 
+		FStateTreeInstanceStorage& InstanceStorage,
 		FStateTreeInstanceStorage* SharedInstanceStorage,
-		const FStateTreeExecutionFrame* ParentFrame, 
-		const FStateTreeExecutionFrame& CurrentFrame, 
-		const FStateTreeDataHandle Handle)
+		const FStateTreeExecutionFrame* ParentFrame,
+		const FStateTreeExecutionFrame& CurrentFrame,
+		const FStateTreeDataHandle& Handle)
 	{
 		switch (Handle.GetSource())
 		{
@@ -296,18 +300,21 @@ namespace UE::StateTree::InstanceData
 		return {};
 	}
 
-	FStateTreeDataView GetDataViewOrTemporary(FStateTreeInstanceStorage& InstanceStorage, FStateTreeInstanceStorage* SharedInstanceStorage,
-		const FStateTreeExecutionFrame* ParentFrame, const FStateTreeExecutionFrame& CurrentFrame, const FStateTreeDataHandle Handle)
+	FStateTreeDataView GetDataViewOrTemporary(
+		FStateTreeInstanceStorage& InstanceStorage,
+		FStateTreeInstanceStorage* SharedInstanceStorage,
+		const FStateTreeExecutionFrame* ParentFrame,
+		const FStateTreeExecutionFrame& CurrentFrame,
+		const FStateTreeDataHandle& Handle)
 	{
-		using namespace Internal;
-		if (IsHandleSourceValid(InstanceStorage, ParentFrame, CurrentFrame, Handle))
+		if (Private::IsHandleSourceValid(InstanceStorage, ParentFrame, CurrentFrame, Handle))
 		{
 			return GetDataView(InstanceStorage, SharedInstanceStorage, ParentFrame, CurrentFrame, Handle);
 		}
 
-		return GetTemporaryDataView(InstanceStorage, ParentFrame, CurrentFrame, Handle);
+		return Private::GetTemporaryDataView(InstanceStorage, ParentFrame, CurrentFrame, Handle);
 	}
-}
+} // namespace
 
 FStateTreeInstanceStorage::FStateTreeInstanceStorage(const FStateTreeInstanceStorage& Other)
 	: InstanceStructs(Other.InstanceStructs)
diff --git a/Engine/Plugins/Runtime/StateTree/Source/StateTreeModule/Private/StateTreeInstanceDataHelpers.h b/Engine/Plugins/Runtime/StateTree/Source/StateTreeModule/Private/StateTreeInstanceDataHelpers.h
new file mode 100644
--- /dev/null
+++ b/Engine/Plugins/Runtime/StateTree/Source/StateTreeModule/Private/StateTreeInstanceDataHelpers.h
@@ -0,0 +1,29 @@
+// Copyright Epic Games, Inc. All Rights Reserved.
+
+#pragma once
+
+#include "StateTreeInstanceData.h"
+
+namespace UE::StateTree::InstanceData::Private
+{
+	/** @return whether the handle is valid when it's an instance data type. */
+	[[nodiscard]] bool IsActiveInstanceHandleSourceValid(
+		const FStateTreeInstanceStorage& Storage,
+		const FStateTreeExecutionFrame& CurrentFrame,
+		const FStateTreeDataHandle& Handle);
+
+	/** @return whether the handle is valid relative to the given frame. */
+	[[nodiscard]] bool IsHandleSourceValid(
+		const FStateTreeInstanceStorage& InstanceStorage,
+		const FStateTreeExecutionFrame* ParentFrame,
+		const FStateTreeExecutionFrame& CurrentFrame,
+		const FStateTreeDataHandle& Handle);
+
+	/** @return data view of the specified handle in temporary instance. */
+	[[nodiscard]] FStateTreeDataView GetTemporaryDataView(
+		FStateTreeInstanceStorage& InstanceStorage,
+		const FStateTreeExecutionFrame* ParentFrame,
+		const FStateTreeExecutionFrame& CurrentFrame,
+		const FStateTreeDataHandle& Handle);
+} // namespace
+
diff --git a/Engine/Plugins/Runtime/StateTree/Source/StateTreeModule/Public/StateTreeInstanceData.h b/Engine/Plugins/Runtime/StateTree/Source/StateTreeModule/Public/StateTreeInstanceData.h
--- a/Engine/Plugins/Runtime/StateTree/Source/StateTreeModule/Public/StateTreeInstanceData.h
+++ b/Engine/Plugins/Runtime/StateTree/Source/StateTreeModule/Public/StateTreeInstanceData.h
@@ -27,30 +27,13 @@ namespace UE::StateTree::Debug
 
 namespace UE::StateTree::InstanceData
 {
-	namespace Internal
-	{
-		/** @return whether the handle is valid relative to the given frame. */
-		[[nodiscard]] bool IsHandleSourceValid(
-			const FStateTreeInstanceStorage& InstanceStorage,
-			const FStateTreeExecutionFrame* ParentFrame,
-			const FStateTreeExecutionFrame& CurrentFrame,
-			const FStateTreeDataHandle Handle);
-
-		/** @return data view of the specified handle in temporary instance. */
-		[[nodiscard]] FStateTreeDataView GetTemporaryDataView(
-			FStateTreeInstanceStorage& InstanceStorage,
-			const FStateTreeExecutionFrame* ParentFrame,
-			const FStateTreeExecutionFrame& CurrentFrame,
-			const FStateTreeDataHandle Handle);
-	}
-
 	/** @return data view of the specified handle relative to the given frame. */
 	[[nodiscard]] UE_API FStateTreeDataView GetDataView(
 		FStateTreeInstanceStorage& InstanceStorage,
 		FStateTreeInstanceStorage* SharedInstanceStorage,
 		const FStateTreeExecutionFrame* ParentFrame,
 		const FStateTreeExecutionFrame& CurrentFrame,
-		const FStateTreeDataHandle Handle);
+		const FStateTreeDataHandle& Handle);
 
 	/** @return data view of the specified handle relative to the given frame, or tries to find a matching temporary instance. */
 	[[nodiscard]] UE_API FStateTreeDataView GetDataViewOrTemporary(
@@ -58,7 +41,7 @@ namespace UE::StateTree::InstanceData
 		FStateTreeInstanceStorage* SharedInstanceStorage,
 		const FStateTreeExecutionFrame* ParentFrame,
 		const FStateTreeExecutionFrame& CurrentFrame,
-		const FStateTreeDataHandle Handle);
+		const FStateTreeDataHandle& Handle);
 }
 
 /** Wrapper class to store an object amongst the structs. */
