StateTree: Update the unit test with transition. Adding cases with bugs to more easily identify what needs to be fixed.

diff --git a/Engine/Plugins/Runtime/StateTree/Source/StateTreeTestSuite/Private/StateTreeBindingTest.cpp b/Engine/Plugins/Runtime/StateTree/Source/StateTreeTestSuite/Private/StateTreeBindingTest.cpp
--- a/Engine/Plugins/Runtime/StateTree/Source/StateTreeTestSuite/Private/StateTreeBindingTest.cpp
+++ b/Engine/Plugins/Runtime/StateTree/Source/StateTreeTestSuite/Private/StateTreeBindingTest.cpp
@@ -223,8 +223,6 @@ struct FStateTreeTest_PropertyFunctions : FStateTreeTestBase
 		AITEST_TRUE(*FString::Printf(TEXT("StateTree TaskB should exit state with value %d"), ParameterPropertyFunctionsAmount), Exec.Expect(TaskB.GetName(), *FString::Printf(TEXT("ExitState%d"), ParameterPropertyFunctionsAmount)));
 		Exec.LogClear();
 
-		Exec.Stop();
-
 		return true;
 	}
 };
@@ -715,6 +713,7 @@ struct FStateTreeTest_TransitionTaskWithBinding : FStateTreeTestBase
 				Status = Exec.Start(&GlobalParameters);
 			}
 			AITEST_EQUAL(TEXT("Start should complete with Running"), Status, EStateTreeRunStatus::Running);
+			AITEST_TRUE(TEXT("In correct states"), Exec.ExpectInActiveStates("Tree1StateRoot", "Tree1StateA"));
 			AITEST_TRUE(TEXT("Start should enter Global tasks"), Exec.Expect("Tree1GlobalTaskA", TEXT("EnterState99"))
 				.Then("Tree1GlobalTaskB", TEXT("EnterState99"))
 				.Then("Tree1GlobalTaskC", TEXT("EnterState99"))
@@ -732,6 +731,7 @@ struct FStateTreeTest_TransitionTaskWithBinding : FStateTreeTestBase
 
 			Status = Exec.Tick(1.0f);
 			AITEST_EQUAL(TEXT("2nd Tick should complete with Running"), Status, EStateTreeRunStatus::Running);
+			AITEST_TRUE(TEXT("In correct states"), Exec.ExpectInActiveStates("Tree1StateRoot", "Tree1StateA"));
 			AITEST_TRUE(TEXT("2nd Tick should tick tasks"), Exec.Expect("Tree1GlobalTaskA", TEXT("Tick88"))
 				.Then("Tree1GlobalTaskB", TEXT("Tick88"))
 				.Then("Tree1StateRootTaskA", TEXT("Tick88"))
diff --git a/Engine/Plugins/Runtime/StateTree/Source/StateTreeTestSuite/Private/StateTreeDelegateTest.cpp b/Engine/Plugins/Runtime/StateTree/Source/StateTreeTestSuite/Private/StateTreeDelegateTest.cpp
--- a/Engine/Plugins/Runtime/StateTree/Source/StateTreeTestSuite/Private/StateTreeDelegateTest.cpp
+++ b/Engine/Plugins/Runtime/StateTree/Source/StateTreeTestSuite/Private/StateTreeDelegateTest.cpp
@@ -41,6 +41,8 @@ struct FStateTreeTest_Delegate_ConcurrentListeners : FStateTreeTestBase
 		AITEST_TRUE("StateTree should init", bInitSucceeded);
 
 		Exec.Start();
+		AITEST_FALSE(TEXT("StateTree ListenerTaskA should not trigger."), Exec.Expect(ListenerTaskA.GetName(), *FString::Printf(TEXT("OnDelegate%d"), 1)));
+		AITEST_FALSE(TEXT("StateTree ListenerTaskB should not trigger."), Exec.Expect(ListenerTaskB.GetName(), *FString::Printf(TEXT("OnDelegate%d"), 1)));
 		Exec.LogClear();
 
 		Exec.Tick(0.1f);
diff --git a/Engine/Plugins/Runtime/StateTree/Source/StateTreeTestSuite/Private/StateTreeLinkedStateTest.cpp b/Engine/Plugins/Runtime/StateTree/Source/StateTreeTestSuite/Private/StateTreeLinkedStateTest.cpp
--- a/Engine/Plugins/Runtime/StateTree/Source/StateTreeTestSuite/Private/StateTreeLinkedStateTest.cpp
+++ b/Engine/Plugins/Runtime/StateTree/Source/StateTreeTestSuite/Private/StateTreeLinkedStateTest.cpp
@@ -115,17 +115,26 @@ struct FStateTreeTest_EnterAndExitLinkedAsset : FStateTreeTestBase
 	{
 		FStateTreeCompilerLog Log;
 
+		// Tree1
+		//  Root1
+		//   A1 (linked to Tree2) OnCompleted->B1
+		//   B1 Task1
+		// Tree2 GlobalTask2 => 2 ticks
+		//  Root2 Task2 => 1 tick
+
 		// Asset 2
 		UStateTree& StateTree2 = NewStateTree();
 		UStateTreeEditorData& EditorData2 = *Cast<UStateTreeEditorData>(StateTree2.EditorData);
 		UStateTreeState& Root2 = EditorData2.AddSubTree(FName(TEXT("Root2")));
 		TStateTreeEditorNode<FTestTask_Stand>& Task2 = Root2.AddTask<FTestTask_Stand>(FName(TEXT("Task2")));
+		Task2.GetNode().TicksToCompletion = 1;
 		TStateTreeEditorNode<FTestTask_Stand>& GlobalTask2 = EditorData2.AddGlobalTask<FTestTask_Stand>(FName(TEXT("GlobalTask2")));
 		GlobalTask2.GetNode().TicksToCompletion = 2;
-
-		FStateTreeCompiler Compiler2(Log);
-		const bool bResult2 = Compiler2.Compile(StateTree2);
-		AITEST_TRUE(TEXT("StateTree2 should get compiled"), bResult2);
+		{
+			FStateTreeCompiler Compiler2(Log);
+			const bool bResult2 = Compiler2.Compile(StateTree2);
+			AITEST_TRUE(TEXT("StateTree2 should get compiled"), bResult2);
+		}
 
 		// Main asset
 		UStateTree& StateTree = NewStateTree();
@@ -138,38 +147,230 @@ struct FStateTreeTest_EnterAndExitLinkedAsset : FStateTreeTestBase
 
 		UStateTreeState& B1 = Root.AddChildState(FName(TEXT("B1")), EStateTreeStateType::State);
 		TStateTreeEditorNode<FTestTask_Stand>& Task1 = B1.AddTask<FTestTask_Stand>(FName(TEXT("Task1")));
+		Task1.GetNode().TicksToCompletion = 1;
+		{
+			FStateTreeCompiler Compiler(Log);
+			const bool bResult = Compiler.Compile(StateTree);
+			AITEST_TRUE(TEXT("StateTree should get compiled"), bResult);
+		}
 
-		FStateTreeCompiler Compiler(Log);
-		const bool bResult = Compiler.Compile(StateTree);
-		AITEST_TRUE(TEXT("StateTree should get compiled"), bResult);
-
-		const TCHAR* EnterStateStr(TEXT("EnterState"));
-		const TCHAR* ExitStateStr(TEXT("ExitState"));
+		const TCHAR* EnterStateStr = TEXT("EnterState");
+		const TCHAR* SustainedEnterStateStr = TEXT("EnterState=Sustained");
+		const TCHAR* ChangedEnterStateStr = TEXT("EnterState=Changed");
+		const TCHAR* ExitStateStr = TEXT("ExitState");
+		const TCHAR* SustainedExitStateStr = TEXT("ExitState=Sustained");
+		const TCHAR* ChangedExitStateStr = TEXT("ExitState=Changed");
+		const TCHAR* StateCompletedStateStr = TEXT("StateCompleted");
+		const TCHAR* TickStr = TEXT("Tick");
 
 		{
-			EStateTreeRunStatus Status = EStateTreeRunStatus::Unset;
 			FStateTreeInstanceData InstanceData;
 			FTestStateTreeExecutionContext Exec(StateTree, StateTree, InstanceData);
 			const bool bInitSucceeded = Exec.IsValid();
 			AITEST_TRUE(TEXT("StateTree should init"), bInitSucceeded);
+			{
+				EStateTreeRunStatus Status = Exec.Start();
+				AITEST_EQUAL("Start should complete with Running", Status, EStateTreeRunStatus::Running);
+				AITEST_TRUE("StateTree should enter GlobalTask2", Exec.Expect(GlobalTask2.GetName(), ChangedEnterStateStr));
+				AITEST_FALSE("StateTree should not exit GlobalTask2", Exec.Expect(GlobalTask2.GetName(), ChangedExitStateStr));
+				AITEST_FALSE("StateTree should not exit GlobalTask2", Exec.Expect(GlobalTask2.GetName(), SustainedExitStateStr));
+				AITEST_TRUE("StateTree should enter Task2", Exec.Expect(Task2.GetName(), ChangedEnterStateStr));
+				AITEST_FALSE("StateTree should not exit Task2", Exec.Expect(Task2.GetName(), ChangedExitStateStr));
+				AITEST_FALSE("StateTree should not exit Task2", Exec.Expect(Task2.GetName(), SustainedExitStateStr));
+				AITEST_FALSE("StateTree should not enter Task1", Exec.Expect(Task1.GetName(), EnterStateStr));
+				Exec.LogClear();
+			}
+			{
+				// Task2 completes.
+				EStateTreeRunStatus Status = Exec.Tick(0.1f);
+				AITEST_EQUAL("Tick should complete with Running", Status, EStateTreeRunStatus::Running);
+				AITEST_TRUE(TEXT("StateTree should exit the root state but not the global"),
+					Exec.Expect(GlobalTask2.GetName(), TickStr)
+					.Then(Task2.GetName(), TickStr)
+					.Then(Task2.GetName(), StateCompletedStateStr)
+					.Then(Task2.GetName(), ChangedExitStateStr)
+					.Then(GlobalTask2.GetName(), ChangedExitStateStr)
+					.Then(Task1.GetName(), ChangedEnterStateStr)
+				);
+				AITEST_FALSE("StateTree should not tick Tasks1", Exec.Expect(Task1.GetName(), TickStr));
+				AITEST_FALSE("StateTree should not exit Task1", Exec.Expect(Task1.GetName(), ExitStateStr));
+				Exec.LogClear();
+			}
+			{
+				Exec.Stop();
+			}
+			// Change the order, the global completes before task2. Task2 should not tick
+			{
+				Task2.GetNode().TicksToCompletion = 1;
+				GlobalTask2.GetNode().TicksToCompletion = 1;
+				FStateTreeCompiler Compiler2(Log);
+				const bool bResult2 = Compiler2.Compile(StateTree2);
+				AITEST_TRUE(TEXT("StateTree2 should get compiled"), bResult2);
+			}
+			{
+				EStateTreeRunStatus Status = Exec.Start();
+				AITEST_EQUAL("Start should complete with Running", Status, EStateTreeRunStatus::Running);
+				Exec.LogClear();
+			}
+			{
+				// GlobalTask2 completes Task2 won't tick and won't complete.
+				EStateTreeRunStatus Status = Exec.Tick(0.1f);
+				AITEST_EQUAL("Tick should complete with Running", Status, EStateTreeRunStatus::Running);
+				AITEST_TRUE(TEXT("StateTree should exit the root state and the global"),
+					Exec.Expect(GlobalTask2.GetName(), TickStr)
+					.Then(Task2.GetName(), ChangedExitStateStr)
+					.Then(GlobalTask2.GetName(), ChangedExitStateStr)
+					.Then(Task1.GetName(), EnterStateStr)
+				);
+				AITEST_FALSE("StateTree should not tick Tasks2", Exec.Expect(Task2.GetName(), TickStr));
+				AITEST_FALSE("StateTree should not tick Tasks1", Exec.Expect(Task1.GetName(), TickStr));
+				AITEST_FALSE("StateTree should not exit Task1", Exec.Expect(Task1.GetName(), ExitStateStr));
+				Exec.LogClear();
+			}
 
-			Status = Exec.Start();
-			AITEST_EQUAL(TEXT("Start should complete with Running"), Status, EStateTreeRunStatus::Running);
-			AITEST_TRUE(TEXT("StateTree should enter GlobalTask2"), Exec.Expect(GlobalTask2.GetName(), EnterStateStr));
-			AITEST_FALSE(TEXT("StateTree should not exit GlobalTask2"), Exec.Expect(GlobalTask2.GetName(), ExitStateStr));
-			AITEST_TRUE(TEXT("StateTree should enter Task2"), Exec.Expect(Task2.GetName(), EnterStateStr));
-			AITEST_FALSE(TEXT("StateTree should not exit Task2"), Exec.Expect(Task2.GetName(), ExitStateStr));
-			AITEST_FALSE(TEXT("StateTree should not enter Task1"), Exec.Expect(Task1.GetName(), EnterStateStr));
-			Exec.LogClear();
+			Exec.Stop();
+		}
 
-			Status = Exec.Tick(0.1f);
-			AITEST_EQUAL(TEXT("Tick should complete with Running"), Status, EStateTreeRunStatus::Running);
-			AITEST_FALSE(TEXT("StateTree should not enter GlobalTask2"), Exec.Expect(GlobalTask2.GetName(), EnterStateStr));
-			AITEST_TRUE(TEXT("StateTree should exit GlobalTask2"), Exec.Expect(GlobalTask2.GetName(), ExitStateStr));
-			AITEST_FALSE(TEXT("StateTree should not enter Task2"), Exec.Expect(Task2.GetName(), EnterStateStr));
-			AITEST_TRUE(TEXT("StateTree should exit Task2"), Exec.Expect(Task2.GetName(), ExitStateStr));
-			AITEST_TRUE(TEXT("StateTree should enter Task1"), Exec.Expect(Task1.GetName(), EnterStateStr));
-			Exec.LogClear();
+		return true;
+	}
+};
+IMPLEMENT_AI_INSTANT_TEST(FStateTreeTest_EnterAndExitLinkedAsset, "System.StateTree.LinkedAsset.EnterAndExit");
+
+struct FStateTreeTest_EnterAndExitLinkedAsset2 : FStateTreeTestBase
+{
+	virtual bool InstantTest() override
+	{
+		FStateTreeCompilerLog Log;
+
+		// Tree1
+		//  Root1 Task0
+		//   A1 (linked to Tree2) no transition
+		//   B1 Task1
+		// Tree2 GlobalTask2 => 2 ticks
+		//  Root2 Task2 => 1 tick
+
+		// Asset 2
+		UStateTree& StateTree2 = NewStateTree();
+		UStateTreeEditorData& EditorData2 = *Cast<UStateTreeEditorData>(StateTree2.EditorData);
+		UStateTreeState& Root2 = EditorData2.AddSubTree(FName(TEXT("Root2")));
+		TStateTreeEditorNode<FTestTask_Stand>& Task2 = Root2.AddTask<FTestTask_Stand>(FName(TEXT("Task2")));
+		Task2.GetNode().TicksToCompletion = 1;
+		TStateTreeEditorNode<FTestTask_Stand>& GlobalTask2 = EditorData2.AddGlobalTask<FTestTask_Stand>(FName(TEXT("GlobalTask2")));
+		GlobalTask2.GetNode().TicksToCompletion = 2;
+		{
+			FStateTreeCompiler Compiler2(Log);
+			const bool bResult2 = Compiler2.Compile(StateTree2);
+			AITEST_TRUE(TEXT("StateTree2 should get compiled"), bResult2);
+		}
+
+		// Main asset
+		UStateTree& StateTree = NewStateTree();
+		UStateTreeEditorData& EditorData = *Cast<UStateTreeEditorData>(StateTree.EditorData);
+
+		UStateTreeState& Root = EditorData.AddSubTree(FName(TEXT("Root1")));
+		TStateTreeEditorNode<FTestTask_Stand>& Task0 = Root.AddTask<FTestTask_Stand>(FName(TEXT("Task0")));
+		Task0.GetNode().TicksToCompletion = 999999;
+
+		UStateTreeState& A1 = Root.AddChildState(FName(TEXT("A1")), EStateTreeStateType::LinkedAsset);
+		A1.SetLinkedStateAsset(&StateTree2);
+
+		UStateTreeState& B1 = Root.AddChildState(FName(TEXT("B1")), EStateTreeStateType::State);
+		TStateTreeEditorNode<FTestTask_Stand>& Task1 = B1.AddTask<FTestTask_Stand>(FName(TEXT("Task1")));
+		Task1.GetNode().TicksToCompletion = 1;
+		{
+			FStateTreeCompiler Compiler(Log);
+			const bool bResult = Compiler.Compile(StateTree);
+			AITEST_TRUE(TEXT("StateTree should get compiled"), bResult);
+		}
+
+		const TCHAR* EnterStateStr = TEXT("EnterState");
+		const TCHAR* SustainedEnterStateStr = TEXT("EnterState=Sustained");
+		const TCHAR* ChangedEnterStateStr = TEXT("EnterState=Changed");
+		const TCHAR* ExitStateStr = TEXT("ExitState");
+		const TCHAR* SustainedExitStateStr = TEXT("ExitState=Sustained");
+		const TCHAR* ChangedExitStateStr = TEXT("ExitState=Changed");
+		const TCHAR* StateCompletedStateStr = TEXT("StateCompleted");
+		const TCHAR* TickStr = TEXT("Tick");
+
+		{
+			FStateTreeInstanceData InstanceData;
+			FTestStateTreeExecutionContext Exec(StateTree, StateTree, InstanceData);
+			const bool bInitSucceeded = Exec.IsValid();
+			AITEST_TRUE(TEXT("StateTree should init"), bInitSucceeded);
+			{
+				EStateTreeRunStatus Status = Exec.Start();
+				AITEST_EQUAL("Start should complete with Running", Status, EStateTreeRunStatus::Running);
+				AITEST_TRUE("StateTree should enter GlobalTask2", Exec.Expect(GlobalTask2.GetName(), ChangedEnterStateStr));
+				AITEST_FALSE("StateTree should not exit GlobalTask2", Exec.Expect(GlobalTask2.GetName(), ChangedExitStateStr));
+				AITEST_FALSE("StateTree should not exit GlobalTask2", Exec.Expect(GlobalTask2.GetName(), SustainedExitStateStr));
+				AITEST_TRUE("StateTree should enter Task2", Exec.Expect(Task2.GetName(), ChangedEnterStateStr));
+				AITEST_FALSE("StateTree should not exit Task2", Exec.Expect(Task2.GetName(), ChangedExitStateStr));
+				AITEST_FALSE("StateTree should not exit Task2", Exec.Expect(Task2.GetName(), SustainedExitStateStr));
+				AITEST_FALSE("StateTree should not enter Task1", Exec.Expect(Task1.GetName(), EnterStateStr));
+				Exec.LogClear();
+			}
+			{
+				// Task2 completes. The linked state didn't complete. Global is sustained.
+				EStateTreeRunStatus Status = Exec.Tick(0.1f);
+				AITEST_EQUAL("Tick should complete with Running", Status, EStateTreeRunStatus::Running);
+				AITEST_TRUE(TEXT("StateTree should exit the root state but not the global"),
+					Exec.Expect(GlobalTask2.GetName(), TickStr)
+					.Then(Task2.GetName(), TickStr)
+					.Then(Task2.GetName(), StateCompletedStateStr)
+					//@TODO the state completed. It should be a new state and not sustained.
+					//.Then(Task2.GetName(), ChangedExitStateStr)
+					//@TODO call enter/exit of globals.
+					//.Then(GlobalTask2.GetName(), SustainedExitStateStr)
+					.Then(Task0.GetName(), SustainedExitStateStr)
+					.Then(Task0.GetName(), SustainedEnterStateStr)
+					//@TODO call enter/exit of globals.
+					//.Then(GlobalTask2.GetName(), SustainedEnterStateStr)
+					//@TODO the state completed. It should be a new state and not sustained.
+					//.Then(Task2.GetName(), ChangedEnterStateStr)
+				);
+				AITEST_FALSE("StateTree should not tick Task1", Exec.Expect(Task1.GetName(), TickStr));
+				AITEST_FALSE("StateTree should not exit Task1", Exec.Expect(Task1.GetName(), ExitStateStr));
+				Exec.LogClear();
+			}
+			{
+				Exec.Stop();
+			}
+			// Change the order, the global completes before task2. Task2 should not tick.
+			{
+				Task2.GetNode().TicksToCompletion = 1;
+				GlobalTask2.GetNode().TicksToCompletion = 1;
+				FStateTreeCompiler Compiler2(Log);
+				const bool bResult2 = Compiler2.Compile(StateTree2);
+				AITEST_TRUE(TEXT("StateTree2 should get compiled"), bResult2);
+			}
+			{
+				EStateTreeRunStatus Status = Exec.Start();
+				AITEST_EQUAL("Start should complete with Running", Status, EStateTreeRunStatus::Running);
+				Exec.LogClear();
+			}
+			{
+				// GlobalTask2 completes. The tree completed. Task2 won't tick and won't complete.
+				EStateTreeRunStatus Status = Exec.Tick(0.1f);
+				AITEST_EQUAL("Tick should complete with Running", Status, EStateTreeRunStatus::Running);
+				AITEST_TRUE(TEXT("StateTree should exit the root state and the global"),
+					Exec.Expect(GlobalTask2.GetName(), TickStr)
+					//@TODO The tree completed. It is a new state.
+					//.Then(Task2.GetName(), ChangedExitStateStr)
+					//@TODO call enter/exit of globals.
+					//.Then(GlobalTask2.GetName(), ChangedExitStateStr)
+					.Then(Task0.GetName(), SustainedExitStateStr)
+					.Then(Task0.GetName(), SustainedEnterStateStr)
+					//@TODO call enter/exit of globals.
+					//.Then(GlobalTask2.GetName(), ChangedEnterStateStr)
+					//@TODO The tree completed. It is a new state.
+					//.Then(Task2.GetName(), EnterStateStr)
+				);
+				AITEST_FALSE("StateTree should not tick Tasks2", Exec.Expect(Task2.GetName(), TickStr));
+				AITEST_FALSE("StateTree should not tick Tasks1", Exec.Expect(Task1.GetName(), TickStr));
+				AITEST_FALSE("StateTree should not exit Task1", Exec.Expect(Task1.GetName(), EnterStateStr));
+				AITEST_FALSE("StateTree should not exit Task1", Exec.Expect(Task1.GetName(), ExitStateStr));
+				Exec.LogClear();
+			}
 
 			Exec.Stop();
 		}
@@ -177,7 +378,7 @@ struct FStateTreeTest_EnterAndExitLinkedAsset : FStateTreeTestBase
 		return true;
 	}
 };
-IMPLEMENT_AI_INSTANT_TEST(FStateTreeTest_EnterAndExitLinkedAsset, "System.StateTree.LinkedAsset.EnterAndExit");
+IMPLEMENT_AI_INSTANT_TEST(FStateTreeTest_EnterAndExitLinkedAsset2, "System.StateTree.LinkedAsset.EnterAndExit2");
 
 struct FStateTreeTest_MultipleSameLinkedAsset : FStateTreeTestBase
 {
diff --git a/Engine/Plugins/Runtime/StateTree/Source/StateTreeTestSuite/Private/StateTreeTest.cpp b/Engine/Plugins/Runtime/StateTree/Source/StateTreeTestSuite/Private/StateTreeTest.cpp
--- a/Engine/Plugins/Runtime/StateTree/Source/StateTreeTestSuite/Private/StateTreeTest.cpp
+++ b/Engine/Plugins/Runtime/StateTree/Source/StateTreeTestSuite/Private/StateTreeTest.cpp
@@ -2221,6 +2221,103 @@ struct FStateTreeTest_SubTreeTransition : FStateTreeTestBase
 };
 IMPLEMENT_AI_INSTANT_TEST(FStateTreeTest_SubTreeTransition, "System.StateTree.SubTreeTransition");
 
+struct FStateTreeTest_Reentrant : FStateTreeTestBase
+{
+	virtual bool InstantTest() override
+	{
+		//Tree1
+		//	-Root
+		//		- State1
+		//			- State2: OnComplete -> State1
+
+		UStateTree& StateTree1 = NewStateTree();
+		{
+			UStateTreeEditorData& EditorData = *Cast<UStateTreeEditorData>(StateTree1.EditorData);
+			UStateTreeState& Root1 = EditorData.AddSubTree(FName(TEXT("Tree1Root")));
+			UStateTreeState& State1 = Root1.AddChildState(FName(TEXT("Tree1State1")));
+			UStateTreeState& State2 = State1.AddChildState(FName(TEXT("Tree1State2")));
+
+			State2.AddTransition(EStateTreeTransitionTrigger::OnStateCompleted, EStateTreeTransitionType::GotoState, &State1);
+			{
+				TStateTreeEditorNode<FTestTask_PrintValue>& State1Task1 = State1.AddTask<FTestTask_PrintValue>("Tree1State1Task1");
+				State1Task1.GetInstanceData().Value = 101;
+				State1Task1.GetInstanceData().TickRunStatus = EStateTreeRunStatus::Running;
+
+			}
+			{
+				TStateTreeEditorNode<FTestTask_PrintValue>& State2Task1 = State2.AddTask<FTestTask_PrintValue>("Tree1State2Task1");
+				State2Task1.GetInstanceData().Value = 201;
+				State2Task1.GetInstanceData().TickRunStatus = EStateTreeRunStatus::Succeeded;
+			}
+			{
+				FStateTreeCompilerLog Log;
+				FStateTreeCompiler Compiler(Log);
+				const bool bResult = Compiler.Compile(StateTree1);
+				AITEST_TRUE("StateTree should get compiled", bResult);
+			}
+		}
+		{
+			FStateTreeInstanceData InstanceData;
+			{
+				FTestStateTreeExecutionContext Exec(StateTree1, StateTree1, InstanceData);
+				const bool bInitSucceeded = Exec.IsValid();
+				AITEST_TRUE("StateTree should init", bInitSucceeded);
+
+				// Start and enter state
+				EStateTreeRunStatus Status = Exec.Start();
+				AITEST_TRUE("StateTree Active States should be in Tree1Root/Tree1State1/Tree1State2",
+					Exec.ExpectInActiveStates("Tree1Root", "Tree1State1", "Tree1State2"));
+				AITEST_TRUE("StateTree should be running", Status == EStateTreeRunStatus::Running);
+				Exec.LogClear();
+			}
+			{
+				// On go to State1, reselect State1 and State2. It's a new State2 same State1.
+				FTestStateTreeExecutionContext Exec(StateTree1, StateTree1, InstanceData);
+
+				UE::StateTree::FActiveStatePath BeforeStatePath = InstanceData.GetExecutionState()->GetActiveStatePath();
+
+				EStateTreeRunStatus Status = Exec.Tick(0.1f);
+				AITEST_TRUE("StateTree Active States should be in Tree1Root/Tree1State1/Tree1State2",
+					Exec.ExpectInActiveStates("Tree1Root", "Tree1State1", "Tree1State2"));
+				AITEST_TRUE("StateTree should be running", Status == EStateTreeRunStatus::Running);
+				AITEST_TRUE("Expect the output tasks", Exec.Expect("Tree1State1Task1", TEXT("Tick101"))
+					.Then("Tree1State2Task1", TEXT("Tick201"))
+					.Then("Tree1State2Task1", TEXT("ExitState201"))
+					//@TODO Should select a new state because it's completed.
+					//.Then("Tree1State2Task1", TEXT("ExitState=Changed"))
+					.Then("Tree1State1Task1", TEXT("ExitState101"))
+					.Then("Tree1State1Task1", TEXT("ExitState=Sustained"))
+					.Then("Tree1State1Task1", TEXT("EnterState101"))
+					.Then("Tree1State1Task1", TEXT("EnterState=Sustained"))
+					.Then("Tree1State2Task1", TEXT("EnterState201"))
+					//@TODO Should select a new state because it's completed.
+					//.Then("Tree1State2Task1", TEXT("EnterState=Changed"))
+				);
+
+
+				UE::StateTree::FActiveStatePath AfterStatePath = InstanceData.GetExecutionState()->GetActiveStatePath();
+
+				AITEST_TRUE(TEXT("Same length"), AfterStatePath.Num() == BeforeStatePath.Num());
+				AITEST_TRUE(TEXT("Length is 3"), AfterStatePath.Num() == 3);
+				AITEST_TRUE(TEXT("Element 1 is the same"), AfterStatePath.GetView()[0] == BeforeStatePath.GetView()[0]);
+				//@TODO Should select a new state because it's completed.
+				//AITEST_TRUE(TEXT("Element 2 is the same"), AfterStatePath.GetView()[1] == BeforeStatePath.GetView()[1]);
+				//@TODO Should select a new state because it's completed.
+				//AITEST_TRUE(TEXT("Element 3 is different"), AfterStatePath.GetView()[2] != BeforeStatePath.GetView()[2]);
+
+				Exec.LogClear();
+			}
+			{
+				FTestStateTreeExecutionContext Exec(StateTree1, StateTree1, InstanceData);
+				Exec.Stop();
+			}
+		}
+
+		return true;
+	}
+};
+IMPLEMENT_AI_INSTANT_TEST(FStateTreeTest_Reentrant, "System.StateTree.Reentrant");
+
 } // namespace UE::StateTree::Tests
 
 UE_ENABLE_OPTIMIZATION_SHIP
diff --git a/Engine/Plugins/Runtime/StateTree/Source/StateTreeTestSuite/Private/StateTreeTestTypes.h b/Engine/Plugins/Runtime/StateTree/Source/StateTreeTestSuite/Private/StateTreeTestTypes.h
--- a/Engine/Plugins/Runtime/StateTree/Source/StateTreeTestSuite/Private/StateTreeTestTypes.h
+++ b/Engine/Plugins/Runtime/StateTree/Source/StateTreeTestSuite/Private/StateTreeTestTypes.h
@@ -286,6 +286,7 @@ struct FTestTask_PrintValue : public FStateTreeTaskBase
 		const FInstanceDataType& InstanceData = Context.GetInstanceData(*this);
 		TestLog.Log(Name, FString::Printf(TEXT("EnterState%d"), InstanceData.Value));
 		TestLog.Log(Name, FString::Printf(TEXT("EnterState:%s"), *InstanceData.GetArrayDescription()));
+		TestLog.Log(Name, FString::Printf(TEXT("EnterState=%s"), *StaticEnum<EStateTreeStateChangeType>()->GetNameStringByValue((int64)Transition.ChangeType)));
 		if (CustomEnterStateFunc)
 		{
 			CustomEnterStateFunc(Context, this);
@@ -299,6 +300,7 @@ struct FTestTask_PrintValue : public FStateTreeTaskBase
 		const FInstanceDataType& InstanceData = Context.GetInstanceData(*this);
 		TestLog.Log(Name, FString::Printf(TEXT("ExitState%d"), InstanceData.Value));
 		TestLog.Log(Name, FString::Printf(TEXT("ExitState:%s"), *InstanceData.GetArrayDescription()));
+		TestLog.Log(Name, FString::Printf(TEXT("ExitState=%s"), *StaticEnum<EStateTreeStateChangeType>()->GetNameStringByValue((int64)Transition.ChangeType)));
 		if (CustomExitStateFunc)
 		{
 			CustomExitStateFunc(Context, this);
@@ -501,7 +503,8 @@ struct FTestTask_Stand : public FStateTreeTaskBase
 	virtual EStateTreeRunStatus EnterState(FStateTreeExecutionContext& Context, const FStateTreeTransitionResult& Transition) const override
 	{
 		FStateTreeTestLog& TestLog = Context.GetExternalData(LogHandle);
-			TestLog.Log(Name, TEXT("EnterState"));
+		TestLog.Log(Name, TEXT("EnterState"));
+		TestLog.Log(Name, FString::Printf(TEXT("EnterState=%s"), *StaticEnum<EStateTreeStateChangeType>()->GetNameStringByValue((int64)Transition.ChangeType)));
 
 		FInstanceDataType& InstanceData = Context.GetInstanceData(*this);
 		
@@ -529,8 +532,9 @@ struct FTestTask_Stand : public FStateTreeTaskBase
 			TestLog.Log(Name, TEXT("ExitStopped"));
 		}
 
-			TestLog.Log(Name, TEXT("ExitState"));
-		}
+		TestLog.Log(Name, TEXT("ExitState"));
+		TestLog.Log(Name, FString::Printf(TEXT("ExitState=%s"), *StaticEnum<EStateTreeStateChangeType>()->GetNameStringByValue((int64)Transition.ChangeType)));
+	}
 
 	virtual void StateCompleted(FStateTreeExecutionContext& Context, const EStateTreeRunStatus CompletionStatus, const FStateTreeActiveStates& CompletedActiveStates) const override
 	{
@@ -1056,4 +1060,4 @@ struct FTestTask_CustomFuncOnDelegate : public FStateTreeTaskBase
 	TFunction<void(const FStateTreeWeakExecutionContext&, FStateTreeDelegateListener)> CustomFunc;
 
 	bool bRemoveOnExit = true;
-};
+};
\ No newline at end of file
diff --git a/Engine/Plugins/Runtime/StateTree/Source/StateTreeTestSuite/Private/StateTreeTransitionTest.cpp b/Engine/Plugins/Runtime/StateTree/Source/StateTreeTestSuite/Private/StateTreeTransitionTest.cpp
--- a/Engine/Plugins/Runtime/StateTree/Source/StateTreeTestSuite/Private/StateTreeTransitionTest.cpp
+++ b/Engine/Plugins/Runtime/StateTree/Source/StateTreeTestSuite/Private/StateTreeTransitionTest.cpp
@@ -166,6 +166,11 @@ struct FStateTreeTest_TransitionNextSelectableState : FStateTreeTestBase
 {
 	virtual bool InstantTest() override
 	{
+		// Root
+		//  State0: OnCompleted->Nextselectable
+		//  State1: EnterCondition fails
+		//  State2: EnterCondition fails
+		//  State3: EnterCondition successed. OnCompleted->Succeeded
 		UStateTree& StateTree = NewStateTree();
 		UStateTreeEditorData& EditorData = *Cast<UStateTreeEditorData>(StateTree.EditorData);
 
@@ -173,6 +178,7 @@ struct FStateTreeTest_TransitionNextSelectableState : FStateTreeTestBase
 		UStateTreeState& State0 = Root.AddChildState(FName(TEXT("State0")));
 		UStateTreeState& State1 = Root.AddChildState(FName(TEXT("State1")));
 		UStateTreeState& State2 = Root.AddChildState(FName(TEXT("State2")));
+		UStateTreeState& State3 = Root.AddChildState(FName(TEXT("State3")));
 
 		auto& EvalA = EditorData.AddEvaluator<FTestEval_A>();
 		EvalA.GetInstanceData().bBoolA = true;
@@ -187,13 +193,20 @@ struct FStateTreeTest_TransitionNextSelectableState : FStateTreeTestBase
 		EditorData.AddPropertyBinding(EvalA, TEXT("bBoolA"), BoolCond1, TEXT("bLeft"));
 		BoolCond1.GetInstanceData().bRight = !EvalA.GetInstanceData().bBoolA;
 
-		// Add Task 2 with Condition that will always succeed
+		// Add Task 2 with Condition that will always fail
 		auto& Task2 = State2.AddTask<FTestTask_Stand>(FName(TEXT("Task2")));
 		auto& BoolCond2 = State2.AddEnterCondition<FStateTreeCompareBoolCondition>();
-		State2.AddTransition(EStateTreeTransitionTrigger::OnStateCompleted, EStateTreeTransitionType::Succeeded);
 
 		EditorData.AddPropertyBinding(EvalA, TEXT("bBoolA"), BoolCond2, TEXT("bLeft"));
-		BoolCond2.GetInstanceData().bRight = EvalA.GetInstanceData().bBoolA;
+		BoolCond2.GetInstanceData().bRight = !EvalA.GetInstanceData().bBoolA;
+
+		// Add Task 3 with Condition that will always succeed
+		auto& Task3 = State3.AddTask<FTestTask_Stand>(FName(TEXT("Task3")));
+		auto& BoolCond3 = State3.AddEnterCondition<FStateTreeCompareBoolCondition>();
+		State3.AddTransition(EStateTreeTransitionTrigger::OnStateCompleted, EStateTreeTransitionType::Succeeded);
+
+		EditorData.AddPropertyBinding(EvalA, TEXT("bBoolA"), BoolCond3, TEXT("bLeft"));
+		BoolCond3.GetInstanceData().bRight = EvalA.GetInstanceData().bBoolA;
 
 		FStateTreeCompilerLog Log;
 		FStateTreeCompiler Compiler(Log);
@@ -212,19 +225,24 @@ struct FStateTreeTest_TransitionNextSelectableState : FStateTreeTestBase
 
 		// Start and enter state
 		Exec.Start();
+		AITEST_TRUE(TEXT("Should be in active state"), Exec.ExpectInActiveStates("Root", "State0"));
 		AITEST_TRUE(TEXT("StateTree Task0 should enter state"), Exec.Expect(Task0.GetName(), EnterStateStr));
 		Exec.LogClear();
 
 		// Transition from State0 and tries to select State1. It should fail (Task1 and Tasks2) and because transition is set to "Next Selectable", it should now select Task 3 and Enter State
 		Exec.Tick(0.1f);
+		AITEST_TRUE(TEXT("Should be in active state"), Exec.ExpectInActiveStates("Root", "State3"));
 		AITEST_TRUE(TEXT("StateTree Task0 should complete"), Exec.Expect(Task0.GetName(), StateCompletedStr));
 		AITEST_FALSE(TEXT("StateTree Task1 should not enter state"), Exec.Expect(Task1.GetName(), EnterStateStr));
-		AITEST_TRUE(TEXT("StateTree Task2 should enter state"), Exec.Expect(Task2.GetName(), EnterStateStr));
+		AITEST_FALSE(TEXT("StateTree Task2 should not enter state"), Exec.Expect(Task2.GetName(), EnterStateStr));
+		AITEST_TRUE(TEXT("StateTree Task3 should enter state"), Exec.Expect(Task3.GetName(), EnterStateStr));
 		Exec.LogClear();
 
-		// Complete Task2
-		Exec.Tick(0.1f);
-		AITEST_TRUE(TEXT("StateTree Task2 should complete"), Exec.Expect(Task2.GetName(), StateCompletedStr));
+		// Complete Task3
+		EStateTreeRunStatus Status = Exec.Tick(0.1f);
+		AITEST_EQUAL(TEXT("Tick should succeed"), Status, EStateTreeRunStatus::Succeeded);
+		AITEST_TRUE(TEXT("Should be in active state"), Exec.GetActiveStateNames().Num() == 0);
+		AITEST_TRUE(TEXT("StateTree Task3 should complete"), Exec.Expect(Task3.GetName(), StateCompletedStr));
 		Exec.LogClear();
 
 		Exec.Stop();
@@ -606,6 +624,11 @@ struct FStateTreeTest_FollowTransitions : FStateTreeTestBase
 {
 	virtual bool InstantTest() override
 	{
+		//Root
+		// Trans
+		// A
+		// B
+		// C
 		UStateTree& StateTree = NewStateTree();
 		UStateTreeEditorData& EditorData = *Cast<UStateTreeEditorData>(StateTree.EditorData);
 		FInstancedPropertyBag& RootPropertyBag = GetRootPropertyBag(EditorData);
@@ -864,6 +887,7 @@ struct FStateTreeTest_RegularTransitions : FStateTreeTestBase
 			AITEST_TRUE(TEXT("1st Tick should tick StateB"), Exec.Expect("TaskB", TEXT("Tick1")));
 			Exec.LogClear();
 
+			// B go to C
 			Status = Exec.Tick(1.0f);
 			AITEST_EQUAL(TEXT("2nd Tick should complete with Running"), Status, EStateTreeRunStatus::Running);
 			AITEST_TRUE(TEXT("2nd Tick should tick Global tasks"), Exec.Expect("GlobalTask", TEXT("Tick111")));
@@ -871,6 +895,8 @@ struct FStateTreeTest_RegularTransitions : FStateTreeTestBase
 			AITEST_TRUE(TEXT("2nd Tick should tick the StateB"), Exec.Expect("TaskB", TEXT("Tick1")));
 			AITEST_TRUE(TEXT("2nd Tick should exit the StateB"), Exec.Expect("TaskB", TEXT("ExitState1")));
 			AITEST_TRUE(TEXT("2nd Tick should enter the StateC"), Exec.Expect("TaskC", TEXT("EnterState2")));
+			AITEST_FALSE(TEXT("2nd Tick should not exit the StateA"), Exec.Expect("TaskA", TEXT("ExitState111")));
+			AITEST_FALSE(TEXT("2nd Tick should not exit the Global tasks"), Exec.Expect("GlobalTask", TEXT("ExitState111")));
 			Exec.LogClear();
 
 			Status = Exec.Tick(1.0f);
@@ -878,8 +904,10 @@ struct FStateTreeTest_RegularTransitions : FStateTreeTestBase
 			AITEST_TRUE(TEXT("3rd Tick should tick Global tasks"), Exec.Expect("GlobalTask", TEXT("Tick111")));
 			AITEST_TRUE(TEXT("3rd Tick should tick StateA"), Exec.Expect("TaskA", TEXT("Tick111")));
 			AITEST_TRUE(TEXT("3rd Tick should tick StateC"), Exec.Expect("TaskC", TEXT("Tick2")));
+			AITEST_FALSE(TEXT("3th Tick should not exit StateC"), Exec.Expect("TaskD", TEXT("ExitState2")));
 			Exec.LogClear();
 
+			// C go to D
 			Status = Exec.Tick(1.0f);
 			AITEST_EQUAL(TEXT("4th Tick should complete with Running"), Status, EStateTreeRunStatus::Running);
 			AITEST_TRUE(TEXT("4th Tick should tick Global tasks"), Exec.Expect("GlobalTask", TEXT("Tick111")));
@@ -887,6 +915,8 @@ struct FStateTreeTest_RegularTransitions : FStateTreeTestBase
 			AITEST_TRUE(TEXT("4th Tick should tick the StateC"), Exec.Expect("TaskC", TEXT("Tick2")));
 			AITEST_TRUE(TEXT("4th Tick should exit the StateC"), Exec.Expect("TaskC", TEXT("ExitState2")));
 			AITEST_TRUE(TEXT("4th Tick should enter the StateD"), Exec.Expect("TaskD", TEXT("EnterState3")));
+			AITEST_FALSE(TEXT("4th Tick should not exit the StateA"), Exec.Expect("TaskA", TEXT("ExitState111")));
+			AITEST_FALSE(TEXT("4th Tick should not exit the Global tasks"), Exec.Expect("GlobalTask", TEXT("ExitState111")));
 			Exec.LogClear();
 
 			Status = Exec.Tick(0.001f);
@@ -894,19 +924,22 @@ struct FStateTreeTest_RegularTransitions : FStateTreeTestBase
 			AITEST_TRUE(TEXT("5th Tick should tick Global tasks"), Exec.Expect("GlobalTask", TEXT("Tick111")));
 			AITEST_TRUE(TEXT("5th Tick should tick StateA"), Exec.Expect("TaskA", TEXT("Tick111")));
 			AITEST_TRUE(TEXT("5th Tick should tick StateD"), Exec.Expect("TaskD", TEXT("Tick3")));
+			AITEST_FALSE(TEXT("5th Tick should not exit StateD"), Exec.Expect("TaskD", TEXT("ExitState3")));
 			Exec.LogClear();
 
+			// D go to root
 			Status = Exec.Tick(1.0f);
 			AITEST_EQUAL(TEXT("6th Tick should complete with Running"), Status, EStateTreeRunStatus::Running);
 			AITEST_TRUE(TEXT("6th Tick should tick Global tasks"), Exec.Expect("GlobalTask", TEXT("Tick111")));
 			AITEST_TRUE(TEXT("6th Tick should tick StateA"), Exec.Expect("TaskA", TEXT("Tick111")));
 			AITEST_TRUE(TEXT("6th Tick should tick StateD"), Exec.Expect("TaskD", TEXT("Tick3")));
 			AITEST_TRUE(TEXT("6th Tick should exit the StateD"), Exec.Expect("TaskD", TEXT("ExitState3")));
-			//AITEST_FALSE(TEXT("6th Tick should not exit the Global tasks"), Exec.Expect("GlobalTask", TEXT("ExitState111")));
-			//AITEST_FALSE(TEXT("6th Tick should not enter the Global tasks"), Exec.Expect("GlobalTask", TEXT("EnterState111")));
-			//AITEST_FALSE(TEXT("6th Tick should not exit the StateA"), Exec.Expect("TaskA", TEXT("ExitState111")));
-			//AITEST_FALSE(TEXT("6th Tick should not enter the StateA"), Exec.Expect("TaskA", TEXT("EnterState111")));
+			AITEST_FALSE(TEXT("6th Tick should not exit the Global tasks"), Exec.Expect("GlobalTask", TEXT("ExitState111")));
+			AITEST_FALSE(TEXT("6th Tick should not enter the Global tasks"), Exec.Expect("GlobalTask", TEXT("EnterState111")));
+			AITEST_TRUE(TEXT("6th Tick should exit the StateA"), Exec.Expect("TaskA", TEXT("ExitState=Sustained")));
+			AITEST_TRUE(TEXT("6th Tick should enter the StateA"), Exec.Expect("TaskA", TEXT("EnterState=Sustained")));
 			AITEST_TRUE(TEXT("6th Tick should enter the StateB"), Exec.Expect("TaskB", TEXT("EnterState1")));
+			AITEST_TRUE(TEXT("6th Tick should enter the StateB"), Exec.Expect("TaskB", TEXT("EnterState=Changed")));
 			Exec.LogClear();
 
 			Status = Exec.Tick(1.0f);
@@ -928,6 +961,574 @@ struct FStateTreeTest_RegularTransitions : FStateTreeTestBase
 };
 IMPLEMENT_AI_INSTANT_TEST(FStateTreeTest_RegularTransitions, "System.StateTree.Transition.RegularTransitions");
 
+
+struct FStateTreeTest_RequestTransition : FStateTreeTestBase
+{
+	virtual bool InstantTest() override
+	{
+		//Tree 1
+		//RootA
+		// StateB
+		//  StateC
+		//   StateD
+		//    StateE
+		//     StateF
+		//  StateI
+		//   StateJ
+		// StateX
+		//  StateY
+		//RootO
+		// StateP
+
+		FStateTreeCompilerLog Log;
+
+		enum class ECustomFunctionToRun
+		{
+			None,
+			TransitionD_To_E,
+			TransitionB_To_I,
+			TransitionB_To_C,
+			TransitionC_To_B,
+			TransitionB_To_X,
+			TransitionB_To_B,
+			TransitionB_To_P,
+		} TransitionToExecute = ECustomFunctionToRun::None;
+		struct FStateHandle
+		{
+			FStateTreeStateHandle B;
+			FStateTreeStateHandle C;
+			FStateTreeStateHandle E;
+			FStateTreeStateHandle I;
+			FStateTreeStateHandle X;
+			FStateTreeStateHandle P;
+			FGuid B_ID;
+			FGuid C_ID;
+			FGuid E_ID;
+			FGuid I_ID;
+			FGuid X_ID;
+			FGuid P_ID;
+			bool bFinishTasksB = false;
+			bool bFinishTasksC = false;
+			bool bFinishTasksD = false;
+		} AllStateHandle;
+
+		// Main asset
+		UStateTree& StateTree = NewStateTree();
+		{
+			UStateTreeEditorData& EditorData = *Cast<UStateTreeEditorData>(StateTree.EditorData);
+			UStateTreeState& StateA = EditorData.AddSubTree("StateA");
+			StateA.AddTask<FTestTask_PrintValue>("StateATask");
+			
+			UStateTreeState& StateB = StateA.AddChildState("StateB", EStateTreeStateType::State);
+			FTestTask_PrintValue& TaskB = StateB.AddTask<FTestTask_PrintValue>("StateBTask").GetNode();
+			AllStateHandle.B_ID = StateB.ID;
+
+			UStateTreeState& StateC = StateB.AddChildState("StateC", EStateTreeStateType::State);
+			FTestTask_PrintValue& TaskC = StateC.AddTask<FTestTask_PrintValue>("StateCTask").GetNode();
+			AllStateHandle.C_ID = StateC.ID;
+
+			UStateTreeState& StateD = StateC.AddChildState("StateD", EStateTreeStateType::State);
+			FTestTask_PrintValue& TaskD = StateD.AddTask<FTestTask_PrintValue>("StateDTask").GetNode();
+			StateD.SelectionBehavior = EStateTreeStateSelectionBehavior::TryEnterState;
+			
+			UStateTreeState& StateE = StateD.AddChildState("StateE", EStateTreeStateType::State);
+			StateE.AddTask<FTestTask_PrintValue>("StateETask");
+			AllStateHandle.E_ID = StateE.ID;
+			
+			UStateTreeState& StateF = StateE.AddChildState("StateF", EStateTreeStateType::State);
+			StateF.AddTask<FTestTask_PrintValue>("StateFTask");
+
+			UStateTreeState& StateI = StateB.AddChildState("StateI", EStateTreeStateType::State);
+			StateI.AddTask<FTestTask_PrintValue>("StateITask");
+			AllStateHandle.I_ID = StateI.ID;
+			
+			UStateTreeState& StateJ = StateI.AddChildState("StateJ", EStateTreeStateType::State);
+			StateJ.AddTask<FTestTask_PrintValue>("StateJTask");
+
+			UStateTreeState& StateX = StateA.AddChildState("StateX", EStateTreeStateType::State);
+			StateX.AddTask<FTestTask_PrintValue>("StateXTask");
+			AllStateHandle.X_ID = StateX.ID;
+
+			UStateTreeState& StateY = StateX.AddChildState("StateY", EStateTreeStateType::State);
+			StateY.AddTask<FTestTask_PrintValue>("StateYTask");
+
+			UStateTreeState& StateO = EditorData.AddSubTree("StateO");
+			StateO.AddTask<FTestTask_PrintValue>("StateOTask");
+
+			UStateTreeState& StateP = StateO.AddChildState("StateP", EStateTreeStateType::State);
+			StateP.AddTask<FTestTask_PrintValue>("StatePTask");
+			AllStateHandle.P_ID = StateP.ID;
+
+			TaskB.CustomTickFunc = [&TransitionToExecute, &AllStateHandle]
+				(FStateTreeExecutionContext& Context, const FTestTask_PrintValue* Task)
+				{
+					if (TransitionToExecute == ECustomFunctionToRun::TransitionB_To_I)
+					{
+						Context.RequestTransition(AllStateHandle.I);
+					}
+					if (TransitionToExecute == ECustomFunctionToRun::TransitionB_To_C)
+					{
+						Context.RequestTransition(AllStateHandle.C);
+					}
+					if (TransitionToExecute == ECustomFunctionToRun::TransitionB_To_X)
+					{
+						Context.RequestTransition(AllStateHandle.X);
+					}
+					if (TransitionToExecute == ECustomFunctionToRun::TransitionB_To_B)
+					{
+						Context.RequestTransition(AllStateHandle.B);
+					}
+					if (TransitionToExecute == ECustomFunctionToRun::TransitionB_To_P)
+					{
+						Context.RequestTransition(AllStateHandle.P);
+					}
+					if (AllStateHandle.bFinishTasksB)
+					{
+						Context.FinishTask(*Task, EStateTreeFinishTaskType::Succeeded);
+					}
+				};
+
+			TaskC.CustomTickFunc = [&TransitionToExecute, &AllStateHandle]
+				(FStateTreeExecutionContext& Context, const FTestTask_PrintValue* Task)
+				{
+					if (TransitionToExecute == ECustomFunctionToRun::TransitionC_To_B)
+					{
+						Context.RequestTransition(AllStateHandle.B);
+					}
+					if (AllStateHandle.bFinishTasksC)
+					{
+						Context.FinishTask(*Task, EStateTreeFinishTaskType::Succeeded);
+					}
+				};
+
+			TaskD.CustomTickFunc = [&TransitionToExecute, &AllStateHandle]
+				(FStateTreeExecutionContext& Context, const FTestTask_PrintValue* Task)
+				{
+					if (TransitionToExecute == ECustomFunctionToRun::TransitionD_To_E)
+					{
+						Context.RequestTransition(AllStateHandle.E);
+					}
+
+					if (AllStateHandle.bFinishTasksD)
+					{
+						Context.FinishTask(*Task, EStateTreeFinishTaskType::Succeeded);
+					}
+				};
+
+			FStateTreeCompiler Compiler(Log);
+			const bool bResult = Compiler.Compile(StateTree);
+			AITEST_TRUE("StateTree should get compiled", bResult);
+
+			AllStateHandle.B = StateTree.GetStateHandleFromId(AllStateHandle.B_ID);
+			AllStateHandle.C = StateTree.GetStateHandleFromId(AllStateHandle.C_ID);
+			AllStateHandle.E = StateTree.GetStateHandleFromId(AllStateHandle.E_ID);
+			AllStateHandle.I = StateTree.GetStateHandleFromId(AllStateHandle.I_ID);
+			AllStateHandle.X = StateTree.GetStateHandleFromId(AllStateHandle.X_ID);
+			AllStateHandle.P = StateTree.GetStateHandleFromId(AllStateHandle.P_ID);
+		}
+
+		FStateTreeInstanceData InstanceData;
+		{
+			FTestStateTreeExecutionContext Exec(StateTree, StateTree, InstanceData);
+			const bool bInitSucceeded = Exec.IsValid();
+			AITEST_TRUE("StateTree should init", bInitSucceeded);
+		}
+
+		auto ResetInstanceData = [this, &StateTree, &InstanceData]()
+		{
+			{
+				FTestStateTreeExecutionContext Exec(StateTree, StateTree, InstanceData);
+				EStateTreeRunStatus Status = Exec.Stop();
+				AITEST_EQUAL(TEXT("Should stop"), Status, EStateTreeRunStatus::Stopped);
+			}
+			InstanceData = FStateTreeInstanceData();
+			{
+				FTestStateTreeExecutionContext Exec(StateTree, StateTree, InstanceData);
+				const EStateTreeRunStatus Status = Exec.Start();
+				AITEST_EQUAL(TEXT("State should complete with Running"), Status, EStateTreeRunStatus::Running);
+				AITEST_TRUE(TEXT("Start should"), Exec.ExpectInActiveStates("StateA", "StateB", "StateC", "StateD"));
+			}
+			return true;
+		};
+		{
+			FTestStateTreeExecutionContext Exec(StateTree, StateTree, InstanceData);
+			const EStateTreeRunStatus Status = Exec.Start();
+			AITEST_EQUAL("Start should complete with Running", Status, EStateTreeRunStatus::Running);
+			AITEST_TRUE("Start should enter Global tasks", Exec.ExpectInActiveStates("StateA", "StateB", "StateC", "StateD"));
+			FTestStateTreeExecutionContext::FLogOrder LogOrder = Exec.Expect("StateATask", TEXT("EnterState=Changed"));
+			AITEST_TRUE(TEXT("Start StateA"), LogOrder);
+			LogOrder = LogOrder.Then("StateBTask", TEXT("EnterState=Changed"));
+			AITEST_TRUE(TEXT("Start StateB"), LogOrder);
+			LogOrder = LogOrder.Then("StateCTask", TEXT("EnterState=Changed"));
+			AITEST_TRUE(TEXT("Start StateC"), LogOrder);
+			LogOrder = LogOrder.Then("StateDTask", TEXT("EnterState=Changed"));
+			AITEST_TRUE(TEXT("Start StateD"), LogOrder);
+			AITEST_FALSE(TEXT("Tick StateA"), Exec.Expect("StateATask", TEXT("Tick0")));
+			AITEST_FALSE(TEXT("Tick StateB"), Exec.Expect("StateBTask", TEXT("Tick0")));
+			AITEST_FALSE(TEXT("Tick StateC"), Exec.Expect("StateCTask", TEXT("Tick0")));
+			AITEST_FALSE(TEXT("Tick StateD"), Exec.Expect("StateDTask", TEXT("Tick0")));
+			AITEST_FALSE(TEXT("Start StateA"), Exec.Expect("StateATask", TEXT("ExitState0")));
+			AITEST_FALSE(TEXT("Start StateB"), Exec.Expect("StateBTask", TEXT("ExitState0")));
+			AITEST_FALSE(TEXT("Start StateC"), Exec.Expect("StateCTask", TEXT("ExitState0")));
+			AITEST_FALSE(TEXT("Start StateD"), Exec.Expect("StateDTask", TEXT("ExitState0")));
+			Exec.LogClear();
+		}
+		{
+			FTestStateTreeExecutionContext Exec(StateTree, StateTree, InstanceData);
+			const EStateTreeRunStatus Status = Exec.Tick(0.1f);
+			AITEST_EQUAL("1st Tick should complete with Running", Status, EStateTreeRunStatus::Running);
+			AITEST_TRUE("1st Tick no transition", Exec.ExpectInActiveStates("StateA", "StateB", "StateC", "StateD"));
+			FTestStateTreeExecutionContext::FLogOrder LogOrder = Exec.Expect("StateATask", TEXT("Tick0"));
+			AITEST_TRUE(TEXT("Tick StateA"), LogOrder);
+			LogOrder = LogOrder.Then("StateBTask", TEXT("Tick0"));
+			AITEST_TRUE(TEXT("Tick StateB"), LogOrder);
+			LogOrder = LogOrder.Then("StateCTask", TEXT("Tick0"));
+			AITEST_TRUE(TEXT("Tick StateC"), LogOrder);
+			LogOrder = LogOrder.Then("StateDTask", TEXT("Tick0"));
+			AITEST_TRUE(TEXT("Tick StateD"), LogOrder);
+			AITEST_FALSE(TEXT("Enter StateA"), Exec.Expect("StateATask", TEXT("EnterState0")));
+			AITEST_FALSE(TEXT("Enter StateB"), Exec.Expect("StateBTask", TEXT("EnterState0")));
+			AITEST_FALSE(TEXT("Enter StateC"), Exec.Expect("StateCTask", TEXT("EnterState0")));
+			AITEST_FALSE(TEXT("Enter StateD"), Exec.Expect("StateDTask", TEXT("EnterState0")));
+			AITEST_FALSE(TEXT("Exit StateA"), Exec.Expect("StateATask", TEXT("ExitState0")));
+			AITEST_FALSE(TEXT("Exit StateB"), Exec.Expect("StateBTask", TEXT("ExitState0")));
+			AITEST_FALSE(TEXT("Exit StateC"), Exec.Expect("StateCTask", TEXT("ExitState0")));
+			AITEST_FALSE(TEXT("Exit StateD"), Exec.Expect("StateDTask", TEXT("ExitState0")));
+			Exec.LogClear();
+		}
+		{
+			TransitionToExecute = ECustomFunctionToRun::TransitionD_To_E;
+
+			FTestStateTreeExecutionContext Exec(StateTree, StateTree, InstanceData);
+			const EStateTreeRunStatus Status = Exec.Tick(0.1f);
+			AITEST_EQUAL(TEXT("2nd Tick should complete with Running"), Status, EStateTreeRunStatus::Running);
+			AITEST_TRUE(TEXT("2nd tick sold be in new state"),  Exec.ExpectInActiveStates("StateA", "StateB", "StateC", "StateD", "StateE", "StateF"));
+			FTestStateTreeExecutionContext::FLogOrder LogOrder = Exec.Expect("StateATask", TEXT("Tick0"));
+			AITEST_TRUE(TEXT("Tick StateA"), LogOrder);
+			LogOrder = LogOrder.Then("StateBTask", TEXT("Tick0"));
+			AITEST_TRUE(TEXT("Tick StateB"), LogOrder);
+			LogOrder = LogOrder.Then("StateCTask", TEXT("Tick0"));
+			AITEST_TRUE(TEXT("Tick StateC"), LogOrder);
+			LogOrder = LogOrder.Then("StateDTask", TEXT("Tick0"));
+			AITEST_TRUE(TEXT("Tick StateD"), LogOrder);
+			LogOrder = LogOrder.Then("StateETask", TEXT("EnterState=Changed"));
+			AITEST_TRUE(TEXT("Enter StateE"), LogOrder);
+			LogOrder = LogOrder.Then("StateFTask", TEXT("EnterState=Changed"));
+			AITEST_TRUE(TEXT("Enter StateF"), LogOrder);
+			AITEST_FALSE(TEXT("Enter StateA"), Exec.Expect("StateATask", TEXT("EnterState0")));
+			AITEST_FALSE(TEXT("Enter StateB"), Exec.Expect("StateBTask", TEXT("EnterState0")));
+			AITEST_FALSE(TEXT("Enter StateC"), Exec.Expect("StateCTask", TEXT("EnterState0")));
+			AITEST_FALSE(TEXT("Enter StateD"), Exec.Expect("StateDTask", TEXT("EnterState0")));
+			AITEST_FALSE(TEXT("Exit StateA"), Exec.Expect("StateATask", TEXT("ExitState0")));
+			AITEST_FALSE(TEXT("Exit StateB"), Exec.Expect("StateBTask", TEXT("ExitState0")));
+			AITEST_FALSE(TEXT("Exit StateC"), Exec.Expect("StateCTask", TEXT("ExitState0")));
+			AITEST_FALSE(TEXT("Exit StateD"), Exec.Expect("StateDTask", TEXT("ExitState0")));
+			Exec.LogClear();
+			TransitionToExecute = ECustomFunctionToRun::None;
+		}
+		// Same transition but with a completed D
+		{
+			TransitionToExecute = ECustomFunctionToRun::TransitionD_To_E;
+			AllStateHandle.bFinishTasksD = true;
+			if (!ResetInstanceData())
+			{
+				return false;
+			}
+
+			FTestStateTreeExecutionContext Exec(StateTree, StateTree, InstanceData);
+			const EStateTreeRunStatus Status = Exec.Tick(0.1f);
+			AITEST_EQUAL(TEXT("Tick should complete with Running"), Status, EStateTreeRunStatus::Running);
+			AITEST_TRUE(TEXT("Start should"), Exec.ExpectInActiveStates("StateA", "StateB", "StateC", "StateD", "StateE", "StateF"));
+			FTestStateTreeExecutionContext::FLogOrder LogOrder = Exec.Expect("StateATask", TEXT("Tick0"));
+			AITEST_TRUE(TEXT("Tick StateA"), LogOrder);
+			LogOrder = LogOrder.Then("StateBTask", TEXT("Tick0"));
+			AITEST_TRUE(TEXT("Tick StateB"), LogOrder);
+			LogOrder = LogOrder.Then("StateCTask", TEXT("Tick0"));
+			AITEST_TRUE(TEXT("Tick StateC"), LogOrder);
+			LogOrder = LogOrder.Then("StateDTask", TEXT("Tick0"));
+			AITEST_TRUE(TEXT("Tick StateD"), LogOrder);
+			//@TODO State D completed. It should be a new state.
+			//LogOrder = LogOrder.Then("StateDTask", TEXT("ExitState=Changed"));
+			//AITEST_TRUE(TEXT("Exit StateD"), LogOrder);
+			//@TODO State D completed. It should be a new state.
+			//LogOrder = LogOrder.Then("StateDTask", TEXT("EnterState=Changed"));
+			//AITEST_TRUE(TEXT("Enter StateD"), LogOrder);
+			LogOrder = LogOrder.Then("StateETask", TEXT("EnterState=Changed"));
+			AITEST_TRUE(TEXT("Enter StateE"), LogOrder);
+			LogOrder = LogOrder.Then("StateFTask", TEXT("EnterState=Changed"));
+			AITEST_TRUE(TEXT("Enter StateF"), LogOrder);
+			AITEST_FALSE(TEXT("Enter StateA"), Exec.Expect("StateATask", TEXT("EnterState0")));
+			AITEST_FALSE(TEXT("Enter StateB"), Exec.Expect("StateBTask", TEXT("EnterState0")));
+			AITEST_FALSE(TEXT("Enter StateC"), Exec.Expect("StateCTask", TEXT("EnterState0")));
+			AITEST_FALSE(TEXT("Exit StateA"), Exec.Expect("StateATask", TEXT("ExitState0")));
+			AITEST_FALSE(TEXT("Exit StateB"), Exec.Expect("StateBTask", TEXT("ExitState0")));
+			AITEST_FALSE(TEXT("Exit StateC"), Exec.Expect("StateCTask", TEXT("ExitState0")));
+			Exec.LogClear();
+			TransitionToExecute = ECustomFunctionToRun::None;
+			AllStateHandle.bFinishTasksD = false;
+		}
+		{
+			TransitionToExecute = ECustomFunctionToRun::TransitionB_To_I;
+			if (!ResetInstanceData())
+			{
+				return false;
+			}
+
+			FTestStateTreeExecutionContext Exec(StateTree, StateTree, InstanceData);
+			const EStateTreeRunStatus Status = Exec.Tick(0.1f);
+			AITEST_EQUAL(TEXT("Tick should complete with Running"), Status, EStateTreeRunStatus::Running);
+			AITEST_TRUE(TEXT("Start should"), Exec.ExpectInActiveStates("StateA", "StateB", "StateI", "StateJ"));
+			FTestStateTreeExecutionContext::FLogOrder LogOrder = Exec.Expect("StateATask", TEXT("Tick0"));
+			AITEST_TRUE(TEXT("Tick StateA"), LogOrder);
+			LogOrder = LogOrder.Then("StateBTask", TEXT("Tick0"));
+			AITEST_TRUE(TEXT("Tick StateB"), LogOrder);
+			LogOrder = LogOrder.Then("StateCTask", TEXT("Tick0"));
+			AITEST_TRUE(TEXT("Tick StateC"), LogOrder);
+			LogOrder = LogOrder.Then("StateDTask", TEXT("Tick0"));
+			AITEST_TRUE(TEXT("Tick StateD"), LogOrder);
+			LogOrder = LogOrder.Then("StateDTask", TEXT("ExitState=Changed"));
+			AITEST_TRUE(TEXT("Exit StateD"), LogOrder);
+			LogOrder = LogOrder.Then("StateCTask", TEXT("ExitState=Changed"));
+			AITEST_TRUE(TEXT("Exit StateC"), LogOrder);
+			LogOrder = LogOrder.Then("StateITask", TEXT("EnterState=Changed"));
+			AITEST_TRUE(TEXT("Enter StateI"), LogOrder);
+			LogOrder = LogOrder.Then("StateJTask", TEXT("EnterState=Changed"));
+			AITEST_TRUE(TEXT("Enter StateJ"), LogOrder);
+			AITEST_FALSE(TEXT("Enter StateA"), Exec.Expect("StateATask", TEXT("EnterState0")));
+			AITEST_FALSE(TEXT("Enter StateB"), Exec.Expect("StateBTask", TEXT("EnterState0")));
+			AITEST_FALSE(TEXT("Exit StateA"), Exec.Expect("StateATask", TEXT("ExitState0")));
+			AITEST_FALSE(TEXT("Exit StateB"), Exec.Expect("StateBTask", TEXT("ExitState0")));
+			Exec.LogClear();
+			TransitionToExecute = ECustomFunctionToRun::None;
+		}
+		{
+			TransitionToExecute = ECustomFunctionToRun::TransitionB_To_C;
+			if (!ResetInstanceData())
+			{
+				return false;
+			}
+
+			FTestStateTreeExecutionContext Exec(StateTree, StateTree, InstanceData);
+			const EStateTreeRunStatus Status = Exec.Tick(0.1f);
+			AITEST_EQUAL(TEXT("Tick should complete with Running"), Status, EStateTreeRunStatus::Running);
+			AITEST_TRUE(TEXT("Start should"), Exec.ExpectInActiveStates("StateA", "StateB", "StateC", "StateD"));
+			FTestStateTreeExecutionContext::FLogOrder LogOrder = Exec.Expect("StateATask", TEXT("Tick0"));
+			AITEST_TRUE(TEXT("Tick StateA"), LogOrder);
+			LogOrder = LogOrder.Then("StateBTask", TEXT("Tick0"));
+			AITEST_TRUE(TEXT("Tick StateB"), LogOrder);
+			LogOrder = LogOrder.Then("StateCTask", TEXT("Tick0"));
+			AITEST_TRUE(TEXT("Tick StateC"), LogOrder);
+			LogOrder = LogOrder.Then("StateDTask", TEXT("Tick0"));
+			AITEST_TRUE(TEXT("Tick StateD"), LogOrder);
+			LogOrder = LogOrder.Then("StateDTask", TEXT("ExitState=Sustained"));
+			AITEST_TRUE(TEXT("Exit StateD"), LogOrder);
+			LogOrder = LogOrder.Then("StateCTask", TEXT("ExitState=Sustained"));
+			AITEST_TRUE(TEXT("Exit StateC"), LogOrder);
+			LogOrder = LogOrder.Then("StateCTask", TEXT("EnterState=Sustained"));
+			AITEST_TRUE(TEXT("Enter StateC"), LogOrder);
+			LogOrder = LogOrder.Then("StateDTask", TEXT("EnterState=Sustained"));
+			AITEST_TRUE(TEXT("Enter StateD"), LogOrder);
+			AITEST_FALSE(TEXT("Enter StateA"), Exec.Expect("StateATask", TEXT("EnterState0")));
+			AITEST_FALSE(TEXT("Enter StateB"), Exec.Expect("StateBTask", TEXT("EnterState0")));
+			AITEST_FALSE(TEXT("Exit StateA"), Exec.Expect("StateATask", TEXT("ExitState0")));
+			AITEST_FALSE(TEXT("Exit StateB"), Exec.Expect("StateBTask", TEXT("ExitState0")));
+			Exec.LogClear();
+			TransitionToExecute = ECustomFunctionToRun::None;
+		}
+		{
+			TransitionToExecute = ECustomFunctionToRun::TransitionC_To_B;
+			if (!ResetInstanceData())
+			{
+				return false;
+			}
+
+			FTestStateTreeExecutionContext Exec(StateTree, StateTree, InstanceData);
+			const EStateTreeRunStatus Status = Exec.Tick(0.1f);
+			AITEST_EQUAL(TEXT("Tick should complete with Running"), Status, EStateTreeRunStatus::Running);
+			AITEST_TRUE(TEXT("Start should"), Exec.ExpectInActiveStates("StateA", "StateB", "StateC", "StateD"));
+			FTestStateTreeExecutionContext::FLogOrder LogOrder = Exec.Expect("StateATask", TEXT("Tick0"));
+			AITEST_TRUE(TEXT("Tick StateA"), LogOrder);
+			LogOrder = LogOrder.Then("StateBTask", TEXT("Tick0"));
+			AITEST_TRUE(TEXT("Tick StateB"), LogOrder);
+			LogOrder = LogOrder.Then("StateCTask", TEXT("Tick0"));
+			AITEST_TRUE(TEXT("Tick StateC"), LogOrder);
+			LogOrder = LogOrder.Then("StateDTask", TEXT("Tick0"));
+			AITEST_TRUE(TEXT("Tick StateD"), LogOrder);
+			LogOrder = LogOrder.Then("StateDTask", TEXT("ExitState=Sustained"));
+			AITEST_TRUE(TEXT("Exit StateD"), LogOrder);
+			LogOrder = LogOrder.Then("StateCTask", TEXT("ExitState=Sustained"));
+			AITEST_TRUE(TEXT("Exit StateC"), LogOrder);
+			LogOrder = LogOrder.Then("StateBTask", TEXT("ExitState=Sustained"));
+			AITEST_TRUE(TEXT("Exit StateB"), LogOrder);
+			LogOrder = LogOrder.Then("StateBTask", TEXT("EnterState=Sustained"));
+			AITEST_TRUE(TEXT("Enter StateB"), LogOrder);
+			LogOrder = LogOrder.Then("StateCTask", TEXT("EnterState=Sustained"));
+			AITEST_TRUE(TEXT("Enter StateC"), LogOrder);
+			LogOrder = LogOrder.Then("StateDTask", TEXT("EnterState=Sustained"));
+			AITEST_TRUE(TEXT("Enter StateD"), LogOrder);
+			AITEST_FALSE(TEXT("Enter StateA"), Exec.Expect("StateATask", TEXT("EnterState0")));
+			AITEST_FALSE(TEXT("Exit StateA"), Exec.Expect("StateATask", TEXT("ExitState0")));
+
+			Exec.LogClear();
+			TransitionToExecute = ECustomFunctionToRun::None;
+		}
+		{
+			TransitionToExecute = ECustomFunctionToRun::TransitionC_To_B;
+			AllStateHandle.bFinishTasksC = true;
+			if (!ResetInstanceData())
+			{
+				return false;
+			}
+
+			FTestStateTreeExecutionContext Exec(StateTree, StateTree, InstanceData);
+			const EStateTreeRunStatus Status = Exec.Tick(0.1f);
+			AITEST_EQUAL(TEXT("Tick should complete with Running"), Status, EStateTreeRunStatus::Running);
+			AITEST_TRUE(TEXT("Start should"), Exec.ExpectInActiveStates("StateA", "StateB", "StateC", "StateD"));
+			FTestStateTreeExecutionContext::FLogOrder LogOrder = Exec.Expect("StateATask", TEXT("Tick0"));
+			AITEST_TRUE(TEXT("Tick StateA"), LogOrder);
+			LogOrder = LogOrder.Then("StateBTask", TEXT("Tick0"));
+			AITEST_TRUE(TEXT("Tick StateB"), LogOrder);
+			LogOrder = LogOrder.Then("StateCTask", TEXT("Tick0"));
+			AITEST_TRUE(TEXT("Tick StateC"), LogOrder);
+			LogOrder = LogOrder.Then("StateDTask", TEXT("Tick0"));
+			AITEST_TRUE(TEXT("Tick StateD"), LogOrder);
+			//@TODO State C completed. It should change for another state
+			//LogOrder = LogOrder.Then("StateDTask", TEXT("ExitState=Changed"));
+			//AITEST_TRUE(TEXT("Exit StateD"), LogOrder);
+			//@TODO State C completed. It should change for another state
+			//LogOrder = LogOrder.Then("StateCTask", TEXT("ExitState=Changed"));
+			//AITEST_TRUE(TEXT("Exit StateC"), LogOrder);
+			LogOrder = LogOrder.Then("StateBTask", TEXT("ExitState=Sustained"));
+			AITEST_TRUE(TEXT("Exit StateB"), LogOrder);
+			LogOrder = LogOrder.Then("StateBTask", TEXT("EnterState=Sustained"));
+			AITEST_TRUE(TEXT("Enter StateB"), LogOrder);
+			//@TODO State C completed. It should be a new state.
+			//LogOrder = LogOrder.Then("StateCTask", TEXT("EnterState=Changed"));
+			//AITEST_TRUE(TEXT("Enter StateC"), LogOrder);
+			//@TODO State C completed. It should be a new state.
+			//LogOrder = LogOrder.Then("StateDTask", TEXT("EnterState=Changed"));
+			//AITEST_TRUE(TEXT("Enter StateD"), LogOrder);
+			AITEST_FALSE(TEXT("Enter StateA"), Exec.Expect("StateATask", TEXT("EnterState0")));
+			AITEST_FALSE(TEXT("Exit StateA"), Exec.Expect("StateATask", TEXT("ExitState0")));
+
+			Exec.LogClear();
+			TransitionToExecute = ECustomFunctionToRun::None;
+		}
+		{
+			TransitionToExecute = ECustomFunctionToRun::TransitionB_To_X;
+			if (!ResetInstanceData())
+			{
+				return false;
+			}
+
+			FTestStateTreeExecutionContext Exec(StateTree, StateTree, InstanceData);
+			const EStateTreeRunStatus Status = Exec.Tick(0.1f);
+			AITEST_EQUAL(TEXT("Tick should complete with Running"), Status, EStateTreeRunStatus::Running);
+			AITEST_TRUE(TEXT("Start should"), Exec.ExpectInActiveStates("StateA", "StateX", "StateY"));
+			FTestStateTreeExecutionContext::FLogOrder LogOrder = Exec.Expect("StateATask", TEXT("Tick0"));
+			AITEST_TRUE(TEXT("Tick StateA"), LogOrder);
+			LogOrder = LogOrder.Then("StateBTask", TEXT("Tick0"));
+			AITEST_TRUE(TEXT("Tick StateB"), LogOrder);
+			LogOrder = LogOrder.Then("StateCTask", TEXT("Tick0"));
+			AITEST_TRUE(TEXT("Tick StateC"), LogOrder);
+			LogOrder = LogOrder.Then("StateDTask", TEXT("Tick0"));
+			AITEST_TRUE(TEXT("Tick StateD"), LogOrder);
+			LogOrder = LogOrder.Then("StateDTask", TEXT("ExitState=Changed"));
+			AITEST_TRUE(TEXT("Exit StateD"), LogOrder);
+			LogOrder = LogOrder.Then("StateCTask", TEXT("ExitState=Changed"));
+			AITEST_TRUE(TEXT("Exit StateC"), LogOrder);
+			LogOrder = LogOrder.Then("StateBTask", TEXT("ExitState=Changed"));
+			AITEST_TRUE(TEXT("Exit StateB"), LogOrder);
+			LogOrder = LogOrder.Then("StateXTask", TEXT("EnterState=Changed"));
+			AITEST_TRUE(TEXT("Enter StateX"), LogOrder);
+			LogOrder = LogOrder.Then("StateYTask", TEXT("EnterState=Changed"));
+			AITEST_TRUE(TEXT("Enter StateY"), LogOrder);
+			AITEST_FALSE(TEXT("Enter StateA"), Exec.Expect("StateATask", TEXT("EnterState0")));
+			AITEST_FALSE(TEXT("Exit StateA"), Exec.Expect("StateATask", TEXT("ExitState0")));
+			Exec.LogClear();
+			TransitionToExecute = ECustomFunctionToRun::None;
+		}
+		{
+			TransitionToExecute = ECustomFunctionToRun::TransitionB_To_B;
+			if (!ResetInstanceData())
+			{
+				return false;
+			}
+
+			FTestStateTreeExecutionContext Exec(StateTree, StateTree, InstanceData);
+			const EStateTreeRunStatus Status = Exec.Tick(0.1f);
+			AITEST_EQUAL(TEXT("Tick should complete with Running"), Status, EStateTreeRunStatus::Running);
+			AITEST_TRUE(TEXT("Start should"), Exec.ExpectInActiveStates("StateA", "StateB", "StateC", "StateD"));
+			FTestStateTreeExecutionContext::FLogOrder LogOrder = Exec.Expect("StateATask", TEXT("Tick0"));
+			AITEST_TRUE(TEXT("Tick StateA"), LogOrder);
+			LogOrder = LogOrder.Then("StateBTask", TEXT("Tick0"));
+			AITEST_TRUE(TEXT("Tick StateB"), LogOrder);
+			LogOrder = LogOrder.Then("StateCTask", TEXT("Tick0"));
+			AITEST_TRUE(TEXT("Tick StateC"), LogOrder);
+			LogOrder = LogOrder.Then("StateDTask", TEXT("Tick0"));
+			AITEST_TRUE(TEXT("Tick StateD"), LogOrder);
+			LogOrder = LogOrder.Then("StateDTask", TEXT("ExitState=Sustained"));
+			AITEST_TRUE(TEXT("Exit StateD"), LogOrder);
+			LogOrder = LogOrder.Then("StateCTask", TEXT("ExitState=Sustained"));
+			AITEST_TRUE(TEXT("Exit StateC"), LogOrder);
+			LogOrder = LogOrder.Then("StateBTask", TEXT("ExitState=Sustained"));
+			AITEST_TRUE(TEXT("Exit StateB"), LogOrder);
+			LogOrder = LogOrder.Then("StateBTask", TEXT("EnterState=Sustained"));
+			AITEST_TRUE(TEXT("Enter StateB"), LogOrder);
+			LogOrder = LogOrder.Then("StateCTask", TEXT("EnterState=Sustained"));
+			AITEST_TRUE(TEXT("Enter StateC"), LogOrder);
+			LogOrder = LogOrder.Then("StateDTask", TEXT("EnterState=Sustained"));
+			AITEST_TRUE(TEXT("Enter StateD"), LogOrder);
+			AITEST_FALSE(TEXT("Enter StateA"), Exec.Expect("StateATask", TEXT("EnterState0")));
+			AITEST_FALSE(TEXT("Exit StateA"), Exec.Expect("StateATask", TEXT("ExitState0")));
+			Exec.LogClear();
+			TransitionToExecute = ECustomFunctionToRun::None;
+		}
+		{
+			TransitionToExecute = ECustomFunctionToRun::TransitionB_To_P;
+			if (!ResetInstanceData())
+			{
+				return false;
+			}
+
+			FTestStateTreeExecutionContext Exec(StateTree, StateTree, InstanceData);
+			const EStateTreeRunStatus Status = Exec.Tick(0.1f);
+			AITEST_EQUAL(TEXT("Tick should complete with Running"), Status, EStateTreeRunStatus::Running);
+			AITEST_TRUE(TEXT("Start should"), Exec.ExpectInActiveStates("StateO", "StateP"));
+			FTestStateTreeExecutionContext::FLogOrder LogOrder = Exec.Expect("StateATask", TEXT("Tick0"));
+			AITEST_TRUE(TEXT("Tick StateA"), LogOrder);
+			LogOrder = LogOrder.Then("StateBTask", TEXT("Tick0"));
+			AITEST_TRUE(TEXT("Tick StateB"), LogOrder);
+			LogOrder = LogOrder.Then("StateCTask", TEXT("Tick0"));
+			AITEST_TRUE(TEXT("Tick StateC"), LogOrder);
+			LogOrder = LogOrder.Then("StateDTask", TEXT("Tick0"));
+			AITEST_TRUE(TEXT("Tick StateD"), LogOrder);
+			LogOrder = LogOrder.Then("StateDTask", TEXT("ExitState=Changed"));
+			AITEST_TRUE(TEXT("Exit StateD"), LogOrder);
+			LogOrder = LogOrder.Then("StateCTask", TEXT("ExitState=Changed"));
+			AITEST_TRUE(TEXT("Exit StateC"), LogOrder);
+			LogOrder = LogOrder.Then("StateBTask", TEXT("ExitState=Changed"));
+			AITEST_TRUE(TEXT("Exit StateB"), LogOrder);
+			LogOrder = LogOrder.Then("StateATask", TEXT("ExitState=Changed"));
+			AITEST_TRUE(TEXT("Exit StateA"), LogOrder);
+			LogOrder = LogOrder.Then("StateOTask", TEXT("EnterState=Changed"));
+			AITEST_TRUE(TEXT("Enter StateO"), LogOrder);
+			LogOrder = LogOrder.Then("StatePTask", TEXT("EnterState=Changed"));
+			AITEST_TRUE(TEXT("Enter StateP"), LogOrder);
+			Exec.LogClear();
+			TransitionToExecute = ECustomFunctionToRun::None;
+		}
+		{
+			FTestStateTreeExecutionContext Exec(StateTree, StateTree, InstanceData);
+			EStateTreeRunStatus Status = Exec.Stop();
+			AITEST_EQUAL(TEXT("Should stop"), Status, EStateTreeRunStatus::Stopped);
+		}
+
+		return true;
+	}
+};
+IMPLEMENT_AI_INSTANT_TEST(FStateTreeTest_RequestTransition, "System.StateTree.Transition.RequestTransition");
+
 } // namespace UE::StateTree::Tests
 
 UE_ENABLE_OPTIMIZATION_SHIP
