[State Tree] check didn't include ObjectInstanceData

diff --git a/Engine/Plugins/Runtime/StateTree/Source/StateTreeModule/Private/StateTreeExecutionContext.cpp b/Engine/Plugins/Runtime/StateTree/Source/StateTreeModule/Private/StateTreeExecutionContext.cpp
--- a/Engine/Plugins/Runtime/StateTree/Source/StateTreeModule/Private/StateTreeExecutionContext.cpp
+++ b/Engine/Plugins/Runtime/StateTree/Source/StateTreeModule/Private/StateTreeExecutionContext.cpp
@@ -3974,7 +3974,10 @@ bool FStateTreeExecutionContext::TestAllConditionsInternal(const FStateTreeExecu
 	{
 		const int32 ConditionIndex = ConditionsOffset + Index;
 		const FStateTreeConditionBase& Cond = CurrentFrame.StateTree->Nodes[ConditionIndex].Get<const FStateTreeConditionBase>();
-		checkf(Cond.InstanceDataHandle.GetSource() == EStateTreeDataSourceType::SharedInstanceData, TEXT("Condition only support SharedInstanceData"));
+		checkf(Cond.InstanceDataHandle.GetSource() == EStateTreeDataSourceType::SharedInstanceData 
+			   || Cond.InstanceDataHandle.GetSource() == EStateTreeDataSourceType::SharedInstanceDataObject, 
+			   TEXT("Condition only support SharedInstanceData"));
+
 		const FStateTreeDataView ConditionInstanceView = GetDataView(CurrentParentFrame, CurrentFrame, Cond.InstanceDataHandle);
 		FNodeInstanceDataScope DataScope(*this, &Cond, ConditionIndex, Cond.InstanceDataHandle, ConditionInstanceView);
 
@@ -4081,7 +4084,11 @@ float FStateTreeExecutionContext::EvaluateUtilityWithValidation(const FStateTree
 	{
 		const int32 ConsiderationIndex = ConsiderationsBegin + Index;
 		const FStateTreeConsiderationBase& Consideration = CurrentFrame.StateTree->Nodes[ConsiderationIndex].Get<const FStateTreeConsiderationBase>();
-		checkf(Consideration.InstanceDataHandle.GetSource() == EStateTreeDataSourceType::SharedInstanceData, TEXT("Consideration only support SharedInstanceData"));
+
+		checkf(Consideration.InstanceDataHandle.GetSource() == EStateTreeDataSourceType::SharedInstanceData
+			   || Consideration.InstanceDataHandle.GetSource() == EStateTreeDataSourceType::SharedInstanceDataObject
+			   , TEXT("Consideration only support SharedInstanceData"));
+
 		const FStateTreeDataView ConsiderationInstanceView = GetDataView(CurrentParentFrame, CurrentFrame, Consideration.InstanceDataHandle);
 		FNodeInstanceDataScope DataScope(*this, &Consideration, ConsiderationIndex, Consideration.InstanceDataHandle, ConsiderationInstanceView);
 
