[State Tree] mistakenly didn't route to the correct constructor in ExecContext

diff --git a/Engine/Plugins/Runtime/StateTree/Source/StateTreeModule/Private/StateTreeExecutionContext.cpp b/Engine/Plugins/Runtime/StateTree/Source/StateTreeModule/Private/StateTreeExecutionContext.cpp
--- a/Engine/Plugins/Runtime/StateTree/Source/StateTreeModule/Private/StateTreeExecutionContext.cpp
+++ b/Engine/Plugins/Runtime/StateTree/Source/StateTreeModule/Private/StateTreeExecutionContext.cpp
@@ -784,8 +784,8 @@ FStateTreeExecutionContext::FStateTreeExecutionContext(TNotNull<UObject*> InOwne
 	}
 }
 
-FStateTreeExecutionContext::FStateTreeExecutionContext(const FStateTreeExecutionContext & InContextToCopy, const UStateTree & InStateTree, FStateTreeInstanceData & InInstanceData)
-	: FStateTreeExecutionContext(InContextToCopy.Owner, InStateTree, InInstanceData, InContextToCopy.CollectExternalDataDelegate)
+FStateTreeExecutionContext::FStateTreeExecutionContext(const FStateTreeExecutionContext& InContextToCopy, const UStateTree& InStateTree, FStateTreeInstanceData& InInstanceData)
+	: FStateTreeExecutionContext(InContextToCopy, &InStateTree, InInstanceData)
 {
 }
 
