[StateTreeDebugger] FStateTreeInstanceDebugId <-> uint64 conversion to facilitate integration in RewindDebugger tracks

diff --git a/Engine/Plugins/Runtime/StateTree/Source/StateTreeModule/Public/StateTreeExecutionTypes.h b/Engine/Plugins/Runtime/StateTree/Source/StateTreeModule/Public/StateTreeExecutionTypes.h
--- a/Engine/Plugins/Runtime/StateTree/Source/StateTreeModule/Public/StateTreeExecutionTypes.h
+++ b/Engine/Plugins/Runtime/StateTree/Source/StateTreeModule/Public/StateTreeExecutionTypes.h
@@ -658,7 +658,13 @@ struct FStateTreeInstanceDebugId
 {
 	FStateTreeInstanceDebugId() = default;
 	FStateTreeInstanceDebugId(const uint32 InstanceId, const uint32 SerialNumber)
-		: Id(InstanceId), SerialNumber(SerialNumber)
+		: Id(InstanceId)
+		, SerialNumber(SerialNumber)
+	{
+	}
+	explicit FStateTreeInstanceDebugId(const uint64 Id)
+		: Id(static_cast<uint32>(Id >> 32))
+		, SerialNumber(static_cast<uint32>(Id))
 	{
 	}
 	
@@ -676,6 +682,11 @@ struct FStateTreeInstanceDebugId
 		return !(*this == Other);
 	}
 
+	uint64 ToUint64() const
+	{
+		return (static_cast<uint64>(Id) << 32) | static_cast<uint64>(SerialNumber);
+	}
+
 	friend uint32 GetTypeHash(const FStateTreeInstanceDebugId InstanceDebugId)
 	{
 		return HashCombine(InstanceDebugId.Id, InstanceDebugId.SerialNumber);
