StateTree: Disable enter/exist runtime validation for 5.6. There's an issue when the objects are garbage collected in a different order than expected.

diff --git a/Engine/Plugins/Runtime/StateTree/Source/StateTreeModule/Private/Debugger/StateTreeRuntimeValidation.cpp b/Engine/Plugins/Runtime/StateTree/Source/StateTreeModule/Private/Debugger/StateTreeRuntimeValidation.cpp
--- a/Engine/Plugins/Runtime/StateTree/Source/StateTreeModule/Private/Debugger/StateTreeRuntimeValidation.cpp
+++ b/Engine/Plugins/Runtime/StateTree/Source/StateTreeModule/Private/Debugger/StateTreeRuntimeValidation.cpp
@@ -28,7 +28,7 @@ static FAutoConsoleVariableRef CVarRuntimeValidationDoesNewerVersionExists(
 	TEXT("Test if a StateTreeExecutionContext started with an old version of a blueprint type.")
 );
 
-bool bRuntimeValidationEnterExitState = true;
+bool bRuntimeValidationEnterExitState = false;
 static FAutoConsoleVariableRef CVarRuntimeValidationEnterExitState(
 	TEXT("StateTree.RuntimeValidation.EnterExitState"),
 	bRuntimeValidationEnterExitState,
