StateTree: Fix scheduled tick when the first state completes on EnterState. Add unittest to detect the issue when the scheduled tick would sleep, and another when the scheduled tick would tick.
#rnx

diff --git a/Engine/Plugins/Runtime/StateTree/Source/StateTreeModule/Private/StateTreeExecutionContext.cpp b/Engine/Plugins/Runtime/StateTree/Source/StateTreeModule/Private/StateTreeExecutionContext.cpp
--- a/Engine/Plugins/Runtime/StateTree/Source/StateTreeModule/Private/StateTreeExecutionContext.cpp
+++ b/Engine/Plugins/Runtime/StateTree/Source/StateTreeModule/Private/StateTreeExecutionContext.cpp
@@ -2793,6 +2793,7 @@ EStateTreeRunStatus FStateTreeExecutionContext::EnterState(FStateTreeTransitionR
 
 	UE_STATETREE_DEBUG_ACTIVE_STATES_EVENT(this, Exec.ActiveFrames);
 
+	Exec.bHasPendingCompletedState = Result != EStateTreeRunStatus::Running;
 	return Result;
 }
 
diff --git a/Engine/Plugins/Runtime/StateTree/Source/StateTreeTestSuite/Private/StateTreeTest.cpp b/Engine/Plugins/Runtime/StateTree/Source/StateTreeTestSuite/Private/StateTreeTest.cpp
--- a/Engine/Plugins/Runtime/StateTree/Source/StateTreeTestSuite/Private/StateTreeTest.cpp
+++ b/Engine/Plugins/Runtime/StateTree/Source/StateTreeTestSuite/Private/StateTreeTest.cpp
@@ -1015,6 +1015,98 @@ struct FStateTreeTest_StateRequiringEvent : FStateTreeTestBase
 };
 IMPLEMENT_AI_INSTANT_TEST(FStateTreeTest_StateRequiringEvent, "System.StateTree.StateRequiringEvent");
 
+struct FStateTreeTest_Start : FStateTreeTestBase
+{
+	virtual UStateTree& SetupTree()
+	{
+		UStateTree& StateTree = NewStateTree();
+		UStateTreeEditorData& EditorData = *Cast<UStateTreeEditorData>(StateTree.EditorData);
+
+		UStateTreeState& Root = EditorData.AddSubTree(FName(TEXT("Root")));
+		UStateTreeState& StateA = Root.AddChildState(FName(TEXT("A")));
+		UStateTreeState& StateB = Root.AddChildState(FName(TEXT("B")));
+		TStateTreeEditorNode<FTestTask_Stand>& TaskA = StateA.AddTask<FTestTask_Stand>(TaskAName);
+		TStateTreeEditorNode<FTestTask_Stand>& TaskB = StateB.AddTask<FTestTask_Stand>(TaskBName);
+
+		// Transition success 
+		StateA.AddTransition(EStateTreeTransitionTrigger::OnStateSucceeded, EStateTreeTransitionType::GotoState, &StateB);
+		TaskA.GetNode().EnterStateResult = EStateTreeRunStatus::Succeeded;
+
+		return StateTree;
+	}
+
+	virtual bool InstantTest() override
+	{
+		UStateTree& StateTree = SetupTree();
+
+		FStateTreeCompilerLog Log;
+		FStateTreeCompiler Compiler(Log);
+		const bool bResult = Compiler.Compile(StateTree);
+
+		AITEST_TRUE("StateTree should get compiled", bResult);
+
+		FStateTreeInstanceData InstanceData;
+		FTestStateTreeExecutionContext Exec(StateTree, StateTree, InstanceData);
+		const bool bInitSucceeded = Exec.IsValid();
+		AITEST_TRUE("StateTree should init", bInitSucceeded);
+
+		const TCHAR* TickStr(TEXT("Tick"));
+		const TCHAR* EnterStateStr(TEXT("EnterState"));
+		const TCHAR* ExitStateStr(TEXT("ExitState"));
+		const TCHAR* StateCompletedStr(TEXT("StateCompleted"));
+
+		{
+			EStateTreeRunStatus Status = Exec.Start();
+			AITEST_EQUAL("Start should be running", Status, EStateTreeRunStatus::Running);
+			AITEST_TRUE("StateTree Active States should be in states.", Exec.ExpectInActiveStates("Root", "A"));
+			AITEST_TRUE("StateTree TaskA should enter state", Exec.Expect(TaskAName, EnterStateStr));
+			AITEST_TRUE("StateTree TaskA should state complete", Exec.Expect(TaskAName, StateCompletedStr));
+			AITEST_TRUE("StateTree execution should not sleep", !Exec.GetNextScheduledTick().ShouldSleep());
+			Exec.LogClear();
+		}
+		{
+			EStateTreeRunStatus Status = Exec.Tick(0.1f);
+			AITEST_TRUE("StateTree Active States should be in states.", Exec.ExpectInActiveStates("Root", "B"));
+			//@TODO Only one StateComplete
+			//AITEST_FALSE("StateTree TaskA should state complete", Exec.Expect(TaskAName, StateCompletedStr));
+			AITEST_TRUE("StateTree TaskA should get exit state expectedly", Exec.Expect(TaskAName, ExitStateStr));
+			AITEST_TRUE("StateTree TaskB should get enter state expectedly", Exec.Expect(TaskBName, EnterStateStr));
+			Exec.LogClear();
+		}
+
+		Exec.Stop();
+
+		return true;
+	}
+
+protected:
+	const FName TaskAName = FName("TaskA");
+	const FName TaskBName = FName("TaskB");
+};
+IMPLEMENT_AI_INSTANT_TEST(FStateTreeTest_Start, "System.StateTree.Start.FirstStateSucceededImmediately");
+
+struct FStateTreeTest_StartScheduledTick : FStateTreeTest_Start
+{
+	virtual UStateTree& SetupTree()
+	{
+		UStateTree& StateTree = NewStateTree();
+		UStateTreeEditorData& EditorData = *Cast<UStateTreeEditorData>(StateTree.EditorData);
+
+		UStateTreeState& Root = EditorData.AddSubTree(FName(TEXT("Root")));
+		UStateTreeState& StateA = Root.AddChildState(FName(TEXT("A")));
+		UStateTreeState& StateB = Root.AddChildState(FName(TEXT("B")));
+		TStateTreeEditorNode<FTestTask_StandNoTick>& TaskA = StateA.AddTask<FTestTask_StandNoTick>(TaskAName);
+		TStateTreeEditorNode<FTestTask_StandNoTick>& TaskB = StateB.AddTask<FTestTask_StandNoTick>(TaskBName);
+
+		// Transition success 
+		StateA.AddTransition(EStateTreeTransitionTrigger::OnStateSucceeded, EStateTreeTransitionType::GotoState, &StateB);
+		TaskA.GetNode().EnterStateResult = EStateTreeRunStatus::Succeeded;
+
+		return StateTree;
+	}
+};
+IMPLEMENT_AI_INSTANT_TEST(FStateTreeTest_StartScheduledTick, "System.StateTree.Start.FirstStateSucceededImmediatelyWithScheduledTick");
+
 //
 // The stop tests test how the combinations of execution path to stop the tree are reported on ExitState() transition.  
 //
diff --git a/Engine/Plugins/Runtime/StateTree/Source/StateTreeTestSuite/Private/StateTreeTest.h b/Engine/Plugins/Runtime/StateTree/Source/StateTreeTestSuite/Private/StateTreeTest.h
--- a/Engine/Plugins/Runtime/StateTree/Source/StateTreeTestSuite/Private/StateTreeTest.h
+++ b/Engine/Plugins/Runtime/StateTree/Source/StateTreeTestSuite/Private/StateTreeTest.h
@@ -10,8 +10,18 @@ class UStateTreeTestSchema : public UStateTreeSchema
 {
 	GENERATED_BODY()
 
-	virtual bool IsStructAllowed(const UScriptStruct* InScriptStruct) const override { return true; }
-	virtual bool IsExternalItemAllowed(const UStruct& InStruct) const override { return true; }
+	virtual bool IsStructAllowed(const UScriptStruct* InScriptStruct) const override
+	{
+		return true;
+	}
+	virtual bool IsExternalItemAllowed(const UStruct& InStruct) const override
+	{
+		return true;
+	}
+	virtual bool IsScheduledTickAllowed() const
+	{
+		return true;
+	}
 };
 
 UCLASS(HideDropdown)
diff --git a/Engine/Plugins/Runtime/StateTree/Source/StateTreeTestSuite/Private/StateTreeTestTypes.h b/Engine/Plugins/Runtime/StateTree/Source/StateTreeTestSuite/Private/StateTreeTestTypes.h
--- a/Engine/Plugins/Runtime/StateTree/Source/StateTreeTestSuite/Private/StateTreeTestTypes.h
+++ b/Engine/Plugins/Runtime/StateTree/Source/StateTreeTestSuite/Private/StateTreeTestTypes.h
@@ -572,6 +572,23 @@ struct FTestTask_Stand : public FStateTreeTaskBase
 	EStateTreeRunStatus EnterStateResult = EStateTreeRunStatus::Running;
 };
 
+USTRUCT()
+struct FTestTask_StandNoTick : public FTestTask_Stand
+{
+	GENERATED_BODY()
+
+	FTestTask_StandNoTick()
+	{
+		bShouldCallTick = false;
+	}
+
+	FTestTask_StandNoTick(const FName InName)
+		: FTestTask_Stand(InName)
+	{
+		bShouldCallTick = false;
+	}
+};
+
 USTRUCT()
 struct FStateTreeTestConditionInstanceData
 {
