Back out "[AutoRTFM] Make UStateTree::GetSharedInstanceData transactionally safe" in 36.10.

This was originally landed in Release-35.10 at https://p4-swarm.epicgames.net/changes/42245712. This CL experimentally backs out this change to see if it somehow affects garbage collection behavior.

** THIS CHANGE MUST BE BACKED OUT BEFORE 36.10 GOES LIVE OR WE WILL INTRODUCE CRASHES IN PRODUCTION **

diff --git a/Engine/Plugins/Runtime/StateTree/Source/StateTreeModule/Private/StateTree.cpp b/Engine/Plugins/Runtime/StateTree/Source/StateTreeModule/Private/StateTree.cpp
--- a/Engine/Plugins/Runtime/StateTree/Source/StateTreeModule/Private/StateTree.cpp
+++ b/Engine/Plugins/Runtime/StateTree/Source/StateTreeModule/Private/StateTree.cpp
@@ -15,9 +15,6 @@
 #include "Misc/EnumerateRange.h"
 #include "UObject/AssetRegistryTagsContext.h"
 #include "StateTreePropertyFunctionBase.h"
-#include "AutoRTFM.h"
-
-#include <atomic>
 
 #if WITH_EDITOR
 #include "Editor.h"
@@ -102,8 +99,7 @@ FGuid UStateTree::GetTransitionIdFromIndex(const FStateTreeIndex16 Index) const
 	return Entry != nullptr ? Entry->Id : FGuid();
 }
 
-UE_AUTORTFM_ALWAYS_OPEN
-static int32 GetThreadIndexForSharedInstanceData()
+TSharedPtr<FStateTreeInstanceData> UStateTree::GetSharedInstanceData() const
 {
 	// Create a unique index for each thread.
 	static std::atomic_int ThreadIndexCounter {0};
@@ -112,17 +108,10 @@ static int32 GetThreadIndexForSharedInstanceData()
 	{
 		ThreadIndex = ThreadIndexCounter.fetch_add(1);
 	}
-
-	return ThreadIndex;
-}
-
-TSharedPtr<FStateTreeInstanceData> UStateTree::GetSharedInstanceData() const
-{
-	int32 ThreadIndex = GetThreadIndexForSharedInstanceData();
-
+ 
 	// If shared instance data for this thread exists, return it.
 	{
-		UE::TReadScopeLock ReadLock(PerThreadSharedInstanceDataLock);
+		FReadScopeLock ReadLock(PerThreadSharedInstanceDataLock);
 		if (ThreadIndex < PerThreadSharedInstanceData.Num())
 		{
 			return PerThreadSharedInstanceData[ThreadIndex];
@@ -130,7 +119,7 @@ TSharedPtr<FStateTreeInstanceData> UStateTree::GetSharedInstanceData() const
 	}
 
 	// Not initialized yet, create new instances up to the index.
-	UE::TWriteScopeLock WriteLock(PerThreadSharedInstanceDataLock);
+	FWriteScopeLock WriteLock(PerThreadSharedInstanceDataLock);
 
 	// It is possible that multiple threads are waiting for the write lock,
 	// which means that execution may get here so that 'ThreadIndex' is already in valid range.
@@ -426,7 +415,7 @@ void UStateTree::AddReferencedObjects(UObject* InThis, FReferenceCollector& Coll
 	const UStateTree* StateTree = Cast<UStateTree>(InThis);
 	check(StateTree);
 	
-	UE::TReadScopeLock ReadLock(StateTree->PerThreadSharedInstanceDataLock);
+	FReadScopeLock ReadLock(StateTree->PerThreadSharedInstanceDataLock);
 
 	for (const TSharedPtr<FStateTreeInstanceData>& InstanceData : StateTree->PerThreadSharedInstanceData)
 	{
@@ -530,7 +519,7 @@ void UStateTree::ResetLinked()
 	OutOfDateStructs.Reset();
 #endif
 
-	UE::TWriteScopeLock WriteLock(PerThreadSharedInstanceDataLock);
+	FWriteScopeLock WriteLock(PerThreadSharedInstanceDataLock);
 	PerThreadSharedInstanceData.Reset();
 }
 
diff --git a/Engine/Plugins/Runtime/StateTree/Source/StateTreeModule/Public/StateTree.h b/Engine/Plugins/Runtime/StateTree/Source/StateTreeModule/Public/StateTree.h
--- a/Engine/Plugins/Runtime/StateTree/Source/StateTreeModule/Public/StateTree.h
+++ b/Engine/Plugins/Runtime/StateTree/Source/StateTreeModule/Public/StateTree.h
@@ -7,7 +7,6 @@
 #include "StateTreePropertyBindings.h"
 #include "StateTreeInstanceData.h"
 #include "StateTreeTypes.h"
-#include "Misc/TransactionallySafeRWLock.h"
 #include "UObject/ObjectKey.h"
 #include "StateTree.generated.h"
 
@@ -322,7 +321,7 @@ private:
 	UPROPERTY()
 	FStateTreeInstanceData SharedInstanceData;
 
-	mutable FTransactionallySafeRWLock PerThreadSharedInstanceDataLock;
+	mutable FRWLock PerThreadSharedInstanceDataLock;
 	mutable TArray<TSharedPtr<FStateTreeInstanceData>> PerThreadSharedInstanceData;
 	
 	/** List of names external data enforced by the schema, created at compilation. */
