[StateTreeDebugger] instance is also traced as object child element for RewindDebugger

diff --git a/Engine/Plugins/Runtime/StateTree/Source/StateTreeModule/Private/Debugger/StateTreeTrace.cpp b/Engine/Plugins/Runtime/StateTree/Source/StateTreeModule/Private/Debugger/StateTreeTrace.cpp
--- a/Engine/Plugins/Runtime/StateTree/Source/StateTreeModule/Private/Debugger/StateTreeTrace.cpp
+++ b/Engine/Plugins/Runtime/StateTree/Source/StateTreeModule/Private/Debugger/StateTreeTrace.cpp
@@ -192,10 +192,12 @@ struct FInstanceEventBufferedData
 
 	FInstanceEventBufferedData() = default;
 	explicit FInstanceEventBufferedData(
+		const UObject* Owner,
 		const UStateTree* StateTree,
 		const FStateTreeInstanceDebugId InstanceId,
 		const FString& InstanceName)
-		: InstanceName(InstanceName)
+		: WeakOwner(Owner)
+		, InstanceName(InstanceName)
 		, WeakStateTree(StateTree)
 		, InstanceId(InstanceId)
 	{
@@ -215,7 +217,7 @@ struct FInstanceEventBufferedData
 		{
 			if (const UStateTree* StateTree = WeakStateTree.Get())
 			{
-				OutputInstanceLifetimeEvent(InstanceId, StateTree, *InstanceName, EStateTreeTraceEventType::Push);
+				OutputInstanceLifetimeEvent(WeakOwner.Get(), InstanceId, StateTree, *InstanceName, EStateTreeTraceEventType::Push);
 
 				if (ActiveStates.IsValid())
 				{
@@ -225,6 +227,7 @@ struct FInstanceEventBufferedData
 		}
 	}
 
+	TWeakObjectPtr<const UObject> WeakOwner;
 	FActiveStates ActiveStates;
 	FString InstanceName;
 	TWeakObjectPtr<const UStateTree> WeakStateTree;
@@ -262,7 +265,7 @@ struct FBufferedDataList
 	void UpdateActiveStates_AnyThread(const FStateTreeInstanceDebugId InstanceId, const TConstArrayView<FStateTreeExecutionFrame> ActiveFrames);
 
 	/** Keeps track of created/destroyed statetree instances */
-	void UpdateInstanceLifetime_AnyThread(const FStateTreeInstanceDebugId InstanceId, const UStateTree* StateTree, const TCHAR* InstanceName, const EStateTreeTraceEventType EventType);
+	void UpdateInstanceLifetime_AnyThread(const UObject* Owner, const FStateTreeInstanceDebugId InstanceId, const UStateTree* StateTree, const TCHAR* InstanceName, const EStateTreeTraceEventType EventType);
 
 	/** Sets the world time that might need to be traced by worker threads. */
 	void SetWorldTime_GameThread(const double WorldTime);
@@ -547,7 +550,7 @@ void FBufferedDataList::UpdateActiveStates_AnyThread(const FStateTreeInstanceDeb
 	}
 }
 
-void FBufferedDataList::UpdateInstanceLifetime_AnyThread(const FStateTreeInstanceDebugId InstanceId, const UStateTree* StateTree, const TCHAR* InstanceName
+void FBufferedDataList::UpdateInstanceLifetime_AnyThread(const UObject* Owner, const FStateTreeInstanceDebugId InstanceId, const UStateTree* StateTree, const TCHAR* InstanceName
 	, const EStateTreeTraceEventType EventType)
 {
 	// This can be called by worker threads from statetree trace macros and by the game thread during the flush.
@@ -559,7 +562,7 @@ void FBufferedDataList::UpdateInstanceLifetime_AnyThread(const FStateTreeInstanc
 
 		if (EventType == EStateTreeTraceEventType::Push)
 		{
-			BufferedEventsMap.Emplace(InstanceId, FInstanceEventBufferedData(StateTree, InstanceId, InstanceName));
+			BufferedEventsMap.Emplace(InstanceId, FInstanceEventBufferedData(Owner, StateTree, InstanceId, InstanceName));
 		}
 		else if (EventType == EStateTreeTraceEventType::Pop)
 		{
@@ -709,6 +712,12 @@ FStateTreeIndex16 FindOrAddDebugIdForAsset(const UStateTree* StateTree)
 	return GBufferedEvents.FindOrAddDebugIdForAsset_AnyThread(StateTree);
 }
 
+// Deprecated
+void OutputInstanceLifetimeEvent(const FStateTreeInstanceDebugId InstanceId, const UStateTree* StateTree, const TCHAR* InstanceName, const EStateTreeTraceEventType EventType)
+{
+	OutputInstanceLifetimeEvent(nullptr, InstanceId, StateTree, InstanceName, EventType);
+}
+
 void OutputAssetDebugIdEvent(
 	const UStateTree* StateTree,
 	const FStateTreeIndex16 AssetDebugId
@@ -732,6 +741,7 @@ void OutputAssetDebugIdEvent(
 }
 
 void OutputInstanceLifetimeEvent(
+	const UObject* Owner,
 	const FStateTreeInstanceDebugId InstanceId,
 	const UStateTree* StateTree,
 	const TCHAR* InstanceName,
@@ -742,6 +752,19 @@ void OutputInstanceLifetimeEvent(
 	{
 		TraceBufferedEvents(InstanceId);
 
+		if (Owner != nullptr)
+		{
+			const uint64 ChildId(InstanceId.ToUint64());
+			if (EventType == EStateTreeTraceEventType::Push)
+			{
+				TRACE_CHILD_ELEMENT_WITH_OUTER_LIFETIME_BEGIN(Owner, ChildId, Owner, INDEX_NONE);
+			}
+			else if (EventType == EStateTreeTraceEventType::Pop)
+			{
+				TRACE_CHILD_ELEMENT_LIFETIME_END(Owner, ChildId);
+			}
+		}
+
 		const FStateTreeIndex16 AssetDebugId = GBufferedEvents.FindOrAddDebugIdForAsset_AnyThread(StateTree);
 
 		UE_TRACE_LOG(StateTreeDebugger, InstanceEvent, StateTreeDebugChannel)
@@ -756,7 +779,7 @@ void OutputInstanceLifetimeEvent(
 	// Update buffered events regardless of the status of the channel since they will be used
 	// when flushing buffered event when a late recording is started or more than one trace are started
 	// during the same game session (i.e. Start Traces -> Stop Traces -> Start Traces).
-	GBufferedEvents.UpdateInstanceLifetime_AnyThread(InstanceId, StateTree, InstanceName, EventType);
+	GBufferedEvents.UpdateInstanceLifetime_AnyThread(Owner, InstanceId, StateTree, InstanceName, EventType);
 }
 
 void OutputInstanceAssetEvent(
diff --git a/Engine/Plugins/Runtime/StateTree/Source/StateTreeModule/Public/Debugger/StateTreeDebug.h b/Engine/Plugins/Runtime/StateTree/Source/StateTreeModule/Public/Debugger/StateTreeDebug.h
--- a/Engine/Plugins/Runtime/StateTree/Source/StateTreeModule/Public/Debugger/StateTreeDebug.h
+++ b/Engine/Plugins/Runtime/StateTree/Source/StateTreeModule/Public/Debugger/StateTreeDebug.h
@@ -45,7 +45,7 @@ enum class EStateTreeTraceEventType : uint8;
 
 #define UE_STATETREE_DEBUG_LOG_EVENT(InExecutionContextPtr, InLogVerbosity, InFormat, ...) TRACE_STATETREE_LOG_EVENT((InExecutionContextPtr)->GetInstanceDebugId(), InLogVerbosity, InFormat, ##__VA_ARGS__)
 #define UE_STATETREE_DEBUG_CONDITION_EVENT(InExecutionContextPtr, InIndex, InDataView, InEventType) TRACE_STATETREE_CONDITION_EVENT((InExecutionContextPtr)->GetInstanceDebugId(), (InExecutionContextPtr)->StealNodeCustomDebugTraceData(), FStateTreeIndex16(InIndex), (InDataView), (InEventType));
-#define UE_STATETREE_DEBUG_INSTANCE_EVENT(InExecutionContextPtr, InEventType) TRACE_STATETREE_INSTANCE_EVENT((InExecutionContextPtr)->GetInstanceDebugId(), (InExecutionContextPtr)->GetStateTree(), *((InExecutionContextPtr)->GetInstanceDescriptionInternal()), (InEventType));
+#define UE_STATETREE_DEBUG_INSTANCE_EVENT(InExecutionContextPtr, InEventType) TRACE_STATETREE_INSTANCE_EVENT((InExecutionContextPtr)->GetOwner(), (InExecutionContextPtr)->GetInstanceDebugId(), (InExecutionContextPtr)->GetStateTree(), *((InExecutionContextPtr)->GetInstanceDescriptionInternal()), (InEventType));
 #define UE_STATETREE_DEBUG_INSTANCE_FRAME_EVENT(InExecutionContextPtr, InFrame) TRACE_STATETREE_INSTANCE_FRAME_EVENT((InExecutionContextPtr)->GetInstanceDebugId(), (InFrame));
 #define UE_STATETREE_DEBUG_ACTIVE_STATES_EVENT(InExecutionContextPtr, InActiveFrames) TRACE_STATETREE_ACTIVE_STATES_EVENT((InExecutionContextPtr)->GetInstanceDebugId(), InActiveFrames);
 #define UE_STATETREE_DEBUG_TASK_EVENT(InExecutionContextPtr, InIndex, InDataView, InEventType, InStatus) TRACE_STATETREE_TASK_EVENT((InExecutionContextPtr)->GetInstanceDebugId(), (InExecutionContextPtr)->StealNodeCustomDebugTraceData(), FStateTreeIndex16(InIndex), (InDataView), (InEventType), (InStatus));
diff --git a/Engine/Plugins/Runtime/StateTree/Source/StateTreeModule/Public/Debugger/StateTreeTrace.h b/Engine/Plugins/Runtime/StateTree/Source/StateTreeModule/Public/Debugger/StateTreeTrace.h
--- a/Engine/Plugins/Runtime/StateTree/Source/StateTreeModule/Public/Debugger/StateTreeTrace.h
+++ b/Engine/Plugins/Runtime/StateTree/Source/StateTreeModule/Public/Debugger/StateTreeTrace.h
@@ -101,7 +101,9 @@ UE_DEPRECATED(5.6, "This method will no longer be exposed publicly.")
 FStateTreeIndex16 FindOrAddDebugIdForAsset(const UStateTree* StateTree);
 void OutputPhaseScopeEvent(FStateTreeInstanceDebugId InstanceId, EStateTreeUpdatePhase Phase, EStateTreeTraceEventType EventType, FStateTreeStateHandle StateHandle);
 void OutputAssetDebugIdEvent(const UStateTree* StateTree, FStateTreeIndex16 AssetDebugId);
+UE_DEPRECATED(5.7, "Use the overload taking a pointer to the UObject owner instead.")
 void OutputInstanceLifetimeEvent(FStateTreeInstanceDebugId InstanceId, const UStateTree* StateTree, const TCHAR* InstanceName, EStateTreeTraceEventType EventType);
+void OutputInstanceLifetimeEvent(const UObject* Owner, FStateTreeInstanceDebugId InstanceId, const UStateTree* StateTree, const TCHAR* InstanceName, EStateTreeTraceEventType EventType);
 void OutputInstanceAssetEvent(FStateTreeInstanceDebugId InstanceId, const UStateTree* StateTree);
 void OutputInstanceFrameEvent(FStateTreeInstanceDebugId InstanceId, const FStateTreeExecutionFrame* Frame);
 void OutputLogEventTrace(FStateTreeInstanceDebugId InstanceId, ELogVerbosity::Type Verbosity, const TCHAR* Fmt, ...);
@@ -114,8 +116,8 @@ void OutputActiveStatesEventTrace(FStateTreeInstanceDebugId InstanceId, TConstAr
 
 } // UE::StateTreeTrace
 
-#define TRACE_STATETREE_INSTANCE_EVENT(InstanceID, StateTree, InstanceName, EventType) \
-	UE::StateTreeTrace::OutputInstanceLifetimeEvent(InstanceID, StateTree, InstanceName, EventType);
+#define TRACE_STATETREE_INSTANCE_EVENT(Owner, InstanceID, StateTree, InstanceName, EventType) \
+	UE::StateTreeTrace::OutputInstanceLifetimeEvent(Owner, InstanceID, StateTree, InstanceName, EventType);
 
 #define TRACE_STATETREE_INSTANCE_FRAME_EVENT(InstanceID, Frame) \
 	UE::StateTreeTrace::OutputInstanceFrameEvent(InstanceID, Frame);
