[State Tree Editor] Used DetailsView to fetch root UObject to updated Compilation hash when changing ST EditorData in the editor

diff --git a/Engine/Plugins/Runtime/StateTree/Source/StateTreeEditorModule/Private/StateTreeEditorMode.cpp b/Engine/Plugins/Runtime/StateTree/Source/StateTreeEditorModule/Private/StateTreeEditorMode.cpp
--- a/Engine/Plugins/Runtime/StateTree/Source/StateTreeEditorModule/Private/StateTreeEditorMode.cpp
+++ b/Engine/Plugins/Runtime/StateTree/Source/StateTreeEditorModule/Private/StateTreeEditorMode.cpp
@@ -706,17 +706,22 @@ void UStateTreeEditorMode::OnAssetFinishedChangingProperties(const FPropertyChan
 	// Make sure nodes get updates when properties are changed.
 	if(UStateTreeEditingSubsystem* StateTreeEditingSubsystem = GEditor->GetEditorSubsystem<UStateTreeEditingSubsystem>())
 	{
-		const int32 NumEditedObjects = PropertyChangedEvent.GetNumObjectsBeingEdited();
-		if (NumEditedObjects > 0)
-		{			
-			for (int32 Index = 0; Index < NumEditedObjects; ++Index)
+		if (TSharedPtr<IDetailsView> AssetDetailsView = GetAssetDetailsView())
+		{
+			// From the path FPropertyHandleBase::NotifyFinishedChangingProperties(), UObject info is not included in PropertyChangedEvent
+			// So we fetch it from DetailsView
+			const TArray<TWeakObjectPtr<UObject>>& SelectedObjects = AssetDetailsView->GetSelectedObjects();
+			for (const TWeakObjectPtr<UObject>& WeakObject : SelectedObjects)
 			{
-				if (const UStateTree* EditedStateTree = Cast<UStateTree>(PropertyChangedEvent.GetObjectBeingEdited(Index)))
+				if (const UStateTreeEditorData* EditorData = Cast<UStateTreeEditorData>(WeakObject.Get()))
 				{
-					if (EditedStateTree == GetStateTree())
+					if (const UStateTree* StateTree = Cast<UStateTree>(EditorData->GetOuter()))
 					{
-						StateTreeEditingSubsystem->FindOrAddViewModel(GetStateTree())->NotifyAssetChangedExternally();
-						break;
+						if (GetStateTree() == StateTree)
+						{
+							StateTreeEditingSubsystem->FindOrAddViewModel(GetStateTree())->NotifyAssetChangedExternally();
+							break;
+						}
 					}
 				}
 			}
@@ -731,7 +736,7 @@ void UStateTreeEditorMode::OnSelectionFinishedChangingProperties(const FProperty
 	{
 		if (TSharedPtr<IDetailsView> DetailsView = GetDetailsView())
 		{
-			const TArray<TWeakObjectPtr<UObject>> SelectedObjects = GetDetailsView()->GetSelectedObjects();
+			const TArray<TWeakObjectPtr<UObject>>& SelectedObjects = DetailsView->GetSelectedObjects();
 			TSet<UStateTreeState*> ChangedStates;
 			for (const TWeakObjectPtr<UObject>& WeakObject : SelectedObjects)
 			{
