StateTree: Add a crash reporter handler for statetree. When a crash occurs inside the state tree call stack (tick, start, stop),  the crash reporter will output the name of the state tree asset being used. It will output the names of all the state tree assets currently used in all the different threads. Because of that, it is possible that we cannot identify the exact tree that causes the crash.

diff --git a/Engine/Plugins/Runtime/StateTree/Source/StateTreeModule/Private/CrashReporter/StateTreeCrashReporterHandler.cpp b/Engine/Plugins/Runtime/StateTree/Source/StateTreeModule/Private/CrashReporter/StateTreeCrashReporterHandler.cpp
new file mode 100644
--- /dev/null
+++ b/Engine/Plugins/Runtime/StateTree/Source/StateTreeModule/Private/CrashReporter/StateTreeCrashReporterHandler.cpp
@@ -0,0 +1,159 @@
+// Copyright Epic Games, Inc. All Rights Reserved.
+
+#include "CrashReporter/StateTreeCrashReporterHandler.h"
+
+#include "Misc/ScopeLock.h"
+#include "HAL/PlatformTLS.h"
+
+#include "StateTree.h"
+
+#if UE_WITH_STATETREE_CRASHREPORTER
+
+namespace UE::StateTree::Private
+{
+static bool bCrashHandlerEnabled = true;
+static FAutoConsoleVariableRef CVarCrashHandlerEnabled(
+	TEXT("StateTree.CrashHandlerEnabled"),
+	bCrashHandlerEnabled,
+	TEXT("StateTree will add context into the crash reporter."),
+	ECVF_Default);
+
+struct FCrashReporterHandlerImpl : FCrashReporterHandler
+{
+public:
+	struct FContext
+	{
+		explicit FContext(TNotNull<const UObject*> InOwner, TNotNull<const UStateTree*> InStateTree, FName InContext, uint32 InThreadID)
+			: Owner(InOwner)
+			, StateTree(InStateTree)
+			, Context(InContext)
+			, ThreadID(InThreadID)
+		{
+			static std::atomic IDGenerator = 0;
+			ID = ++IDGenerator;
+		}
+
+		TWeakObjectPtr<const UObject> Owner;
+		TWeakObjectPtr<const UStateTree> StateTree;
+		FName Context;
+		uint32 ThreadID;
+		uint32 ID;
+	};
+
+public:
+	uint32 PushInfo(TNotNull<const UObject*> InOwner, TNotNull<const UStateTree*> InStateTree, FName InContext)
+	{
+		FScopeLock LockGuard(&CriticalSection);
+		return Contexts.Emplace_GetRef(InOwner, InStateTree, InContext, FPlatformTLS::GetCurrentThreadId()).ID;
+	}
+
+	void PopInfo(uint32 ID)
+	{
+		FScopeLock LockGuard(&CriticalSection);
+		const int32 FoundIndex = Contexts.IndexOfByPredicate([ID](const FContext& Other)
+			{
+				return Other.ID == ID;
+			});
+		if (FoundIndex != INDEX_NONE)
+		{
+			Contexts.RemoveAtSwap(FoundIndex);
+		}
+	}
+
+	void DumpInfo(FCrashContextExtendedWriter& Writer)
+	{
+		// Callback from the crash reporter.Avoids memory allocation in case the crash is an OOM.
+		FScopeLock LockGuard(&CriticalSection);
+
+		for (int32 Index = 0; Index < Contexts.Num(); ++Index)
+		{
+			const FContext& Context = Contexts[Index];
+			Builder.Reset();
+
+			if (const UObject* Owner = Context.Owner.Get())
+			{
+				Owner->GetFullName(Builder);
+			}
+			else
+			{
+				Builder.Append(TEXT("<Invalid>"));
+			}
+
+			Builder << TEXT('\n');
+			if (const UStateTree* StateTree = Context.StateTree.Get())
+			{
+				StateTree->GetFullName(Builder);
+			}
+			else
+			{
+				Builder.Append(TEXT("<Invalid>"));
+			}
+
+			Builder << TEXT('\n');
+			Builder << Context.Context;
+
+			Builder << TEXT('\n');
+			Builder << Context.ThreadID;
+
+			if (Builder.Len())
+			{
+				// Avoid allocating a new name for each entry.
+				TCHAR Identifier[] = TEXT("StateTree_000");
+				constexpr int32 NumberIndex = 10;
+				Identifier[NumberIndex + 0] = (TCHAR)((Index / 100) % 10 + int32(TEXT('0')));
+				Identifier[NumberIndex + 1] = (TCHAR)((Index / 10) % 10 + int32(TEXT('0')));
+				Identifier[NumberIndex + 2] = (TCHAR)((Index / 1) % 10 + int32(TEXT('0')));
+				Writer.AddString(Identifier, Builder.ToString());
+			}
+		}
+	}
+
+private:
+	TArray<FContext, TInlineAllocator<32>> Contexts;
+	FCriticalSection CriticalSection;
+	TStringBuilder<1024> Builder;
+};
+
+
+static FCrashReporterHandlerImpl* Instance = nullptr;
+static FDelegateHandle AdditionalCrashContextDelegateHandle;
+
+} // namespace UE::StateTree::Private
+
+namespace UE::StateTree
+{
+
+void FCrashReporterHandler::Register()
+{
+	Private::Instance = new Private::FCrashReporterHandlerImpl();
+	Private::AdditionalCrashContextDelegateHandle = FGenericCrashContext::OnAdditionalCrashContextDelegate().AddRaw(Private::Instance, &Private::FCrashReporterHandlerImpl::DumpInfo);
+}
+
+void FCrashReporterHandler::Unregister()
+{
+	FGenericCrashContext::OnAdditionalCrashContextDelegate().Remove(Private::AdditionalCrashContextDelegateHandle);
+	Private::Instance = nullptr;
+}
+
+FCrashReporterScope::FCrashReporterScope(TNotNull<const UObject*> Owner, TNotNull<const UStateTree*> StateTree, FName Context)
+{
+	bWasEnabled = Private::bCrashHandlerEnabled;
+	if (bWasEnabled)
+	{
+		check(Private::Instance);
+		ID = Private::Instance->PushInfo(Owner, StateTree, Context);
+	}
+}
+
+
+FCrashReporterScope::~FCrashReporterScope()
+{
+	if (bWasEnabled)
+	{
+		Private::Instance->PopInfo(ID);
+	}
+}
+
+} //namespace UE::StateTree
+
+#endif
diff --git a/Engine/Plugins/Runtime/StateTree/Source/StateTreeModule/Private/CrashReporter/StateTreeCrashReporterHandler.h b/Engine/Plugins/Runtime/StateTree/Source/StateTreeModule/Private/CrashReporter/StateTreeCrashReporterHandler.h
new file mode 100644
--- /dev/null
+++ b/Engine/Plugins/Runtime/StateTree/Source/StateTreeModule/Private/CrashReporter/StateTreeCrashReporterHandler.h
@@ -0,0 +1,54 @@
+// Copyright Epic Games, Inc. All Rights Reserved.
+
+#pragma once
+
+#include "GenericPlatform/GenericPlatformCrashContext.h"
+#include "Misc/NotNull.h"
+
+#ifndef UE_WITH_STATETREE_CRASHREPORTER
+	#define UE_WITH_STATETREE_CRASHREPORTER WITH_ADDITIONAL_CRASH_CONTEXTS
+#endif
+
+class UObject;
+class UStateTree;
+
+#if UE_WITH_STATETREE_CRASHREPORTER
+
+namespace UE::StateTree
+{
+
+/**
+ * Additional crash context for StateTree.
+ * Will add extra information on the asset when crashing inside a StateTree callstack.
+ */
+struct FCrashReporterHandler
+{
+	static void Register();
+	static void Unregister();
+};
+
+/** Helper object allowing easy tracking of StateTree code in crash reporter. */
+class FCrashReporterScope
+{
+public:
+	explicit FCrashReporterScope(TNotNull<const UObject*> Owner, TNotNull<const UStateTree*> StateTree, FName Context);
+	~FCrashReporterScope();
+
+	FCrashReporterScope() = delete;
+	FCrashReporterScope(const FCrashReporterScope&) = delete;
+	FCrashReporterScope& operator=(const FCrashReporterScope&) = delete;
+
+private:
+	bool bWasEnabled = false;
+	uint32 ID = 0;
+};
+
+} //namespace
+
+#endif //UE_WITH_STATETREE_CRASHREPORTER
+
+#if UE_WITH_STATETREE_CRASHREPORTER
+#define UE_STATETREE_CRASH_REPORTER_SCOPE(InOwner, InStateTree, InContext) ::UE::StateTree::FCrashReporterScope ANONYMOUS_VARIABLE(AddCrashContext) {(InOwner), (InStateTree), (InContext)}
+#else //UE_WITH_STATETREE_CRASHREPORTER
+#define UE_STATETREE_CRASH_REPORTER_SCOPE(InOwner, InStateTree, InContext)
+#endif //UE_WITH_STATETREE_CRASHREPORTER
diff --git a/Engine/Plugins/Runtime/StateTree/Source/StateTreeModule/Private/StateTreeExecutionContext.cpp b/Engine/Plugins/Runtime/StateTree/Source/StateTreeModule/Private/StateTreeExecutionContext.cpp
--- a/Engine/Plugins/Runtime/StateTree/Source/StateTreeModule/Private/StateTreeExecutionContext.cpp
+++ b/Engine/Plugins/Runtime/StateTree/Source/StateTreeModule/Private/StateTreeExecutionContext.cpp
@@ -12,6 +12,7 @@
 #include "StateTreeReference.h"
 #include "AutoRTFM.h"
 #include "Containers/StaticArray.h"
+#include "CrashReporter/StateTreeCrashReporterHandler.h"
 #include "Debugger/StateTreeDebug.h"
 #include "Debugger/StateTreeTrace.h"
 #include "Debugger/StateTreeTraceTypes.h"
@@ -59,6 +60,10 @@ namespace Private
 		TEXT("You should upgrade your asset. This option is to help migrate to the new behavior.")
 	);
 
+	const FLazyName Name_Tick = "Tick";
+	const FLazyName Name_Start = "Start";
+	const FLazyName Name_Stop = "Stop";
+
 	constexpr uint32 NumEStateTreeRunStatus()
 	{
 #ifdef UE_STATETREE_ESTATETREERUNSTATUS_PRIVATE
@@ -1096,6 +1101,7 @@ void FStateTreeExecutionContext::SetUpdatePhaseInExecutionState(FStateTreeExecut
 EStateTreeRunStatus FStateTreeExecutionContext::Start(FStartParameters Parameters)
 {
 	CSV_SCOPED_TIMING_STAT_EXCLUSIVE(StateTree_Start);
+	UE_STATETREE_CRASH_REPORTER_SCOPE(&Owner, &RootStateTree, UE::StateTree::ExecutionContext::Private::Name_Start.Resolve());
 
 	if (!IsValid())
 	{
@@ -1264,6 +1270,7 @@ EStateTreeRunStatus FStateTreeExecutionContext::Start(FStartParameters Parameter
 EStateTreeRunStatus FStateTreeExecutionContext::Stop(EStateTreeRunStatus CompletionStatus)
 {
 	CSV_SCOPED_TIMING_STAT_EXCLUSIVE(StateTree_Stop);
+	UE_STATETREE_CRASH_REPORTER_SCOPE(&Owner, &RootStateTree, UE::StateTree::ExecutionContext::Private::Name_Stop.Resolve());
 
 	if (!IsValid())
 	{
@@ -1395,6 +1402,7 @@ EStateTreeRunStatus FStateTreeExecutionContext::TickPostlude()
 EStateTreeRunStatus FStateTreeExecutionContext::Tick(const float DeltaTime)
 {
 	CSV_SCOPED_TIMING_STAT_EXCLUSIVE(StateTree_Tick);
+	UE_STATETREE_CRASH_REPORTER_SCOPE(&Owner, &RootStateTree, UE::StateTree::ExecutionContext::Private::Name_Tick.Resolve());
 
 	TGuardValue<bool> ScheduledNextTickScope(bAllowedToScheduleNextTick, false);
 
@@ -1413,6 +1421,7 @@ EStateTreeRunStatus FStateTreeExecutionContext::Tick(const float DeltaTime)
 EStateTreeRunStatus FStateTreeExecutionContext::TickUpdateTasks(const float DeltaTime)
 {
 	CSV_SCOPED_TIMING_STAT_EXCLUSIVE(StateTree_Tick);
+	UE_STATETREE_CRASH_REPORTER_SCOPE(&Owner, &RootStateTree, UE::StateTree::ExecutionContext::Private::Name_Tick.Resolve());
 
 	TGuardValue<bool> ScheduledNextTickScope(bAllowedToScheduleNextTick, false);
 
@@ -1430,6 +1439,7 @@ EStateTreeRunStatus FStateTreeExecutionContext::TickUpdateTasks(const float Delt
 EStateTreeRunStatus FStateTreeExecutionContext::TickTriggerTransitions()
 {
 	CSV_SCOPED_TIMING_STAT_EXCLUSIVE(StateTree_Tick);
+	UE_STATETREE_CRASH_REPORTER_SCOPE(&Owner, &RootStateTree, UE::StateTree::ExecutionContext::Private::Name_Tick.Resolve());
 
 	TGuardValue<bool> ScheduledNextTickScope(bAllowedToScheduleNextTick, false);
 
diff --git a/Engine/Plugins/Runtime/StateTree/Source/StateTreeModule/Private/StateTreeModule.cpp b/Engine/Plugins/Runtime/StateTree/Source/StateTreeModule/Private/StateTreeModule.cpp
--- a/Engine/Plugins/Runtime/StateTree/Source/StateTreeModule/Private/StateTreeModule.cpp
+++ b/Engine/Plugins/Runtime/StateTree/Source/StateTreeModule/Private/StateTreeModule.cpp
@@ -4,6 +4,7 @@
 #include "StateTreeModuleImpl.h"
 
 #include "StateTreeTypes.h"
+#include "CrashReporter/StateTreeCrashReporterHandler.h"
 
 #if WITH_STATETREE_TRACE
 #include "Debugger/StateTreeTrace.h"
@@ -74,6 +75,10 @@ FStateTreeModule::FStateTreeModule()
 
 void FStateTreeModule::StartupModule()
 {
+#if UE_WITH_STATETREE_CRASHREPORTER
+	UE::StateTree::FCrashReporterHandler::Register();
+#endif
+
 #if WITH_STATETREE_TRACE_DEBUGGER
 	ITraceServicesModule& TraceServicesModule = FModuleManager::LoadModuleChecked<ITraceServicesModule>("TraceServices");
 	TraceAnalysisService = TraceServicesModule.GetAnalysisService();
@@ -145,6 +150,11 @@ void FStateTreeModule::ShutdownModule()
 
 	IModularFeatures::Get().UnregisterModularFeature(TraceServices::ModuleFeatureName, &StateTreeTraceModule);
 #endif // WITH_STATETREE_TRACE_DEBUGGER
+
+#if UE_WITH_STATETREE_CRASHREPORTER
+	UE::StateTree::FCrashReporterHandler::Unregister();
+#endif
+
 }
 
 bool FStateTreeModule::StartTraces(int32& OutTraceId)
