[StateTree] Add SharedInstance to the compiler output log.. Fix binding layout.

#rnx

diff --git a/Engine/Plugins/Runtime/PropertyBindingUtils/Source/PropertyBindingUtils/Private/PropertyBindingBindingCollection.cpp b/Engine/Plugins/Runtime/PropertyBindingUtils/Source/PropertyBindingUtils/Private/PropertyBindingBindingCollection.cpp
--- a/Engine/Plugins/Runtime/PropertyBindingUtils/Source/PropertyBindingUtils/Private/PropertyBindingBindingCollection.cpp
+++ b/Engine/Plugins/Runtime/PropertyBindingUtils/Source/PropertyBindingUtils/Private/PropertyBindingBindingCollection.cpp
@@ -1265,7 +1265,7 @@ FString FPropertyBindingBindingCollection::DebugAsString() const
 	DebugString.Appendf(TEXT("\nSourceStructs (%d)\n"), GetNumBindableStructDescriptors());
 	if (GetNumBindableStructDescriptors())
 	{
-		DebugString.Appendf(TEXT("  [ (Idx) | %-45s | % -40s | % -55s ]\n"), TEXT("Type"), TEXT("DataSource"), TEXT("Name"));
+		DebugString.Appendf(TEXT("  [ (Idx) | %-45s | % -80s ]\n"), TEXT("Type"), TEXT("DataSource"));
 		int32 Index = 0;
 		VisitSourceStructDescriptorInternal([&Index, &DebugString](const FPropertyBindingBindableStructDescriptor& Descriptor)
 		{
@@ -1321,10 +1321,10 @@ FString FPropertyBindingBindingCollection::DebugAsString() const
 	DebugString.Appendf(TEXT("\nPropertyCopies (%d)\n"), PropertyCopies.Num());
 	if (PropertyCopies.Num())
 	{
-		DebugString.Appendf(TEXT("  [ (Idx) | %-7s | %-4s | %-7s | %-10s | %-7s | %-4s | %-7s | %-10s | %-15s | %-20s | %-4s ]\n"),
+		DebugString.Appendf(TEXT("  [ (Idx) | %-7s | %-4s | %-7s | %-10s | %-7s | %-4s | %-7s | %-10s | %-20s | %-4s ]\n"),
 			TEXT("Src Idx"), TEXT("Off."), TEXT("Next"), TEXT("Type"),
 			TEXT("Tgt Idx"), TEXT("Off."), TEXT("Next"), TEXT("Type"),
-			TEXT("Source"), TEXT("Copy Type"), TEXT("Size"));
+			TEXT("Copy Type"), TEXT("Size"));
 		for (int32 Index = 0; Index < PropertyCopies.Num(); ++Index)
 		{
 			const FPropertyBindingCopyInfo& PropertyCopy = PropertyCopies[Index];
diff --git a/Engine/Plugins/Runtime/StateTree/Source/StateTreeModule/Private/StateTree.cpp b/Engine/Plugins/Runtime/StateTree/Source/StateTreeModule/Private/StateTree.cpp
--- a/Engine/Plugins/Runtime/StateTree/Source/StateTreeModule/Private/StateTree.cpp
+++ b/Engine/Plugins/Runtime/StateTree/Source/StateTreeModule/Private/StateTree.cpp
@@ -1364,6 +1364,22 @@ FString UStateTree::DebugInternalLayoutAsString() const
 		}
 	}
 
+	// Shared Instance InstanceData data (e.g. conditions/evaluators)
+	DebugString.Appendf(TEXT("\nShared Instance Data(%d)\n"), SharedInstanceData.Num());
+	for (int32 Index = 0; Index < SharedInstanceData.Num(); Index++)
+	{
+		if (SharedInstanceData.IsObject(Index))
+		{
+			const UObject* Data = SharedInstanceData.GetObject(Index);
+			PrintObjectNameSafe(Index, Data);
+		}
+		else
+		{
+			const FConstStructView Data = SharedInstanceData.GetStruct(Index);
+			PrintViewNameSafe(Index, Data);
+		}
+	}
+
 	// External data (e.g. fragments, subsystems)
 	DebugString.Appendf(TEXT("\nExternal Data(%d)\n"), ExternalDataDescs.Num());
 	if (ExternalDataDescs.Num())
@@ -1430,12 +1446,12 @@ FString UStateTree::DebugInternalLayoutAsString() const
 	DebugString.Appendf(TEXT("\nEvaluators(%d)\n"), EvaluatorsNum);
 	if (EvaluatorsNum)
 	{
-		DebugString.Appendf(TEXT("  [ (Idx) | %-30s | %8s | %10s ]\n"),
+		DebugString.Appendf(TEXT("  [ (Idx) | %-30s | %8s | %14s ]\n"),
 			TEXT("Name"), TEXT("Bindings"), TEXT("Struct Idx"));
 		for (int32 EvalIndex = EvaluatorsBegin; EvalIndex < (EvaluatorsBegin + EvaluatorsNum); EvalIndex++)
 		{
 			const FStateTreeEvaluatorBase& Eval = Nodes[EvalIndex].Get<const FStateTreeEvaluatorBase>();
-			DebugString.Appendf(TEXT("  | (%3d) | %-30s | %8d | %10s |\n"),
+			DebugString.Appendf(TEXT("  | (%3d) | %-30s | %8d | %14s |\n"),
 				EvalIndex,
 				*Eval.Name.ToString(),
 				Eval.BindingsBatch.Get(),
@@ -1444,7 +1460,7 @@ FString UStateTree::DebugInternalLayoutAsString() const
 	}
 
 	// Tasks
-	DebugString.Appendf(TEXT("\nTasks\n  [ (Idx) | %-30s | %-30s | %8s | %10s ]\n"),
+	DebugString.Appendf(TEXT("\nTasks\n  [ (Idx) | %-30s | %-30s | %8s | %14s ]\n"),
 		TEXT("State"), TEXT("Name"), TEXT("Bindings"), TEXT("Struct Idx"));
 	for (const FCompactStateTreeState& State : States)
 	{
@@ -1453,7 +1469,7 @@ FString UStateTree::DebugInternalLayoutAsString() const
 			for (int32 TaskIndex = State.TasksBegin; TaskIndex < (State.TasksBegin + State.TasksNum); TaskIndex++)
 			{
 				const FStateTreeTaskBase& Task = Nodes[TaskIndex].Get<const FStateTreeTaskBase>();
-				DebugString.Appendf(TEXT("  | (%3d) | %-30s | %-30s | %8d | %10s |\n"),
+				DebugString.Appendf(TEXT("  | (%3d) | %-30s | %-30s | %8d | %14s |\n"),
 					TaskIndex,
 					*State.Name.ToString(),
 					*Task.Name.ToString(),
@@ -1465,7 +1481,7 @@ FString UStateTree::DebugInternalLayoutAsString() const
 	for (int32 TaskIndex = GlobalTasksBegin; TaskIndex < (GlobalTasksBegin + GlobalTasksNum); TaskIndex++)
 	{
 		const FStateTreeTaskBase& Task = Nodes[TaskIndex].Get<const FStateTreeTaskBase>();
-		DebugString.Appendf(TEXT("  | (%3d) | %-30s | %-30s | %8d | %10s |\n"),
+		DebugString.Appendf(TEXT("  | (%3d) | %-30s | %-30s | %8d | %14s |\n"),
 			TaskIndex,
 			TEXT("Global"),
 			*Task.Name.ToString(),
@@ -1474,14 +1490,14 @@ FString UStateTree::DebugInternalLayoutAsString() const
 	}
 
 	// Conditions
-	DebugString.Appendf(TEXT("\nConditions\n  [ (Idx) | %-30s | %8s | %12s | %10s ]\n"),
+	DebugString.Appendf(TEXT("\nConditions\n  [ (Idx) | %-30s | %8s | %12s | %14s ]\n"),
 		TEXT("Name"), TEXT("Operand"), TEXT("Evaluation"), TEXT("Struct Idx"));
 	{
 		for (int32 Index = 0; Index < Nodes.Num(); ++Index)
 		{
 			if(const FStateTreeConditionBase* Cond = Nodes[Index].GetPtr<const FStateTreeConditionBase>())
 			{
-				DebugString.Appendf(TEXT("  | (%3d) | %-30s | %8s | %12s | %10s |\n"),
+				DebugString.Appendf(TEXT("  | (%3d) | %-30s | %8s | %12s | %14s |\n"),
 					Index,
 					*Cond->Name.ToString(),
 					*UEnum::GetDisplayValueAsText(Cond->Operand).ToString(),
diff --git a/Engine/Plugins/Runtime/StateTree/Source/StateTreeModule/Private/StateTreeTypes.cpp b/Engine/Plugins/Runtime/StateTree/Source/StateTreeModule/Private/StateTreeTypes.cpp
--- a/Engine/Plugins/Runtime/StateTree/Source/StateTreeModule/Private/StateTreeTypes.cpp
+++ b/Engine/Plugins/Runtime/StateTree/Source/StateTreeModule/Private/StateTreeTypes.cpp
@@ -147,15 +147,15 @@ FString FStateTreeDataHandle::Describe() const
 	case EStateTreeDataSourceType::GlobalInstanceData:
 		return FString::Printf(TEXT("Global[%d]"), Index);
 	case EStateTreeDataSourceType::GlobalInstanceDataObject:
-		return FString::Printf(TEXT("GlobalO[%d]"), Index);
+		return FString::Printf(TEXT("GlobalObj[%d]"), Index);
 	case EStateTreeDataSourceType::ActiveInstanceData:
 		return FString::Printf(TEXT("Active[%d]"), Index);
 	case EStateTreeDataSourceType::ActiveInstanceDataObject:
-		return FString::Printf(TEXT("ActiveO[%d]"), Index);
+		return FString::Printf(TEXT("ActiveObj[%d]"), Index);
 	case EStateTreeDataSourceType::SharedInstanceData:
 		return FString::Printf(TEXT("Shared[%d]"), Index);
 	case EStateTreeDataSourceType::SharedInstanceDataObject:
-		return FString::Printf(TEXT("SharedO[%d]"), Index);
+		return FString::Printf(TEXT("SharedObj[%d]"), Index);
 	case EStateTreeDataSourceType::ContextData:
 		return FString::Printf(TEXT("Context[%d]"), Index);
 	case EStateTreeDataSourceType::GlobalParameterData:
