[StateTre] Update comment made in CL 42830250
#rnx

diff --git a/Engine/Plugins/Runtime/StateTree/Source/StateTreeModule/Private/StateTree.cpp b/Engine/Plugins/Runtime/StateTree/Source/StateTreeModule/Private/StateTree.cpp
--- a/Engine/Plugins/Runtime/StateTree/Source/StateTreeModule/Private/StateTree.cpp
+++ b/Engine/Plugins/Runtime/StateTree/Source/StateTreeModule/Private/StateTree.cpp
@@ -1364,7 +1364,7 @@ FString UStateTree::DebugInternalLayoutAsString() const
 		}
 	}
 
-	// Shared Instance InstanceData data (e.g. conditions/evaluators)
+	// Shared Instance data (e.g. conditions/evaluators)
 	DebugString.Appendf(TEXT("\nShared Instance Data(%d)\n"), SharedInstanceData.Num());
 	for (int32 Index = 0; Index < SharedInstanceData.Num(); Index++)
 	{
