[StateTree] Fix typo for SetParamtersProperty, making the k2 node blueprint internal.
#rnx

diff --git a/Engine/Plugins/Runtime/StateTree/Source/StateTreeModule/Public/StateTreeFunctionLibrary.h b/Engine/Plugins/Runtime/StateTree/Source/StateTreeModule/Public/StateTreeFunctionLibrary.h
--- a/Engine/Plugins/Runtime/StateTree/Source/StateTreeModule/Public/StateTreeFunctionLibrary.h
+++ b/Engine/Plugins/Runtime/StateTree/Source/StateTreeModule/Public/StateTreeFunctionLibrary.h
@@ -24,10 +24,10 @@ public:
 	UFUNCTION(BlueprintCallable, DisplayName = "Make State Tree Reference", Category = "StateTree|Reference", meta = (Keywords = "construct build", NativeMakeFunc, BlueprintInternalUseOnly = "true"))
 	static UE_API FStateTreeReference MakeStateTreeReference(UStateTree* StateTree);
 
-	UFUNCTION(BlueprintCallable, CustomThunk, DisplayName = "Set Parameter Property", Category = "StateTree|Reference", meta = (BluerpintInternalUseOnly = "true", CustomStructureParam = "NewValue"))
+	UFUNCTION(BlueprintCallable, CustomThunk, DisplayName = "Set Parameter Property", Category = "StateTree|Reference", meta = (BlueprintInternalUseOnly = "true", CustomStructureParam = "NewValue"))
 	static UE_API void K2_SetParametersProperty(UPARAM(ref)FStateTreeReference& Reference, FGuid PropertyID, UPARAM(ref) const int32& NewValue);
 
-	UFUNCTION(BlueprintCallable, CustomThunk, DisplayName = "Get Parameter Property", Category = "StateTree|Reference", meta = (BluerpintInternalUseOnly = "true", CustomStructureParam = "ReturnValue"))
+	UFUNCTION(BlueprintCallable, CustomThunk, DisplayName = "Get Parameter Property", Category = "StateTree|Reference", meta = (BlueprintInternalUseOnly = "true", CustomStructureParam = "ReturnValue"))
 	static UE_API void K2_GetParametersProperty(UPARAM(ref)const FStateTreeReference& Reference, FGuid PropertyID, int32& ReturnValue);
 
 private:
