[StateTreeDebugger] refactored track to share functionality with RewindDebugger track

diff --git a/Engine/Plugins/Runtime/StateTree/Source/StateTreeEditorModule/Private/Debugger/SStateTreeDebuggerView.cpp b/Engine/Plugins/Runtime/StateTree/Source/StateTreeEditorModule/Private/Debugger/SStateTreeDebuggerView.cpp
--- a/Engine/Plugins/Runtime/StateTree/Source/StateTreeEditorModule/Private/Debugger/SStateTreeDebuggerView.cpp
+++ b/Engine/Plugins/Runtime/StateTree/Source/StateTreeEditorModule/Private/Debugger/SStateTreeDebuggerView.cpp
@@ -3,7 +3,6 @@
 #if WITH_STATETREE_TRACE_DEBUGGER
 
 #include "SStateTreeDebuggerView.h"
-#include "SStateTreeDebuggerViewRow.h"
 #include "Debugger/StateTreeDebugger.h"
 #include "DetailsViewArgs.h"
 #include "Editor.h"
@@ -14,6 +13,7 @@
 #include "PropertyEditorModule.h"
 #include "SStateTreeDebuggerInstanceTree.h"
 #include "SStateTreeDebuggerTimelines.h"
+#include "SStateTreeDebuggerViewRow.h"
 #include "StateTree.h"
 #include "StateTreeCompiler.h"
 #include "StateTreeDebuggerCommands.h"
@@ -120,9 +120,9 @@ void GenerateElementForProperties(const TCHAR* TypeAsText, const TCHAR* ValueAsT
 		ScriptStruct = LoadObject<UScriptStruct>(nullptr, TypeAsText);
 	}
 
-	FInstancedStruct NodeDataStruct;
 	if (ScriptStruct != nullptr)
 	{
+		FInstancedStruct NodeDataStruct;
 		NodeDataStruct.InitializeAs(ScriptStruct);
 
 		ScriptStruct->ImportText(ValueAsText, NodeDataStruct.GetMutableMemory(), /*OwnerObject*/nullptr, PPF_None, GLog, ScriptStruct->GetName());
@@ -137,7 +137,6 @@ void GenerateElementForProperties(const TCHAR* TypeAsText, const TCHAR* ValueAsT
 		Class = LoadObject<UClass>(nullptr, TypeAsText);
 	}
 
-	TWeakObjectPtr<> NodeDataObject;
 	if (Class != nullptr)
 	{
 		FTraceTextObjectFactory ObjectFactory;
@@ -146,6 +145,7 @@ void GenerateElementForProperties(const TCHAR* TypeAsText, const TCHAR* ValueAsT
 			ObjectFactory.ProcessBuffer(GetTransientPackage(), RF_Transactional, ValueAsText);
 			if (ObjectFactory.NodeInstanceObject)
 			{
+				TWeakObjectPtr<> NodeDataObject;
 				NodeDataObject = ObjectFactory.NodeInstanceObject;
 				NodeDataObject->AddToRoot();
 
@@ -161,7 +161,7 @@ void GenerateElementForProperties(const TCHAR* TypeAsText, const TCHAR* ValueAsT
 		{
 			const FProperty* const Property = PropertyIt.Key();
 			check(Property);
-			const EStateTreePropertyUsage Usage = UE::StateTree::GetUsageFromMetaData(Property);
+			const EStateTreePropertyUsage Usage = StateTree::GetUsageFromMetaData(Property);
 
 			// If the property is set to one of these usages, display it even if it is not edit on instance.
 			// It is a common mistake to forget to set the "eye" on these properties it and wonder why it does not show up.
@@ -447,7 +447,7 @@ void SStateTreeDebuggerView::Construct(const FArguments& InArgs, const UStateTre
 	}
 	RightToolbar.EndSection();
 
-	// Place holder toolbar for now but the intent to add more functionalities (e.g. Search)
+	// Placeholder toolbar for now but the intent to add more functionalities (e.g. Search)
 	FSlimHorizontalToolBarBuilder FrameDetailsToolbar(nullptr, FMultiBoxCustomization::None);
 	FrameDetailsToolbar.BeginSection("FrameDetails");
 	{
@@ -619,7 +619,7 @@ void SStateTreeDebuggerView::Construct(const FArguments& InArgs, const UStateTre
 								[
 									InstanceTimelinesTreeView.ToSharedRef()
 								]
-								+ SOverlay::Slot().HAlign(EHorizontalAlignment::HAlign_Right)
+								+ SOverlay::Slot().HAlign(HAlign_Right)
 								[
 									ScrollBar.ToSharedRef()
 								]
@@ -1014,14 +1014,14 @@ void SStateTreeDebuggerView::GenerateElementsForProperties(const FStateTreeTrace
 
 	if (!InstanceDataAsString.IsEmpty())
 	{
-		auto CreatePropertyElement = [ParentElement, &CreateTreeElement](const FStringView Line, const int32 NestedCount = 0)
+		auto CreatePropertyElement = [&CreateTreeElement](const FStringView Line, const int32 NestedCount = 0)
 			{
 				constexpr int32 Indent = 4;
 				FString ConvertedString = FString::Printf(TEXT("%*s"), NestedCount * Indent, TEXT(""));
 				ConvertedString.Append(Line);
 				ConvertedString.ReplaceInline(TEXT("="), TEXT(" = "));
 				ConvertedString.ReplaceInline(TEXT("\""), TEXT(""));
-				CreateTreeElement.template operator()<FStateTreeTracePropertyEvent>(ConvertedString);
+				CreateTreeElement.operator()<FStateTreeTracePropertyEvent>(ConvertedString);
 			};
 
 		// Try to parse Struct for which properties are exported between '(' and ')' 
@@ -1081,7 +1081,7 @@ void SStateTreeDebuggerView::GenerateElementsForProperties(const FStateTreeTrace
 
 	if (!DebugText.IsEmpty())
 	{
-		CreateTreeElement.template operator()<FStateTreeTraceLogEvent>(DebugText);
+		CreateTreeElement.operator()<FStateTreeTraceLogEvent>(DebugText);
 	}
 }
 
@@ -1167,7 +1167,7 @@ void SStateTreeDebuggerView::OnNewSession()
 	{
 		ResetTracks();
 	}
-	else if (Debugger->GetSelectedInstanceDescriptor() == nullptr)
+	else if (Debugger->GetSelectedInstanceId().IsInvalid())
 	{
 		// In PIE it is possible to stop/start the recording multiple times during the same game session,
 		// in this case we try to reselect the currently selected instance in case it can be reactivated.
@@ -1184,7 +1184,7 @@ void SStateTreeDebuggerView::OnNewSession()
 void SStateTreeDebuggerView::OnNewInstance(FStateTreeInstanceDebugId InstanceId)
 {
 	check(Debugger);
-	const UE::StateTreeDebugger::FInstanceDescriptor* FoundDescriptor = Debugger->GetInstanceDescriptor(InstanceId);
+	const TSharedPtr<const UE::StateTreeDebugger::FInstanceDescriptor> FoundDescriptor = Debugger->GetDescriptor(InstanceId);
 	if (!ensureMsgf(FoundDescriptor != nullptr, TEXT("This callback is from the Debugger so we expect to be able to find matching descriptor.")))
 	{
 		return;
@@ -1204,7 +1204,7 @@ void SStateTreeDebuggerView::OnNewInstance(FStateTreeInstanceDebugId InstanceId)
 	if (FStateTreeDebuggerOwnerTrack* OwnerTrack = static_cast<FStateTreeDebuggerOwnerTrack*>(ExistingOwnerTrack->Get()))
 	{
 		const FString TrackName = FString::Printf(TEXT("Execution #%d"), OwnerTrack->NumSubTracks()+1);
-		const TSharedPtr<FStateTreeDebuggerInstanceTrack> SubTrack = MakeShared<FStateTreeDebuggerInstanceTrack>(Debugger, InstanceId, FText::FromString(TrackName), ViewRange);
+		const TSharedPtr<FStateTreeDebuggerInstanceTrack> SubTrack = MakeShared<FStateTreeDebuggerInstanceTrack>(Debugger.ToSharedRef(), InstanceId, FText::FromString(TrackName), ViewRange);
 		OwnerTrack->AddSubTrack(SubTrack);
 
 		// Look at current selection; if nothing selected or stale track then select new track
diff --git a/Engine/Plugins/Runtime/StateTree/Source/StateTreeEditorModule/Private/Debugger/StateTreeDebuggerTrack.cpp b/Engine/Plugins/Runtime/StateTree/Source/StateTreeEditorModule/Private/Debugger/StateTreeDebuggerTrack.cpp
--- a/Engine/Plugins/Runtime/StateTree/Source/StateTreeEditorModule/Private/Debugger/StateTreeDebuggerTrack.cpp
+++ b/Engine/Plugins/Runtime/StateTree/Source/StateTreeEditorModule/Private/Debugger/StateTreeDebuggerTrack.cpp
@@ -4,142 +4,161 @@
 
 #include "Debugger/StateTreeDebuggerTrack.h"
 #include "Debugger/StateTreeDebugger.h"
-#include "SStateTreeDebuggerEventTimelineView.h"
+
+namespace UE::StateTreeDebugger
+{
 
 //----------------------------------------------------------------------//
-// FStateTreeDebuggerInstanceTrack
+// FInstanceTrackHelper
 //----------------------------------------------------------------------//
-FStateTreeDebuggerInstanceTrack::FStateTreeDebuggerInstanceTrack(
-	const TSharedPtr<FStateTreeDebugger>& InDebugger,
-	const FStateTreeInstanceDebugId InInstanceId,
-	const FText& InName,
-	const TRange<double>& InViewRange
-	)
-	: FStateTreeDebuggerBaseTrack(FSlateIcon("StateTreeEditorStyle", "StateTreeEditor.Debugger.InstanceTrack", "StateTreeEditor.Debugger.InstanceTrack"), InName)
-	, StateTreeDebugger(InDebugger)
-	, InstanceId(InInstanceId)
+FInstanceTrackHelper::FInstanceTrackHelper(const FStateTreeInstanceDebugId InInstanceId, const TRange<double>& InViewRange)
+	: InstanceId(InInstanceId)
 	, ViewRange(InViewRange)
 {
-	EventData = MakeShared<SStateTreeDebuggerEventTimelineView::FTimelineEventData>();
+	TimelineData = MakeShared<SStateTreeDebuggerEventTimelineView::FTimelineEventData>();
 }
 
-void FStateTreeDebuggerInstanceTrack::OnSelected()
+bool FInstanceTrackHelper::RebuildEventData(
+	const TNotNull<const UStateTree*> InStateTree
+	, const FInstanceEventCollection& InEventCollection
+	, const TRange<double>& InLifetime
+	, const double InRecordingDuration
+	, const bool bInActive) const
 {
-	if (FStateTreeDebugger* Debugger = StateTreeDebugger.Get())
+	if (InEventCollection.IsInvalid())
 	{
-		Debugger->SelectInstance(InstanceId);
+		return false;
 	}
-}
 
-bool FStateTreeDebuggerInstanceTrack::UpdateInternal()
-{
-	const int32 PrevNumPoints = EventData->Points.Num();
-	const int32 PrevNumWindows = EventData->Windows.Num();
-
-	EventData->Points.SetNum(0, EAllowShrinking::No);
-	EventData->Windows.SetNum(0);
-	
-	const FStateTreeDebugger* Debugger = StateTreeDebugger.Get();
-	check(Debugger);
-	const UStateTree* StateTree = Debugger->GetAsset();
-	const UE::StateTreeDebugger::FInstanceEventCollection& EventCollection = Debugger->GetEventCollection(InstanceId);
-	const double RecordingDuration = Debugger->GetRecordingDuration();
-	
-	if (StateTree != nullptr && EventCollection.IsValid())
-	{
-		auto MakeRandomColor = [bActive = !bIsStale](const uint32 InSeed)->FLinearColor
+	const int32 PrevNumPoints = TimelineData->Points.Num();
+	const int32 PrevNumWindows = TimelineData->Windows.Num();
+
+	TimelineData->Points.SetNum(0, EAllowShrinking::No);
+	TimelineData->Windows.SetNum(0);
+
+	auto MakeRandomColor = [bInActive](const uint32 InSeed)->FLinearColor
 		{
 			const FRandomStream Stream(InSeed);
 			const uint8 Hue = static_cast<uint8>(Stream.FRand() * 255.0f);
-			const uint8 SatVal = bActive ? 196 : 128;
+			const uint8 SatVal = bInActive ? 196 : 128;
 			return FLinearColor::MakeFromHSV8(Hue, SatVal, SatVal);
 		};
-		
-		const TConstArrayView<UE::StateTreeDebugger::FFrameSpan> Spans = EventCollection.FrameSpans;
-		const TConstArrayView<FStateTreeTraceEventVariantType> Events = EventCollection.Events;
-		const uint32 NumStateChanges = EventCollection.ActiveStatesChanges.Num();
-
-		TArray<UE::StateTreeDebugger::FInstanceEventCollection::FContiguousTraceInfo> TracesInfo = EventCollection.ContiguousTracesData;
-		// Append the on going trace info to "stopped" previous trace
-		if (NumStateChanges > 0)
+
+	const TConstArrayView<FFrameSpan> Spans = InEventCollection.FrameSpans;
+	const TConstArrayView<FStateTreeTraceEventVariantType> Events = InEventCollection.Events;
+	const uint32 NumStateChanges = InEventCollection.ActiveStatesChanges.Num();
+
+	TArray<FInstanceEventCollection::FContiguousTraceInfo> TracesInfo = InEventCollection.ContiguousTracesData;
+	// Append the ongoing trace info to "stopped" previous trace
+	if (NumStateChanges > 0)
+	{
+		TracesInfo.Emplace(InEventCollection.ActiveStatesChanges.Last().SpanIndex);
+	}
+
+	int32 StateChangeEndIndex = INDEX_NONE;
+	for (int32 TraceIndex = 0; TraceIndex < TracesInfo.Num(); TraceIndex++)
+	{
+		FInstanceEventCollection::FContiguousTraceInfo TraceInfo = TracesInfo[TraceIndex];
+		// Start at first event for the first trace or from the end index of the previous trace 
+		const int32 StateChangeBeginIndex = (StateChangeEndIndex == INDEX_NONE) ? 0 : StateChangeEndIndex;
+
+		// Find the starting index of the next trace to stop our iteration
+		StateChangeEndIndex = InEventCollection.ActiveStatesChanges.IndexOfByPredicate(
+			[LastSpanIndex = TraceInfo.LastSpanIndex](const FInstanceEventCollection::FActiveStatesChangePair& Pair)
+			{
+				return Pair.SpanIndex > LastSpanIndex;
+			});
+
+		// When not found means we are processing the last (or the only) trace
+		if (StateChangeEndIndex == INDEX_NONE)
 		{
-			TracesInfo.Emplace(EventCollection.ActiveStatesChanges.Last().SpanIndex);
+			StateChangeEndIndex = NumStateChanges;
 		}
 
-		int32 StateChangeEndIndex = INDEX_NONE;
-		for (int32 TraceIndex = 0; TraceIndex < TracesInfo.Num(); TraceIndex++)
+		for (int32 StateChangeIndex = StateChangeBeginIndex; StateChangeIndex < StateChangeEndIndex; ++StateChangeIndex)
 		{
-			UE::StateTreeDebugger::FInstanceEventCollection::FContiguousTraceInfo TraceInfo = TracesInfo[TraceIndex];
-			// Start at first event for the first trace or from the end index of the previous trace 
-			const int32 StateChangeBeginIndex = (StateChangeEndIndex == INDEX_NONE) ? 0 : StateChangeEndIndex;
-
-			// Find the starting index of the next trace to stop our iteration
-			StateChangeEndIndex = EventCollection.ActiveStatesChanges.IndexOfByPredicate(
-				[LastSpanIndex = TraceInfo.LastSpanIndex](const UE::StateTreeDebugger::FInstanceEventCollection::FActiveStatesChangePair& Pair)
-				{
-					return Pair.SpanIndex > LastSpanIndex;
-				});
-
-			// When not found means we are processing the last (or the only) trace
-			if (StateChangeEndIndex == INDEX_NONE)
+			const uint32 SpanIndex = InEventCollection.ActiveStatesChanges[StateChangeIndex].SpanIndex;
+			const uint32 EventIndex = InEventCollection.ActiveStatesChanges[StateChangeIndex].EventIndex;
+			const FStateTreeTraceActiveStatesEvent& Event = Events[EventIndex].Get<FStateTreeTraceActiveStatesEvent>();
+
+			FString StatePath = Event.GetValueString(*InStateTree);
+			FFrameSpan Span = InEventCollection.FrameSpans[SpanIndex];
+			SStateTreeDebuggerEventTimelineView::FTimelineEventData::EventWindow& Window = TimelineData->Windows.AddDefaulted_GetRef();
+			Window.Color = MakeRandomColor(CityHash32(reinterpret_cast<const char*>(*StatePath), StatePath.Len() * sizeof(FString::ElementType)));
+			Window.Description = FText::FromString(StatePath);
+			Window.TimeStart = Span.GetWorldTimeStart();
+
+			// For the last received event we use either the current recording duration if the instance is still active or the last recorded frame time
+			if (StateChangeIndex == (NumStateChanges - 1))
 			{
-				StateChangeEndIndex = NumStateChanges;
+				Window.TimeEnd = InLifetime.Contains(InRecordingDuration) ? InRecordingDuration : InEventCollection.FrameSpans.Last().GetWorldTimeEnd();
 			}
-
-			for (int32 StateChangeIndex = StateChangeBeginIndex; StateChangeIndex < StateChangeEndIndex; ++StateChangeIndex)
+			else
 			{
-				const uint32 SpanIndex = EventCollection.ActiveStatesChanges[StateChangeIndex].SpanIndex;
-				const uint32 EventIndex = EventCollection.ActiveStatesChanges[StateChangeIndex].EventIndex;
-				const FStateTreeTraceActiveStatesEvent& Event = Events[EventIndex].Get<FStateTreeTraceActiveStatesEvent>();
-				
-				FString StatePath = Event.GetValueString(*StateTree);
-				UE::StateTreeDebugger::FFrameSpan Span = EventCollection.FrameSpans[SpanIndex];
-				SStateTreeDebuggerEventTimelineView::FTimelineEventData::EventWindow& Window = EventData->Windows.AddDefaulted_GetRef();
-				Window.Color = MakeRandomColor(CityHash32(reinterpret_cast<const char*>(*StatePath), StatePath.Len() * sizeof(FString::ElementType)));
-				Window.Description = FText::FromString(StatePath);
-				Window.TimeStart = Span.GetWorldTimeStart();
-
-				// For the last received event we use either the current recording duration if the instance is still active
-				if (StateChangeIndex == (NumStateChanges - 1))
-				{
-					// or the last recorded frame time.
-					Window.TimeEnd = Debugger->IsActiveInstance(RecordingDuration, InstanceId)
-						? RecordingDuration
-						: EventCollection.FrameSpans.Last().GetWorldTimeEnd();
-				}
-				else
-				{
-					// When there is another state change after the current one in the list we use it to close the window.
-					// If the event is not the last of that specific trace then we use the start time of the next span.
-					// Otherwise, we use the end time of the last frame that was part of that trace.
-					const int32 NextStateChangeSpanIndex = EventCollection.ActiveStatesChanges[StateChangeIndex+1].SpanIndex;
-					Window.TimeEnd = (StateChangeIndex < (StateChangeEndIndex - 1))
-						? EventCollection.FrameSpans[NextStateChangeSpanIndex].GetWorldTimeStart()
-						: EventCollection.FrameSpans[NextStateChangeSpanIndex - 1].GetWorldTimeEnd(); //TraceInfo.LastRecordingTime;
-				}
+				// When there is another state change after the current one in the list we use it to close the window.
+				// If the event is not the last of that specific trace then we use the start time of the next span.
+				// Otherwise, we use the end time of the last frame that was part of that trace.
+				const int32 NextStateChangeSpanIndex = InEventCollection.ActiveStatesChanges[StateChangeIndex + 1].SpanIndex;
+				Window.TimeEnd = (StateChangeIndex < (StateChangeEndIndex - 1))
+					? InEventCollection.FrameSpans[NextStateChangeSpanIndex].GetWorldTimeStart()
+					: InEventCollection.FrameSpans[NextStateChangeSpanIndex - 1].GetWorldTimeEnd(); //TraceInfo.LastRecordingTime;
 			}
 		}
+	}
 
-		for (int32 SpanIndex = 0; SpanIndex < Spans.Num(); SpanIndex++)
-		{
-			const UE::StateTreeDebugger::FFrameSpan& Span = Spans[SpanIndex];
+	for (int32 SpanIndex = 0; SpanIndex < Spans.Num(); SpanIndex++)
+	{
+		const FFrameSpan& Span = Spans[SpanIndex];
 
-			const int32 StartIndex = Span.EventIdx;
-			const int32 MaxIndex = (SpanIndex + 1 < Spans.Num()) ? Spans[SpanIndex+1].EventIdx : Events.Num();
-			for (int EventIndex = StartIndex; EventIndex < MaxIndex; ++EventIndex)
+		const int32 StartIndex = Span.EventIdx;
+		const int32 MaxIndex = (SpanIndex + 1 < Spans.Num()) ? Spans[SpanIndex + 1].EventIdx : Events.Num();
+		for (int EventIndex = StartIndex; EventIndex < MaxIndex; ++EventIndex)
+		{
+			if (Events[EventIndex].IsType<FStateTreeTraceLogEvent>())
 			{
-				if (Events[EventIndex].IsType<FStateTreeTraceLogEvent>())
-				{
-					SStateTreeDebuggerEventTimelineView::FTimelineEventData::EventPoint Point;
-					Point.Time = Span.GetWorldTimeStart();
-					Point.Color = FColorList::Salmon;
-					EventData->Points.Add(Point);
-				}
+				SStateTreeDebuggerEventTimelineView::FTimelineEventData::EventPoint Point;
+				Point.Time = Span.GetWorldTimeStart();
+				Point.Color = FColorList::Salmon;
+				TimelineData->Points.Add(Point);
 			}
 		}
 	}
 
-	const bool bChanged = (PrevNumPoints != EventData->Points.Num() || PrevNumWindows != EventData->Windows.Num());
+	return (PrevNumPoints != TimelineData->Points.Num() || PrevNumWindows != TimelineData->Windows.Num());
+}
+
+} // namespace UE::StateTreeDebugger
+
+//----------------------------------------------------------------------//
+// FStateTreeDebuggerInstanceTrack
+//----------------------------------------------------------------------//
+FStateTreeDebuggerInstanceTrack::FStateTreeDebuggerInstanceTrack(
+	const TSharedRef<FStateTreeDebugger>& InDebugger,
+	const FStateTreeInstanceDebugId InInstanceId,
+	const FText& InName,
+	const TRange<double>& InViewRange
+	)
+	: FStateTreeDebuggerBaseTrack(FSlateIcon("StateTreeEditorStyle", "StateTreeEditor.Debugger.InstanceTrack", "StateTreeEditor.Debugger.InstanceTrack"), InName)
+	, StateTreeDebugger(InDebugger)
+	, InstanceTrackHelper(new UE::StateTreeDebugger::FInstanceTrackHelper(InInstanceId, InViewRange))
+	, Descriptor(InDebugger->GetDescriptor(InInstanceId))
+{
+}
+
+void FStateTreeDebuggerInstanceTrack::OnSelected()
+{
+	StateTreeDebugger.Get().SelectInstance(InstanceTrackHelper.Get().GetInstanceId());
+}
+
+bool FStateTreeDebuggerInstanceTrack::UpdateInternal()
+{
+	const FStateTreeDebugger& Debugger = StateTreeDebugger.Get();
+	const bool bActive = !IsStale();
+	const bool bChanged = InstanceTrackHelper.Get().RebuildEventData(Debugger.GetAsset()
+		, Debugger.GetEventCollection(InstanceTrackHelper.Get().GetInstanceId())
+		, Descriptor->Lifetime
+		, Debugger.GetRecordingDuration()
+		, bActive);
 
 	// Tracks can be reactivated when multiple recordings are made in a single PIE session.
 	if (bChanged && IsStale())
@@ -152,12 +171,9 @@ bool FStateTreeDebuggerInstanceTrack::UpdateInternal()
 
 TSharedPtr<SWidget> FStateTreeDebuggerInstanceTrack::GetTimelineViewInternal()
 {
-	return SNew(SStateTreeDebuggerEventTimelineView)
-		.ViewRange_Lambda([this](){ return ViewRange; })
-		.EventData_Lambda([this](){ return EventData; });
+	return InstanceTrackHelper.Get().CreateTimelineView();
 }
 
-
 //----------------------------------------------------------------------//
 // FStateTreeDebuggerOwnerTrack
 //----------------------------------------------------------------------//
@@ -177,11 +193,11 @@ bool FStateTreeDebuggerOwnerTrack::UpdateInternal()
 	return bChanged;
 }
 
-void FStateTreeDebuggerOwnerTrack::IterateSubTracksInternal(TFunction<void(TSharedPtr<FRewindDebuggerTrack> SubTrack)> IteratorFunction)
+void FStateTreeDebuggerOwnerTrack::IterateSubTracksInternal(TFunction<void(TSharedPtr<FRewindDebuggerTrack> SubTrack)> InIteratorFunction)
 {
 	for (TSharedPtr<FStateTreeDebuggerInstanceTrack>& Track : SubTracks)
 	{
-		IteratorFunction(Track);
+		InIteratorFunction(Track);
 	}
 }
 
@@ -189,7 +205,7 @@ void FStateTreeDebuggerOwnerTrack::MarkAsStale()
 {
 	for (TSharedPtr<FStateTreeDebuggerInstanceTrack>& Track : SubTracks)
 	{
-		if (FStateTreeDebuggerBaseTrack* InstanceTrack = Track.Get())
+		if (FStateTreeDebuggerInstanceTrack* InstanceTrack = Track.Get())
 		{
 			InstanceTrack->MarkAsStale();
 		}
@@ -206,7 +222,7 @@ bool FStateTreeDebuggerOwnerTrack::IsStale() const
 
 	for (const TSharedPtr<FStateTreeDebuggerInstanceTrack>& Track : SubTracks)
 	{
-		const FStateTreeDebuggerBaseTrack* InstanceTrack = Track.Get();
+		const FStateTreeDebuggerInstanceTrack* InstanceTrack = Track.Get();
 		if (InstanceTrack && InstanceTrack->IsStale() == false)
 		{
 			return false;
diff --git a/Engine/Plugins/Runtime/StateTree/Source/StateTreeEditorModule/Private/Debugger/StateTreeDebuggerTrack.h b/Engine/Plugins/Runtime/StateTree/Source/StateTreeEditorModule/Private/Debugger/StateTreeDebuggerTrack.h
--- a/Engine/Plugins/Runtime/StateTree/Source/StateTreeEditorModule/Private/Debugger/StateTreeDebuggerTrack.h
+++ b/Engine/Plugins/Runtime/StateTree/Source/StateTreeEditorModule/Private/Debugger/StateTreeDebuggerTrack.h
@@ -6,10 +6,68 @@
 
 #include "RewindDebuggerTrack.h"
 #include "SStateTreeDebuggerEventTimelineView.h"
-#include "StateTreeExecutionTypes.h"
+#include "Debugger/StateTreeDebuggerTypes.h"
 
 struct FStateTreeDebugger;
 
+namespace UE::StateTreeDebugger
+{
+
+/**
+ * Helper struct for debuggers relying on StateTree instance tracks
+ */
+struct FInstanceTrackHelper : public TSharedFromThis<FInstanceTrackHelper>
+{
+	explicit FInstanceTrackHelper(const FStateTreeInstanceDebugId InInstanceId, const TRange<double>& InViewRange);
+	
+	TSharedPtr<SStateTreeDebuggerEventTimelineView::FTimelineEventData> GetTimelineData() const
+	{
+		return TimelineData;
+	}
+
+	bool RebuildEventData(
+		TNotNull<const UStateTree*> InStateTree
+		, const FInstanceEventCollection& InEventCollection
+		, const TRange<double>& InLifetime
+		, const double InRecordingDuration
+		, const bool bInActive) const;
+
+	FStateTreeInstanceDebugId GetInstanceId() const
+	{
+		return InstanceId;
+	}
+
+	TSharedPtr<SWidget> CreateTimelineView()
+	{
+		return SNew(SStateTreeDebuggerEventTimelineView)
+			.ViewRange_Lambda([WeakHelper = AsWeak()]
+				{
+					if (const TSharedPtr<FInstanceTrackHelper> Helper = WeakHelper.Pin())
+					{
+						return Helper->ViewRange;
+					}
+					return TRange<double>{};
+				})
+			.EventData_Lambda([WeakHelper = AsWeak()]()
+				{
+					if (const TSharedPtr<FInstanceTrackHelper> Helper = WeakHelper.Pin())
+					{
+						return Helper->TimelineData;
+					}
+					return TSharedPtr<SStateTreeDebuggerEventTimelineView::FTimelineEventData>{};
+				});
+	}
+
+protected:
+
+	FStateTreeInstanceDebugId InstanceId;
+	TSharedPtr<SStateTreeDebuggerEventTimelineView::FTimelineEventData> TimelineData;
+	const TRange<double>& ViewRange;
+};
+
+} // UE::StateTreeDebugger
+
+
 /** Base struct for Debugger tracks to append some functionalities not available in RewindDebuggerTrack */
 struct FStateTreeDebuggerBaseTrack : RewindDebugger::FRewindDebuggerTrack
 {
@@ -19,45 +77,68 @@ struct FStateTreeDebuggerBaseTrack : RewindDebugger::FRewindDebuggerTrack
 	{
 	}
 
-	virtual bool IsStale() const { return false; }
-	virtual void MarkAsStale() {}
-	virtual void OnSelected() {}
+	virtual bool IsStale() const
+	{
+		return false;
+	}
+
+	virtual void MarkAsStale()
+	{
+	}
+
+	virtual void OnSelected()
+	{
+	}
 
 protected:
+	virtual FSlateIcon GetIconInternal() override
+	{
+		return Icon;
+	}
+
+	virtual FName GetNameInternal() const override
+	{
+		return FName(TrackName.ToString());
+	}
+
+	virtual FText GetDisplayNameInternal() const override
+	{
+		return TrackName;
+	}
+
 	FSlateIcon Icon;
 	FText TrackName;
 };
 
-/**
- * Struct use to store timeline events for a single StateTree instance
- * This is currently not tied to RewindDebugger but we use this base
- * type to facilitate a future integration.
- */
+/** Track used to represent timeline events for a single StateTree instance. */
 struct FStateTreeDebuggerInstanceTrack : FStateTreeDebuggerBaseTrack
 {
 	explicit FStateTreeDebuggerInstanceTrack(
-		const TSharedPtr<FStateTreeDebugger>& InDebugger,
+		const TSharedRef<FStateTreeDebugger>& InDebugger,
 		const FStateTreeInstanceDebugId InInstanceId,
 		const FText& InName,
 		const TRange<double>& InViewRange);
 
+	virtual void OnSelected() override;
+	virtual void MarkAsStale() override
+	{
+		bIsStale = true;
+	}
+
+	virtual bool IsStale() const override
+	{
+		return bIsStale;
+	}
+
 protected:
 	virtual bool UpdateInternal() override;
 	virtual TSharedPtr<SWidget> GetTimelineViewInternal() override;
-	virtual FSlateIcon GetIconInternal() override { return Icon; }
-	virtual FName GetNameInternal() const override { return FName(TrackName.ToString()); }
-	virtual FText GetDisplayNameInternal() const override { return TrackName; }
-
-	virtual void OnSelected() override;
-	virtual void MarkAsStale() override { bIsStale = true; }
-	virtual bool IsStale() const override { return bIsStale; }
 
 private:
-	TSharedPtr<SStateTreeDebuggerEventTimelineView::FTimelineEventData> EventData;
 	
-	TSharedPtr<FStateTreeDebugger> StateTreeDebugger;
-	FStateTreeInstanceDebugId InstanceId;
-	const TRange<double>& ViewRange;
+	TSharedRef<FStateTreeDebugger> StateTreeDebugger;
+	TSharedRef<UE::StateTreeDebugger::FInstanceTrackHelper> InstanceTrackHelper;
+	TSharedPtr<const UE::StateTreeDebugger::FInstanceDescriptor> Descriptor;
 	bool bIsStale = false;
 };
 
@@ -67,18 +148,23 @@ struct FStateTreeDebuggerOwnerTrack : FStateTreeDebuggerBaseTrack
 {
 	explicit FStateTreeDebuggerOwnerTrack(const FText& InInstanceName);
 
-	void AddSubTrack(const TSharedPtr<FStateTreeDebuggerInstanceTrack>& SubTrack) { SubTracks.Emplace(SubTrack); }
-	int32 NumSubTracks() const { return SubTracks.Num(); }
+	void AddSubTrack(const TSharedPtr<FStateTreeDebuggerInstanceTrack>& InSubTrack)
+	{
+		SubTracks.Emplace(InSubTrack);
+	}
 
-protected:
-	virtual bool UpdateInternal() override;	
-	virtual FSlateIcon GetIconInternal() override { return Icon; }
-	virtual FName GetNameInternal() const override { return FName(TrackName.ToString()); }
-	virtual FText GetDisplayNameInternal() const override { return TrackName; }
-	virtual void IterateSubTracksInternal(TFunction<void(TSharedPtr<FRewindDebuggerTrack> SubTrack)> IteratorFunction) override;
+	int32 NumSubTracks() const
+	{
+		return SubTracks.Num();
+	}
 
 	virtual void MarkAsStale() override;
 	virtual bool IsStale() const override;
+
+protected:
+	virtual bool UpdateInternal() override;
+	virtual void IterateSubTracksInternal(TFunction<void(TSharedPtr<FRewindDebuggerTrack> SubTrack)> InIteratorFunction) override;
+
 private:
 	TArray<TSharedPtr<FStateTreeDebuggerInstanceTrack>> SubTracks;
 };
diff --git a/Engine/Plugins/Runtime/StateTree/Source/StateTreeEditorModule/Private/Debugger/StateTreeRewindDebuggerTrack.cpp b/Engine/Plugins/Runtime/StateTree/Source/StateTreeEditorModule/Private/Debugger/StateTreeRewindDebuggerTrack.cpp
new file mode 100644
--- /dev/null
+++ b/Engine/Plugins/Runtime/StateTree/Source/StateTreeEditorModule/Private/Debugger/StateTreeRewindDebuggerTrack.cpp
@@ -0,0 +1,188 @@
+// Copyright Epic Games, Inc. All Rights Reserved.
+
+#include "StateTreeRewindDebuggerTrack.h"
+
+#include "Debugger/StateTreeDebugger.h"
+#include "Debugger/StateTreeTraceProvider.h"
+#include "Editor.h"
+#include "Modules/ModuleManager.h"
+#include "PropertyEditorModule.h"
+#include "Subsystems/AssetEditorSubsystem.h"
+
+#define LOCTEXT_NAMESPACE "StateTreeRewindDebuggerTrack"
+
+namespace UE::StateTreeDebugger
+{
+
+static const FLazyName StateTreeInstancesName("StateTreeInstances");
+
+
+//----------------------------------------------------------------------//
+// FRewindDebuggerTrackCreator
+//----------------------------------------------------------------------//
+void FRewindDebuggerTrackCreator::GetTrackTypesInternal(TArray<RewindDebugger::FRewindDebuggerTrackType>& Types) const
+{
+	Types.Add({ StateTreeInstancesName, LOCTEXT("StateTree", "StateTreeInstance") });
+}
+
+TSharedPtr<RewindDebugger::FRewindDebuggerTrack> FRewindDebuggerTrackCreator::CreateTrackInternal(const RewindDebugger::FObjectId& InObjectId) const
+{
+	return MakeShared<FRewindDebuggerTrack>(InObjectId);
+}
+
+bool FRewindDebuggerTrackCreator::HasDebugInfoInternal(const RewindDebugger::FObjectId& InObjectId) const
+{
+	TRACE_CPUPROFILER_EVENT_SCOPE(FStateTreeRewindDebuggerTrack::HasDebugInfoInternal);
+
+	if (!InObjectId.IsChildElement())
+	{
+		return false;
+	}
+
+	if (const IRewindDebugger* RewindDebugger = IRewindDebugger::Instance())
+	{
+		if (const TraceServices::IAnalysisSession* Session = RewindDebugger->GetAnalysisSession())
+		{
+			if (const IStateTreeTraceProvider* StateTreeTraceProvider = Session->ReadProvider<IStateTreeTraceProvider>(FStateTreeTraceProvider::ProviderName))
+			{
+				return StateTreeTraceProvider->GetInstanceDescriptor(FStateTreeInstanceDebugId(InObjectId.GetChildElementId())).IsValid();
+			}
+		}
+	}
+
+	return false;
+}
+
+bool FRewindDebuggerTrackCreator::IsChildElementSupportedInternal() const
+{
+	return true;
+}
+
+FName FRewindDebuggerTrackCreator::GetTargetTypeNameInternal() const
+{
+	static const FName AssociatedObjectName("MassStateTreeSubsystem");
+	return AssociatedObjectName;
+}
+
+//----------------------------------------------------------------------//
+// FRewindDebuggerTrack
+//----------------------------------------------------------------------//
+FRewindDebuggerTrack::FRewindDebuggerTrack(const RewindDebugger::FObjectId& InObjectId)
+	: EventCollection(FStateTreeInstanceDebugId(InObjectId.GetChildElementId()))
+	, InstanceTrackHelper(new FInstanceTrackHelper(EventCollection.InstanceId, IRewindDebugger::Instance()->GetCurrentViewRange()))
+	, ObjectId(InObjectId)
+{
+	FPropertyEditorModule& PropertyEditorModule = FModuleManager::GetModuleChecked<FPropertyEditorModule>("PropertyEditor");
+	FDetailsViewArgs DetailsViewArgs;
+	DetailsViewArgs.NameAreaSettings = FDetailsViewArgs::HideNameArea;
+	DetailsView = PropertyEditorModule.CreateDetailView(DetailsViewArgs);
+
+	if (const IRewindDebugger* RewindDebugger = IRewindDebugger::Instance())
+	{
+		if (const TraceServices::IAnalysisSession* Session = RewindDebugger->GetAnalysisSession())
+		{
+			if (const IStateTreeTraceProvider* StateTreeTraceProvider = Session->ReadProvider<IStateTreeTraceProvider>(FStateTreeTraceProvider::ProviderName))
+			{
+				Descriptor = StateTreeTraceProvider->GetInstanceDescriptor(EventCollection.InstanceId);
+			}
+		}
+	}
+
+	check(Descriptor.IsValid());
+}
+
+FText FRewindDebuggerTrack::GetDisplayNameInternal() const
+{
+	if (Descriptor.IsValid())
+	{
+		return FText::FromString(Descriptor->Name);
+	}
+
+	return LOCTEXT("StateTreeRewindDebuggerTrackName", "StateTreeInstance");
+}
+
+FSlateIcon FRewindDebuggerTrack::GetIconInternal()
+{
+	return Icon;
+}
+
+FName FRewindDebuggerTrack::GetNameInternal() const
+{
+	if (Descriptor.IsValid())
+	{
+		return *Descriptor->Name;
+	}
+
+	return "StateTreeDebuggerTrack";
+}
+
+uint64 FRewindDebuggerTrack::GetObjectIdInternal() const
+{
+	return ObjectId.GetMainId();
+}
+
+uint64 FRewindDebuggerTrack::GetChildElementIdInternal() const
+{
+	return ObjectId.GetChildElementId();
+}
+
+FInstanceEventCollection* FRewindDebuggerTrack::GetOrCreateEventCollection(const FStateTreeInstanceDebugId InstanceId)
+{
+	ensureMsgf(EventCollection.IsValid(), TEXT("Event collection is expected to be initialized"));
+	if (EventCollection.InstanceId == InstanceId)
+	{
+		return &EventCollection;
+	}
+
+	return nullptr;
+}
+
+bool FRewindDebuggerTrack::UpdateInternal()
+{
+	TRACE_CPUPROFILER_EVENT_SCOPE(FStateTreeRewindDebuggerTrack::UpdateInternal);
+
+	const IRewindDebugger* RewindDebugger = IRewindDebugger::Instance();
+	if (const TraceServices::IAnalysisSession* Session = RewindDebugger->GetAnalysisSession())
+	{
+		const double TraceTime = RewindDebugger->CurrentTraceTime();
+		const UStateTree* StateTree = Descriptor->StateTree.Get();
+
+		ReadTrace(*Session, TraceTime, /*ITraceReader*/this, FTraceFilter(InstanceTrackHelper.Get().GetInstanceId()), LastTraceReadTime);
+
+		return InstanceTrackHelper.Get().RebuildEventData(StateTree
+			, EventCollection
+			, Descriptor->Lifetime
+			, TraceTime
+			, /*Active*/true);
+	}
+
+	return false;
+}
+
+TSharedPtr<SWidget> FRewindDebuggerTrack::GetTimelineViewInternal()
+{
+	return InstanceTrackHelper.Get().CreateTimelineView();
+}
+
+TSharedPtr<SWidget> FRewindDebuggerTrack::GetDetailsViewInternal()
+{
+	return DetailsView;
+}
+
+bool FRewindDebuggerTrack::HandleDoubleClickInternal()
+{
+	if (Descriptor.IsValid())
+	{
+		if (const UStateTree* StateTreeAsset = Descriptor->StateTree.Get())
+		{
+			GEditor->GetEditorSubsystem<UAssetEditorSubsystem>()->OpenEditorForAsset(StateTreeAsset->GetPathName());
+			return true;
+		}
+	}
+
+	return false;
+}
+
+
+} // namespace UE::StateTreeDebugger
+#undef LOCTEXT_NAMESPACE
\ No newline at end of file
diff --git a/Engine/Plugins/Runtime/StateTree/Source/StateTreeEditorModule/Private/Debugger/StateTreeRewindDebuggerTrack.h b/Engine/Plugins/Runtime/StateTree/Source/StateTreeEditorModule/Private/Debugger/StateTreeRewindDebuggerTrack.h
new file mode 100644
--- /dev/null
+++ b/Engine/Plugins/Runtime/StateTree/Source/StateTreeEditorModule/Private/Debugger/StateTreeRewindDebuggerTrack.h
@@ -0,0 +1,65 @@
+// Copyright Epic Games, Inc. All Rights Reserved.
+#pragma once
+
+#include "Debugger/StateTreeDebugger.h"
+#include "Debugger/StateTreeDebuggerTypes.h"
+#include "IRewindDebuggerTrackCreator.h"
+#include "RewindDebuggerTrack.h"
+#include "StateTreeDebuggerTrack.h"
+
+class IDetailsView;
+
+namespace UE::StateTreeDebugger
+{
+/**
+ * RewindDebugger track creator for StateTree instances
+ */
+struct FRewindDebuggerTrackCreator final : RewindDebugger::IRewindDebuggerTrackCreator
+{
+protected:
+	virtual FName GetTargetTypeNameInternal() const override;
+	virtual FName GetNameInternal() const override
+	{
+		return "StateTreeInstanceTrack";
+	}
+	virtual void GetTrackTypesInternal(TArray<RewindDebugger::FRewindDebuggerTrackType>& Types) const override;
+	virtual TSharedPtr<RewindDebugger::FRewindDebuggerTrack> CreateTrackInternal(const RewindDebugger::FObjectId& InObjectId) const override;
+	virtual bool HasDebugInfoInternal(const RewindDebugger::FObjectId& InObjectId) const override;
+	virtual bool IsChildElementSupportedInternal() const override;
+};
+
+/**
+ * RewindDebugger track representing a single StateTree instance
+ */
+struct FRewindDebuggerTrack : RewindDebugger::FRewindDebuggerTrack, ITraceReader
+{
+	explicit FRewindDebuggerTrack(const RewindDebugger::FObjectId& InObjectId);
+
+private:
+	//~ Begin UE::StateTreeDebugger::ITraceReader interface
+	virtual UE::StateTreeDebugger::FInstanceEventCollection* GetOrCreateEventCollection(FStateTreeInstanceDebugId InstanceId) override;
+	//~ End UE::StateTreeDebugger::ITraceReader interface
+
+	//~ Begin RewindDebugger::FRewindDebuggerTrack interface
+	virtual bool UpdateInternal() override;
+	virtual TSharedPtr<SWidget> GetTimelineViewInternal() override;
+	virtual TSharedPtr<SWidget> GetDetailsViewInternal() override;
+	virtual FText GetDisplayNameInternal() const override;
+	virtual FSlateIcon GetIconInternal() override;
+	virtual FName GetNameInternal() const override;
+	virtual uint64 GetObjectIdInternal() const override;
+	virtual uint64 GetChildElementIdInternal() const override;
+	virtual bool HandleDoubleClickInternal() override;
+	//~ End RewindDebugger::FRewindDebuggerTrack interface
+
+	FInstanceEventCollection EventCollection;
+	TSharedRef<FInstanceTrackHelper> InstanceTrackHelper;
+	TSharedPtr<const FInstanceDescriptor> Descriptor;
+	TSharedPtr<IDetailsView> DetailsView;
+	FSlateIcon Icon;
+	RewindDebugger::FObjectId ObjectId;
+	/** Last time in the recording that we used to fetch events, and that we will use for the next read. */
+	double LastTraceReadTime = 0;
+};
+
+} // UE::StateTreeDebugger
\ No newline at end of file
diff --git a/Engine/Plugins/Runtime/StateTree/Source/StateTreeEditorModule/Private/StateTreeEditorModule.cpp b/Engine/Plugins/Runtime/StateTree/Source/StateTreeEditorModule/Private/StateTreeEditorModule.cpp
--- a/Engine/Plugins/Runtime/StateTree/Source/StateTreeEditorModule/Private/StateTreeEditorModule.cpp
+++ b/Engine/Plugins/Runtime/StateTree/Source/StateTreeEditorModule/Private/StateTreeEditorModule.cpp
@@ -5,36 +5,36 @@
 #include "Blueprint/StateTreeConsiderationBlueprintBase.h"
 #include "Blueprint/StateTreeEvaluatorBlueprintBase.h"
 #include "Blueprint/StateTreeTaskBlueprintBase.h"
-#include "StateTreePropertyFunctionBase.h"
 #include "Customizations/StateTreeAnyEnumDetails.h"
-#include "Customizations/StateTreeEnumValueScorePairsDetails.h"
+#include "Customizations/StateTreeBindingExtension.h"
+#include "Customizations/StateTreeBlueprintPropertyRefDetails.h"
 #include "Customizations/StateTreeEditorColorDetails.h"
 #include "Customizations/StateTreeEditorDataDetails.h"
 #include "Customizations/StateTreeEditorNodeDetails.h"
+#include "Customizations/StateTreeEnumValueScorePairsDetails.h"
+#include "Customizations/StateTreeEventDescDetails.h"
 #include "Customizations/StateTreeReferenceDetails.h"
 #include "Customizations/StateTreeReferenceOverridesDetails.h"
 #include "Customizations/StateTreeStateDetails.h"
 #include "Customizations/StateTreeStateLinkDetails.h"
 #include "Customizations/StateTreeStateParametersDetails.h"
 #include "Customizations/StateTreeTransitionDetails.h"
-#include "Customizations/StateTreeEventDescDetails.h"
-#include "Customizations/StateTreeBindingExtension.h"
-#include "Customizations/StateTreeBlueprintPropertyRefDetails.h"
+#include "Debugger/StateTreeDebuggerCommands.h"
+#include "Debugger/StateTreeRewindDebuggerExtensions.h"
+#include "Debugger/StateTreeRewindDebuggerTrack.h"
 #include "IRewindDebuggerExtension.h"
+#include "Modules/ModuleManager.h"
 #include "PropertyEditorModule.h"
 #include "StateTree.h"
-#include "StateTreeCompiler.h"
 #include "StateTreeCompilerLog.h"
-#include "StateTreeEditorData.h"
-#include "Debugger/StateTreeDebuggerCommands.h"
-#include "Debugger/StateTreeRewindDebuggerExtensions.h"
 #include "StateTreeDelegates.h"
 #include "StateTreeEditingSubsystem.h"
 #include "StateTreeEditor.h"
 #include "StateTreeEditorCommands.h"
+#include "StateTreeEditorData.h"
 #include "StateTreeEditorStyle.h"
 #include "StateTreeNodeClassCache.h"
-#include "Modules/ModuleManager.h"
+#include "StateTreePropertyFunctionBase.h"
 
 #define LOCTEXT_NAMESPACE "StateTreeEditor"
 
@@ -70,10 +70,15 @@ void FStateTreeEditorModule::StartupModule()
 
 #if WITH_STATETREE_TRACE_DEBUGGER
 	FStateTreeDebuggerCommands::Register();
+
 	RewindDebuggerExtension = MakePimpl<UE::StateTreeDebugger::FRewindDebuggerExtension>();
 	IModularFeatures::Get().RegisterModularFeature(IRewindDebuggerExtension::ModularFeatureName, RewindDebuggerExtension.Get());
+
 	RewindDebuggerRuntimeExtension = MakePimpl<UE::StateTreeDebugger::FRewindDebuggerRuntimeExtension>();
 	IModularFeatures::Get().RegisterModularFeature(IRewindDebuggerRuntimeExtension::ModularFeatureName, RewindDebuggerRuntimeExtension.Get());
+
+	RewindDebuggerTrackCreator = MakePimpl<UE::StateTreeDebugger::FRewindDebuggerTrackCreator>();
+	IModularFeatures::Get().RegisterModularFeature(RewindDebugger::IRewindDebuggerTrackCreator::ModularFeatureName, RewindDebuggerTrackCreator.Get());
 #endif // WITH_STATETREE_TRACE_DEBUGGER
 
 	MenuExtensibilityManager = MakeShareable(new FExtensibilityManager);
@@ -109,6 +114,7 @@ void FStateTreeEditorModule::ShutdownModule()
 #if WITH_STATETREE_TRACE_DEBUGGER
 	IModularFeatures::Get().UnregisterModularFeature(IRewindDebuggerRuntimeExtension::ModularFeatureName, RewindDebuggerRuntimeExtension.Get());
 	IModularFeatures::Get().UnregisterModularFeature(IRewindDebuggerExtension::ModularFeatureName, RewindDebuggerExtension.Get());
+	IModularFeatures::Get().UnregisterModularFeature(RewindDebugger::IRewindDebuggerTrackCreator::ModularFeatureName, RewindDebuggerTrackCreator.Get());
 	FStateTreeDebuggerCommands::Unregister();
 #endif // WITH_STATETREE_TRACE_DEBUGGER
 	
diff --git a/Engine/Plugins/Runtime/StateTree/Source/StateTreeEditorModule/Public/StateTreeEditorModule.h b/Engine/Plugins/Runtime/StateTree/Source/StateTreeEditorModule/Public/StateTreeEditorModule.h
--- a/Engine/Plugins/Runtime/StateTree/Source/StateTreeEditorModule/Public/StateTreeEditorModule.h
+++ b/Engine/Plugins/Runtime/StateTree/Source/StateTreeEditorModule/Public/StateTreeEditorModule.h
@@ -19,6 +19,7 @@ namespace UE::StateTreeDebugger
 {
 	class FRewindDebuggerExtension;
 	class FRewindDebuggerRuntimeExtension;
+	struct FRewindDebuggerTrackCreator;
 }
 
 STATETREEEDITORMODULE_API DECLARE_LOG_CATEGORY_EXTERN(LogStateTreeEditor, Log, All);
@@ -76,6 +77,7 @@ protected:
 #if WITH_STATETREE_TRACE_DEBUGGER
 	TPimplPtr<UE::StateTreeDebugger::FRewindDebuggerExtension> RewindDebuggerExtension;
 	TPimplPtr<UE::StateTreeDebugger::FRewindDebuggerRuntimeExtension> RewindDebuggerRuntimeExtension;
+	TPimplPtr<UE::StateTreeDebugger::FRewindDebuggerTrackCreator> RewindDebuggerTrackCreator;
 #endif  // WITH_STATETREE_TRACE_DEBUGGER
 
 	FDelegateHandle OnUserDefinedStructReinstancedHandle;
diff --git a/Engine/Plugins/Runtime/StateTree/Source/StateTreeModule/Private/Debugger/StateTreeDebugger.cpp b/Engine/Plugins/Runtime/StateTree/Source/StateTreeModule/Private/Debugger/StateTreeDebugger.cpp
--- a/Engine/Plugins/Runtime/StateTree/Source/StateTreeModule/Private/Debugger/StateTreeDebugger.cpp
+++ b/Engine/Plugins/Runtime/StateTree/Source/StateTreeModule/Private/Debugger/StateTreeDebugger.cpp
@@ -6,19 +6,19 @@
 #include "Debugger/IStateTreeTraceProvider.h"
 #include "Debugger/StateTreeTraceProvider.h"
 #include "Debugger/StateTreeTraceTypes.h"
+#include "GenericPlatform/GenericPlatformMisc.h"
 #include "Misc/Paths.h"
 #include "Modules/ModuleManager.h"
 #include "StateTreeDelegates.h"
 #include "StateTreeModule.h"
+#include "Trace/Analysis.h"
+#include "Trace/Analyzer.h"
 #include "Trace/StoreClient.h"
 #include "TraceServices/AnalysisService.h"
 #include "TraceServices/ITraceServicesModule.h"
 #include "TraceServices/Model/AnalysisSession.h"
-#include "TraceServices/Model/Frames.h"
-#include "Trace/Analyzer.h"
-#include "Trace/Analysis.h"
 #include "TraceServices/Model/Diagnostics.h"
-#include "GenericPlatform/GenericPlatformMisc.h"
+#include "TraceServices/Model/Frames.h"
 
 #define LOCTEXT_NAMESPACE "StateTreeDebugger"
 
@@ -49,8 +49,8 @@ namespace UE::StateTreeDebugger
 					EventData.GetString("Branch", SessionInfo.Branch);
 					EventData.GetString("BuildVersion", SessionInfo.BuildVersion);
 					SessionInfo.Changelist = EventData.GetValue<uint32>("Changelist", 0);
-					SessionInfo.ConfigurationType = (EBuildConfiguration) EventData.GetValue<uint8>("ConfigurationType");
-					SessionInfo.TargetType = (EBuildTargetType) EventData.GetValue<uint8>("TargetType");
+					SessionInfo.ConfigurationType = static_cast<EBuildConfiguration>(EventData.GetValue<uint8>("ConfigurationType"));
+					SessionInfo.TargetType = static_cast<EBuildTargetType>(EventData.GetValue<uint8>("TargetType"));
 
 					return false;
 				}
@@ -149,7 +149,7 @@ void FStateTreeDebugger::SyncToCurrentSessionDuration()
 	}
 }
 
-const UE::StateTreeDebugger::FInstanceDescriptor* FStateTreeDebugger::GetInstanceDescriptor(const FStateTreeInstanceDebugId InstanceId) const
+TSharedPtr<const UE::StateTreeDebugger::FInstanceDescriptor> FStateTreeDebugger::GetDescriptor(const FStateTreeInstanceDebugId InstanceId) const
 {
 	if (const TraceServices::IAnalysisSession* Session = GetAnalysisSession())
 	{
@@ -157,7 +157,7 @@ const UE::StateTreeDebugger::FInstanceDescriptor* FStateTreeDebugger::GetInstanc
 
 		if (const IStateTreeTraceProvider* Provider = Session->ReadProvider<IStateTreeTraceProvider>(FStateTreeTraceProvider::ProviderName))
 		{
-			return Provider->GetInstanceDescriptor(InstanceId).Get();
+			return Provider->GetInstanceDescriptor(InstanceId);
 		}
 	}
 
@@ -166,13 +166,13 @@ const UE::StateTreeDebugger::FInstanceDescriptor* FStateTreeDebugger::GetInstanc
 
 FText FStateTreeDebugger::GetInstanceName(const FStateTreeInstanceDebugId InstanceId) const
 {
-	const UE::StateTreeDebugger::FInstanceDescriptor* FoundDescriptor = GetInstanceDescriptor(InstanceId);
+	const TSharedPtr<const UE::StateTreeDebugger::FInstanceDescriptor> FoundDescriptor = GetDescriptor(InstanceId);
 	return (FoundDescriptor != nullptr) ? FText::FromString(FoundDescriptor->Name) : LOCTEXT("InstanceNotFound","Instance not found");
 }
 
 FText FStateTreeDebugger::GetInstanceDescription(const FStateTreeInstanceDebugId InstanceId) const
 {
-	const UE::StateTreeDebugger::FInstanceDescriptor* FoundDescriptor = GetInstanceDescriptor(InstanceId);
+	const TSharedPtr<const UE::StateTreeDebugger::FInstanceDescriptor> FoundDescriptor = GetDescriptor(InstanceId);
 	return (FoundDescriptor != nullptr) ? DescribeInstance(*FoundDescriptor) : LOCTEXT("InstanceNotFound","Instance not found");
 }
 
@@ -182,7 +182,7 @@ void FStateTreeDebugger::SelectInstance(const FStateTreeInstanceDebugId Instance
 	{
 		SelectedInstanceId = InstanceId;
 
-		// Notify so listener can cleanup anything related to previous instance
+		// Notify so listener can clean up anything related to previous instance
 		OnSelectedInstanceCleared.ExecuteIfBound();
 
 		// Update event collection index for newly debugged instance
@@ -235,9 +235,9 @@ bool FStateTreeDebugger::RequestAnalysisOfEditorSession()
 	constexpr int32 InvalidTraceId = 0;
 	int32 ActiveTraceId = InvalidTraceId;
 
-	// StartTraces returns true if a new connection was created. In this case we will receive OnTracingStateChanged
-	// and we'll try to start an analysis on that new connection as soon as possible.
-	// Otherwise it might have been able to use an active connection in which case it was returned in the output parameter.
+	// StartTraces returns true if a new connection was created.
+	// In this case we will receive OnTracingStateChanged and try to start an analysis on that new connection as soon as possible.
+	// Otherwise, it might have been able to use an active connection in which case it was returned in the output parameter.
 	if (StateTreeModule.StartTraces(ActiveTraceId))
 	{
 		return true;
@@ -596,7 +596,7 @@ void FStateTreeDebugger::StepForwardToNextStateChange()
 
 bool FStateTreeDebugger::HasStateBreakpoint(const FStateTreeStateHandle StateHandle, const EStateTreeBreakpointType BreakpointType) const
 {
-	return Breakpoints.ContainsByPredicate([StateHandle, BreakpointType](const FStateTreeDebuggerBreakpoint Breakpoint)
+	return Breakpoints.ContainsByPredicate([StateHandle, BreakpointType](const FStateTreeDebuggerBreakpoint& Breakpoint)
 		{
 			if (Breakpoint.BreakpointType == BreakpointType)
 			{
@@ -609,7 +609,7 @@ bool FStateTreeDebugger::HasStateBreakpoint(const FStateTreeStateHandle StateHan
 
 bool FStateTreeDebugger::HasTaskBreakpoint(const FStateTreeIndex16 Index, const EStateTreeBreakpointType BreakpointType) const
 {
-	return Breakpoints.ContainsByPredicate([Index, BreakpointType](const FStateTreeDebuggerBreakpoint Breakpoint)
+	return Breakpoints.ContainsByPredicate([Index, BreakpointType](const FStateTreeDebuggerBreakpoint& Breakpoint)
 	{
 		if (Breakpoint.BreakpointType == BreakpointType)
 		{
@@ -622,7 +622,7 @@ bool FStateTreeDebugger::HasTaskBreakpoint(const FStateTreeIndex16 Index, const
 
 bool FStateTreeDebugger::HasTransitionBreakpoint(const FStateTreeIndex16 Index, const EStateTreeBreakpointType BreakpointType) const
 {
-	return Breakpoints.ContainsByPredicate([Index, BreakpointType](const FStateTreeDebuggerBreakpoint Breakpoint)
+	return Breakpoints.ContainsByPredicate([Index, BreakpointType](const FStateTreeDebuggerBreakpoint& Breakpoint)
 	{
 		if (Breakpoint.BreakpointType == BreakpointType)
 		{
@@ -726,11 +726,16 @@ void FStateTreeDebugger::ReadTrace(const uint64 FrameIndex)
 	{
 		TraceServices::FAnalysisSessionReadScope SessionReadScope(*Session);
 
-		const TraceServices::IFrameProvider& FrameProvider = ReadFrameProvider(*Session);
+		const TraceServices::IFrameProvider& FrameProvider = TraceServices::ReadFrameProvider(*Session);
 
 		if (const TraceServices::FFrame* TargetFrame = FrameProvider.GetFrame(TraceFrameType_Game, FrameIndex))
 		{
-			ReadTrace(*Session, FrameProvider, *TargetFrame);
+			UE::StateTreeDebugger::ReadTrace(*Session
+				, FrameProvider
+				, *TargetFrame
+				, /*ITraceReader*/this
+				, UE::StateTreeDebugger::FTraceFilter(StateTreeAsset.Get())
+				, LastTraceReadTime);
 		}
 	}
 
@@ -742,30 +747,15 @@ void FStateTreeDebugger::ReadTrace(const double ScrubTime)
 {
 	if (const TraceServices::IAnalysisSession* Session = GetAnalysisSession())
 	{
-		TraceServices::FAnalysisSessionReadScope SessionReadScope(*Session);
-
-		const TraceServices::IFrameProvider& FrameProvider = ReadFrameProvider(*Session);
-
-		TraceServices::FFrame TargetFrame;
-		if (FrameProvider.GetFrameFromTime(TraceFrameType_Game, ScrubTime, TargetFrame))
-		{
-			// Process only completed frames 
-			if (TargetFrame.EndTime == std::numeric_limits<double>::infinity())
-			{
-				if (const TraceServices::FFrame* PreviousCompleteFrame = FrameProvider.GetFrame(TraceFrameType_Game, TargetFrame.Index - 1))
-				{
-					ReadTrace(*Session, FrameProvider, *PreviousCompleteFrame);
-				}
-			}
-			else
-			{
-				ReadTrace(*Session, FrameProvider, TargetFrame);
-			}
-		}
+		UE::StateTreeDebugger::ReadTrace(
+			*Session
+			, ScrubTime
+			, this
+			, UE::StateTreeDebugger::FTraceFilter(StateTreeAsset.Get())
+			, LastTraceReadTime);
+
+		SendNotifications();
 	}
-
-	// Notify outside session read scope
-	SendNotifications();
 }
 
 void FStateTreeDebugger::SendNotifications()
@@ -800,24 +790,6 @@ void FStateTreeDebugger::SendNotifications()
 	}
 }
 
-void FStateTreeDebugger::ReadTrace(
-	const TraceServices::IAnalysisSession& Session,
-	const TraceServices::IFrameProvider& FrameProvider,
-	const TraceServices::FFrame& Frame
-	)
-{
-	TraceServices::FFrame LastReadFrame;
-	const bool bValidLastReadFrame = FrameProvider.GetFrameFromTime(TraceFrameType_Game, LastTraceReadTime, LastReadFrame);
-	if (LastTraceReadTime == 0 || (bValidLastReadFrame && Frame.Index > LastReadFrame.Index))
-	{
-		if (const IStateTreeTraceProvider* Provider = Session.ReadProvider<IStateTreeTraceProvider>(FStateTreeTraceProvider::ProviderName))
-		{
-			AddEvents(LastTraceReadTime, Frame.EndTime, FrameProvider, *Provider);
-			LastTraceReadTime = Frame.EndTime;
-		}
-	}
-}
-
 bool FStateTreeDebugger::EvaluateBreakpoints(const FStateTreeInstanceDebugId InstanceId, const FStateTreeTraceEventVariantType& Event)
 {
 	if (StateTreeAsset == nullptr // asset is required to properly match state handles
@@ -834,15 +806,15 @@ bool FStateTreeDebugger::EvaluateBreakpoints(const FStateTreeInstanceDebugId Ins
 		if (Breakpoint.IsMatchingEvent(Event))
 		{
 			HitBreakpoint.Index = BreakpointIndex;
-            HitBreakpoint.InstanceId = InstanceId;
-            HitBreakpoint.Time = RecordingDuration;
+			HitBreakpoint.InstanceId = InstanceId;
+			HitBreakpoint.Time = RecordingDuration;
 		}
 	}
 
 	return HitBreakpoint.IsSet();
 }
 
-bool FStateTreeDebugger::ProcessEvent(const FStateTreeInstanceDebugId InstanceId, const TraceServices::FFrame& Frame, const FStateTreeTraceEventVariantType& Event)
+UE::StateTreeDebugger::FInstanceEventCollection* FStateTreeDebugger::GetOrCreateEventCollection(FStateTreeInstanceDebugId InstanceId)
 {
 	UE::StateTreeDebugger::FInstanceEventCollection* ExistingCollection = EventCollections.FindByPredicate(
 		[InstanceId](const UE::StateTreeDebugger::FInstanceEventCollection& Entry)
@@ -862,157 +834,258 @@ bool FStateTreeDebugger::ProcessEvent(const FStateTreeInstanceDebugId InstanceId
 		// Otherwise (i.e. EventCollection already exists) it is updated when switching instance (i.e. SelectInstance)
 		if (SelectedInstanceId == InstanceId && ScrubState.GetEventCollectionIndex() == INDEX_NONE)
 		{
-			SetScrubStateCollectionIndex(EventCollections.Num()-1);
+			SetScrubStateCollectionIndex(EventCollections.Num() - 1);
 		}
 	}
+	return ExistingCollection;
+}
 
-	check(ExistingCollection);
-	TArray<FStateTreeTraceEventVariantType>& Events = ExistingCollection->Events;
-
-	TraceServices::FFrame FrameToAddInSpans = Frame;
-	bool bShouldAddFrameToSpans = false;
-
-	double RecordingWorldTime = 0;
-	Visit([&RecordingWorldTime](auto& TypedEvent)
+void FStateTreeDebugger::OnTraceEventProcessed(const FStateTreeInstanceDebugId InstanceId, const FStateTreeTraceEventVariantType& Event)
+{
+	Visit([&RecordingDuration = RecordingDuration](auto& TypedEvent)
 		{
-			RecordingWorldTime = TypedEvent.RecordingWorldTime;
+			RecordingDuration = TypedEvent.RecordingWorldTime;
 		}, Event);
 
+	EvaluateBreakpoints(InstanceId, Event);
+}
+
+const UE::StateTreeDebugger::FInstanceEventCollection& FStateTreeDebugger::GetEventCollection(FStateTreeInstanceDebugId InstanceId) const
+{
+	using namespace UE::StateTreeDebugger;
+	const FInstanceEventCollection* ExistingCollection = EventCollections.FindByPredicate([InstanceId](const FInstanceEventCollection& Entry)
+	{
+		return Entry.InstanceId == InstanceId;
+	});
+
+	return ExistingCollection != nullptr ? *ExistingCollection : FInstanceEventCollection::Invalid;
+}
+
+void FStateTreeDebugger::ResetEventCollections()
+{
+	EventCollections.Reset();
+	SetScrubStateCollectionIndex(INDEX_NONE);
+	RecordingDuration = 0;
+}
+
+//----------------------------------------------------------------//
+// UE::StateTreeDebugger
+//----------------------------------------------------------------//
+namespace UE::StateTreeDebugger
+{
+
+bool ProcessEvent(
+	FInstanceEventCollection& InEventCollection,
+	const TraceServices::FFrame& InFrame,
+	const FStateTreeTraceEventVariantType& InEvent,
+	TNotNull<ITraceReader*> TraceReader)
+{
+	TArray<FStateTreeTraceEventVariantType>& Events = InEventCollection.Events;
+
+	TraceServices::FFrame FrameToAddInSpans = InFrame;
+	bool bShouldAddFrameToSpans = false;
+
 	// Add new frame span if none added yet
-	if (ExistingCollection->FrameSpans.IsEmpty())
+	if (InEventCollection.FrameSpans.IsEmpty())
 	{
 		bShouldAddFrameToSpans = true;
 	}
 	else
 	{
-		const UE::StateTreeDebugger::FFrameSpan& LastSpan = ExistingCollection->FrameSpans.Last();
+		const FFrameSpan& LastSpan = InEventCollection.FrameSpans.Last();
 		const TraceServices::FFrame& LastFrame = LastSpan.Frame;
-		const uint64 FrameIndexOffset = ExistingCollection->ContiguousTracesData.IsEmpty()
+		const uint64 FrameIndexOffset = InEventCollection.ContiguousTracesData.IsEmpty()
 			? 0
-			: (ExistingCollection->FrameSpans[ExistingCollection->ContiguousTracesData.Last().LastSpanIndex].Frame.Index + 1);
+			: (InEventCollection.FrameSpans[InEventCollection.ContiguousTracesData.Last().LastSpanIndex].Frame.Index + 1);
 
 		// Add new frame span for new larger frame index
-		if (Frame.Index + FrameIndexOffset > LastFrame.Index)
+		if (InFrame.Index + FrameIndexOffset > LastFrame.Index)
 		{
 			bShouldAddFrameToSpans = true;
 
 			// Apply current offset to the frame index
 			FrameToAddInSpans.Index += FrameIndexOffset;
 		}
-		else if (Frame.Index < LastFrame.Index && Frame.StartTime > LastFrame.StartTime)
+		else if (InFrame.Index < LastFrame.Index && InFrame.StartTime > LastFrame.StartTime)
 		{
 			// Frame index will restart at 0 if a new session is started,
 			// in that case we offset the frame we store to append to existing data
 			bShouldAddFrameToSpans = true;
 
-			const UE::StateTreeDebugger::FInstanceEventCollection::FContiguousTraceInfo& TraceInfo =
-				ExistingCollection->ContiguousTracesData.Emplace_GetRef(
-					UE::StateTreeDebugger::FInstanceEventCollection::FContiguousTraceInfo(ExistingCollection->FrameSpans.Num()-1));
-			FrameToAddInSpans.Index += ExistingCollection->FrameSpans[TraceInfo.LastSpanIndex].Frame.Index + 1;
+			const FInstanceEventCollection::FContiguousTraceInfo& TraceInfo =
+				InEventCollection.ContiguousTracesData.Emplace_GetRef(
+					FInstanceEventCollection::FContiguousTraceInfo(InEventCollection.FrameSpans.Num() - 1));
+			FrameToAddInSpans.Index += InEventCollection.FrameSpans[TraceInfo.LastSpanIndex].Frame.Index + 1;
 		}
 	}
 
 	if (bShouldAddFrameToSpans)
 	{
-		// Update global recording duration
-		RecordingDuration = RecordingWorldTime;
-		ExistingCollection->FrameSpans.Add(UE::StateTreeDebugger::FFrameSpan(FrameToAddInSpans, RecordingWorldTime, Events.Num()));
+		double RecordingWorldTime = 0;
+		Visit([&RecordingWorldTime](auto& TypedEvent)
+			{
+				RecordingWorldTime = TypedEvent.RecordingWorldTime;
+			}, InEvent);
+
+		InEventCollection.FrameSpans.Add(FFrameSpan(FrameToAddInSpans, RecordingWorldTime, Events.Num()));
 	}
 
 	// Add activate states change info
-	if (Event.IsType<FStateTreeTraceActiveStatesEvent>())
+	if (InEvent.IsType<FStateTreeTraceActiveStatesEvent>())
 	{
-		checkf(ExistingCollection->FrameSpans.Num() > 0, TEXT("Expecting to always be in a frame span at this point."));
-		const int32 FrameSpanIndex = ExistingCollection->FrameSpans.Num()-1;
+		checkf(InEventCollection.FrameSpans.Num() > 0, TEXT("Expecting to always be in a frame span at this point."));
+		const int32 FrameSpanIndex = InEventCollection.FrameSpans.Num() - 1;
 
 		// Add new entry for the first event or if the last event is for a different frame
-		if (ExistingCollection->ActiveStatesChanges.IsEmpty()
-			|| ExistingCollection->ActiveStatesChanges.Last().SpanIndex != FrameSpanIndex)
+		if (InEventCollection.ActiveStatesChanges.IsEmpty()
+			|| InEventCollection.ActiveStatesChanges.Last().SpanIndex != FrameSpanIndex)
 		{
-			ExistingCollection->ActiveStatesChanges.Push({FrameSpanIndex, Events.Num()});
+			InEventCollection.ActiveStatesChanges.Push({ FrameSpanIndex, Events.Num() });
 		}
 		else
 		{
 			// Multiple events for change of active states in the same frame, keep the last one until we implement scrubbing within a frame
-			ExistingCollection->ActiveStatesChanges.Last().EventIndex = Events.Num();
+			InEventCollection.ActiveStatesChanges.Last().EventIndex = Events.Num();
 		}
 	}
 
 	// Store event in the collection
-	Events.Emplace(Event);
+	Events.Emplace(InEvent);
 
-	// Process at the end so RecordingDuration is up to date and we can associate it to a hit breakpoint if necessary.
-	EvaluateBreakpoints(InstanceId, Event);
+	// Notify the reader
+	TraceReader->OnTraceEventProcessed(InEventCollection.InstanceId, InEvent);
 
 	return /*bKeepProcessing*/true;
 }
 
-const UE::StateTreeDebugger::FInstanceEventCollection& FStateTreeDebugger::GetEventCollection(FStateTreeInstanceDebugId InstanceId) const\
+void AddEvents(
+	const double InStartTime,
+	const double InEndTime,
+	const TraceServices::IFrameProvider& InFrameProvider,
+	const FStateTreeInstanceDebugId InInstanceId,
+	const IStateTreeTraceProvider::FEventsTimeline& TimelineData,
+	TNotNull<ITraceReader*> TraceReader)
 {
-	using namespace UE::StateTreeDebugger;
-	const FInstanceEventCollection* ExistingCollection = EventCollections.FindByPredicate([InstanceId](const FInstanceEventCollection& Entry)
+	FInstanceEventCollection* EventCollection = TraceReader->GetOrCreateEventCollection(InInstanceId);
+	if (EventCollection == nullptr)
 	{
-		return Entry.InstanceId == InstanceId;
-	});
+		return;
+	}
 
-	return ExistingCollection != nullptr ? *ExistingCollection : FInstanceEventCollection::Invalid;
-}
+	// Keep track of the frames containing events. Starting with an invalid frame.
+	TraceServices::FFrame Frame;
+	Frame.Index = INDEX_NONE;
 
-void FStateTreeDebugger::ResetEventCollections()
-{
-	EventCollections.Reset();
-	SetScrubStateCollectionIndex(INDEX_NONE);
-	RecordingDuration = 0;
+	TimelineData.EnumerateEvents(InStartTime, InEndTime,
+		[&InFrameProvider, &Frame, EventCollection, TraceReader](const double EventStartTime, const double EventEndTime, uint32 InDepth, const FStateTreeTraceEventVariantType& Event)
+		{
+			bool bValidFrame = true;
+
+			// Fetch frame when not set yet or if events no longer part of the current one
+			if (Frame.Index == INDEX_NONE ||
+				(EventEndTime < Frame.StartTime || Frame.EndTime < EventStartTime))
+			{
+				bValidFrame = InFrameProvider.GetFrameFromTime(TraceFrameType_Game, EventStartTime, Frame);
+
+				if (bValidFrame == false)
+				{
+					// Edge case for events from a missing first complete frame.
+					// (i.e. FrameProvider didn't get BeginFrame event but StateTreeEvent were sent in that frame)
+					// Doing this will merge our two first frames of state tree events using the same recording world time
+					// but this should happen only for late start recording.
+					const TraceServices::FFrame* FirstFrame = InFrameProvider.GetFrame(TraceFrameType_Game, 0);
+					if (FirstFrame != nullptr && EventEndTime < FirstFrame->StartTime)
+					{
+						Frame = *FirstFrame;
+						bValidFrame = true;
+					}
+				}
+			}
+
+			if (bValidFrame)
+			{
+				const bool bKeepProcessing = ProcessEvent(*EventCollection, Frame, Event, TraceReader);
+				return bKeepProcessing ? TraceServices::EEventEnumerate::Continue : TraceServices::EEventEnumerate::Stop;
+			}
+
+			// Skip events outside of game frames
+			return TraceServices::EEventEnumerate::Continue;
+		});
 }
 
-void FStateTreeDebugger::AddEvents(const double StartTime, const double EndTime, const TraceServices::IFrameProvider& FrameProvider, const IStateTreeTraceProvider& StateTreeTraceProvider)
+void ReadTrace(
+	const TraceServices::IAnalysisSession& Session,
+	const double ScrubTime,
+	const TNotNull<ITraceReader*> TraceReader,
+	const FTraceFilter& Filter,
+	double& OutLastTraceReadTime)
 {
-	check(StateTreeAsset.IsValid());
-	StateTreeTraceProvider.ReadTimelines(*StateTreeAsset,
-		[this, StartTime, EndTime, &FrameProvider](const FStateTreeInstanceDebugId InstanceId, const IStateTreeTraceProvider::FEventsTimeline& TimelineData)
+	TraceServices::FAnalysisSessionReadScope SessionReadScope(Session);
+
+	const TraceServices::IFrameProvider& FrameProvider = TraceServices::ReadFrameProvider(Session);
+
+	TraceServices::FFrame TargetFrame;
+	if (FrameProvider.GetFrameFromTime(TraceFrameType_Game, ScrubTime, TargetFrame))
+	{
+		// Process only completed frames
+		bool bValidFrame = true;
+		if (TargetFrame.EndTime == std::numeric_limits<double>::infinity())
 		{
-			// Keep track of the frames containing events. Starting with an invalid frame.
-			TraceServices::FFrame Frame;
-			Frame.Index = INDEX_NONE;
+			if (const TraceServices::FFrame* PreviousCompleteFrame = FrameProvider.GetFrame(TraceFrameType_Game, TargetFrame.Index - 1))
+			{
+				TargetFrame = *PreviousCompleteFrame;
+			}
+			else
+			{
+				bValidFrame = false;
+			}
+		}
 
-			TimelineData.EnumerateEvents(StartTime,	EndTime,
-				[this, InstanceId, &FrameProvider, &Frame](const double EventStartTime, const double EventEndTime, uint32 InDepth, const FStateTreeTraceEventVariantType& Event)
-				{
-					bool bValidFrame = true;
+		if (bValidFrame)
+		{
+			ReadTrace(Session, FrameProvider, TargetFrame, TraceReader, Filter, OutLastTraceReadTime);
+		}
+	}
+}
 
-					// Fetch frame when not set yet or if events no longer part of the current one
-					if (Frame.Index == INDEX_NONE ||
-						(EventEndTime < Frame.StartTime || Frame.EndTime < EventStartTime))
+void ReadTrace(
+	const TraceServices::IAnalysisSession& Session,
+	const TraceServices::IFrameProvider& FrameProvider,
+	const TraceServices::FFrame& Frame,
+	TNotNull<ITraceReader*> TraceReader,
+	const FTraceFilter& Filter,
+	double& OutLastTraceReadTime)
+{
+	TraceServices::FFrame LastReadFrame;
+	const bool bValidLastReadFrame = FrameProvider.GetFrameFromTime(TraceFrameType_Game, OutLastTraceReadTime, LastReadFrame);
+	if (OutLastTraceReadTime == 0 || (bValidLastReadFrame && Frame.Index > LastReadFrame.Index))
+	{
+		if (const IStateTreeTraceProvider* Provider = Session.ReadProvider<IStateTreeTraceProvider>(FStateTreeTraceProvider::ProviderName))
+		{
+			if (Filter.Instance.IsValid())
+			{
+				Provider->ReadTimelines(Filter.Instance,
+					[StartTime = OutLastTraceReadTime, EndTime = Frame.EndTime, &FrameProvider, TraceReader](const FStateTreeInstanceDebugId InstanceId, const IStateTreeTraceProvider::FEventsTimeline& TimelineData)
 					{
-						bValidFrame = FrameProvider.GetFrameFromTime(TraceFrameType_Game, EventStartTime, Frame);
-
-						if (bValidFrame == false)
-						{
-							// Edge case for events from a missing first complete frame.
-							// (i.e. FrameProvider didn't get BeginFrame event but StateTreeEvent were sent in that frame)
-							// Doing this will merge our two first frames of state tree events using the same recording world time
-							// but this should happen only for late start recording.
-							const TraceServices::FFrame* FirstFrame = FrameProvider.GetFrame(TraceFrameType_Game, 0);
-							if (FirstFrame != nullptr && EventEndTime < FirstFrame->StartTime)
-							{
-								Frame = *FirstFrame;
-								bValidFrame = true;
-							}
-						}
-					}
-
-					if (bValidFrame)
+						AddEvents(StartTime, EndTime, FrameProvider, InstanceId, TimelineData, TraceReader);
+					});
+			}
+			else if (Filter.Asset)
+			{
+				Provider->ReadTimelines(*Filter.Asset,
+					[StartTime = OutLastTraceReadTime, EndTime = Frame.EndTime, &FrameProvider, TraceReader](const FStateTreeInstanceDebugId InstanceId, const IStateTreeTraceProvider::FEventsTimeline& TimelineData)
 					{
-						const bool bKeepProcessing = ProcessEvent(InstanceId, Frame, Event);
-						return bKeepProcessing ? TraceServices::EEventEnumerate::Continue : TraceServices::EEventEnumerate::Stop;
-					}
+						AddEvents(StartTime, EndTime, FrameProvider, InstanceId, TimelineData, TraceReader);
+					});
+			}
 
-					// Skip events outside of game frames
-					return TraceServices::EEventEnumerate::Continue;
-				});
-		});
+			OutLastTraceReadTime = Frame.EndTime;
+		}
+	}
 }
 
+} // UE::StateTreeDebugger
+
 #undef LOCTEXT_NAMESPACE
 
 #endif // WITH_STATETREE_TRACE_DEBUGGER
diff --git a/Engine/Plugins/Runtime/StateTree/Source/StateTreeModule/Private/Debugger/StateTreeDebuggerTypes.cpp b/Engine/Plugins/Runtime/StateTree/Source/StateTreeModule/Private/Debugger/StateTreeDebuggerTypes.cpp
--- a/Engine/Plugins/Runtime/StateTree/Source/StateTreeModule/Private/Debugger/StateTreeDebuggerTypes.cpp
+++ b/Engine/Plugins/Runtime/StateTree/Source/StateTreeModule/Private/Debugger/StateTreeDebuggerTypes.cpp
@@ -10,7 +10,7 @@ namespace UE::StateTreeDebugger
 //----------------------------------------------------------------//
 // FInstanceDescriptor
 //----------------------------------------------------------------//
-FInstanceDescriptor::FInstanceDescriptor(const UStateTree* InStateTree, const FStateTreeInstanceDebugId InId, const FString& InName, const TRange<double> InLifetime)
+FInstanceDescriptor::FInstanceDescriptor(const UStateTree* InStateTree, const FStateTreeInstanceDebugId InId, const FString& InName, const TRange<double>& InLifetime)
 	: Lifetime(InLifetime)
 	, StateTree(InStateTree)
 	, Name(InName)
@@ -99,7 +99,7 @@ void FScrubState::SetFrameSpanIndex(const int32 NewFrameSpanIndex)
 	checkf(EventCollections.IsValidIndex(EventCollectionIndex), TEXT("Internal method expecting validity checks before getting called."));
 	const FInstanceEventCollection& EventCollection = EventCollections[EventCollectionIndex];
 
-	checkf(EventCollections[EventCollectionIndex].FrameSpans.IsValidIndex(FrameSpanIndex), TEXT("Internal method expecting validity checks before getting called."));
+	checkf(EventCollection.FrameSpans.IsValidIndex(FrameSpanIndex), TEXT("Internal method expecting validity checks before getting called."));
 	ScrubTime = EventCollection.FrameSpans[FrameSpanIndex].GetWorldTimeStart();
 	TraceFrameIndex = EventCollection.FrameSpans[FrameSpanIndex].Frame.Index;
 	ScrubTimeBoundState = EScrubTimeBoundState::InBounds;
diff --git a/Engine/Plugins/Runtime/StateTree/Source/StateTreeModule/Public/Debugger/StateTreeDebugger.h b/Engine/Plugins/Runtime/StateTree/Source/StateTreeModule/Public/Debugger/StateTreeDebugger.h
--- a/Engine/Plugins/Runtime/StateTree/Source/StateTreeModule/Public/Debugger/StateTreeDebugger.h
+++ b/Engine/Plugins/Runtime/StateTree/Source/StateTreeModule/Public/Debugger/StateTreeDebugger.h
@@ -27,6 +27,49 @@ class IStateTreeModule;
 class UStateTreeState;
 class UStateTree;
 
+namespace UE::StateTreeDebugger
+{
+	struct ITraceReader
+	{
+		virtual ~ITraceReader() = default;
+		virtual FInstanceEventCollection* GetOrCreateEventCollection(FStateTreeInstanceDebugId InstanceId) = 0;
+		virtual void OnTraceEventProcessed(FStateTreeInstanceDebugId InstanceId, const FStateTreeTraceEventVariantType& Event)
+		{
+		}
+	};
+
+	struct FTraceFilter
+	{
+		FTraceFilter() = delete;
+		explicit FTraceFilter(const TNotNull<const UStateTree*>& Asset)
+			: Asset(Asset)
+		{
+		}
+		explicit FTraceFilter(const FStateTreeInstanceDebugId& Instance)
+			: Instance(Instance)
+		{
+		}
+
+		TObjectPtr<const UStateTree> Asset = nullptr;
+		FStateTreeInstanceDebugId Instance;
+	};
+
+	UE_API void ReadTrace(
+		const TraceServices::IAnalysisSession& Session,
+		const double ScrubTime,
+		TNotNull<ITraceReader*> TraceReader,
+		const FTraceFilter& Filter,
+		double& OutLastTraceReadTime);
+
+	UE_API void ReadTrace(
+		const TraceServices::IAnalysisSession& Session,
+		const TraceServices::IFrameProvider& FrameProvider,
+		const TraceServices::FFrame& Frame,
+		TNotNull<ITraceReader*> TraceReader,
+		const FTraceFilter& Filter,
+		double& OutLastTraceReadTime);
+}
+
 DECLARE_DELEGATE_OneParam(FOnStateTreeDebuggerScrubStateChanged, const UE::StateTreeDebugger::FScrubState& ScrubState);
 DECLARE_DELEGATE_TwoParams(FOnStateTreeDebuggerBreakpointHit, FStateTreeInstanceDebugId InstanceId, const FStateTreeDebuggerBreakpoint Breakpoint);
 DECLARE_DELEGATE_OneParam(FOnStateTreeDebuggerActiveStatesChanges, const FStateTreeTraceActiveStates& ActiveStates);
@@ -34,7 +77,7 @@ DECLARE_DELEGATE_OneParam(FOnStateTreeDebuggerNewInstance, FStateTreeInstanceDeb
 DECLARE_DELEGATE(FOnStateTreeDebuggerNewSession);
 DECLARE_DELEGATE(FOnStateTreeDebuggerDebuggedInstanceSet);
 
-struct FStateTreeDebugger : FTickableGameObject
+struct FStateTreeDebugger : FTickableGameObject, UE::StateTreeDebugger::ITraceReader
 {
 	struct FTraceDescriptor
 	{
@@ -98,9 +141,9 @@ struct FStateTreeDebugger : FTickableGameObject
 		return StateTreeAsset.Get();
 	}
 
-	void SetAsset(const UStateTree* InStateTreeAsset)
+	void SetAsset(const UStateTree* Asset)
 	{
-		StateTreeAsset = InStateTreeAsset;
+		StateTreeAsset = Asset;
 	}
 	
 	/** Forces a single refresh to latest state. Useful when simulation is paused. */
@@ -135,10 +178,21 @@ struct FStateTreeDebugger : FTickableGameObject
 		return SelectedInstanceId;
 	}
 
-	UE_API const UE::StateTreeDebugger::FInstanceDescriptor* GetInstanceDescriptor(const FStateTreeInstanceDebugId InstanceId) const ;
+	UE_API TSharedPtr<const UE::StateTreeDebugger::FInstanceDescriptor> GetDescriptor(const FStateTreeInstanceDebugId InstanceId) const;
+	TSharedPtr<const UE::StateTreeDebugger::FInstanceDescriptor> GetSelectedDescriptor() const
+	{
+		return GetDescriptor(SelectedInstanceId);
+	}
+
+	UE_DEPRECATED(5.7, "Use GetDescriptor instead")
+	UE_API const UE::StateTreeDebugger::FInstanceDescriptor* GetInstanceDescriptor(const FStateTreeInstanceDebugId InstanceId) const
+	{
+		return GetDescriptor(InstanceId).Get();
+	}
+	UE_DEPRECATED(5.7, "Use GetSelectedDescriptor instead")
 	const UE::StateTreeDebugger::FInstanceDescriptor* GetSelectedInstanceDescriptor() const
 	{
-		return GetInstanceDescriptor(SelectedInstanceId);
+		return GetDescriptor(SelectedInstanceId).Get();
 	}
 
 	UE_API bool HasStateBreakpoint(FStateTreeStateHandle StateHandle, EStateTreeBreakpointType BreakpointType) const;
@@ -313,11 +367,11 @@ private:
 	
 	void ReadTrace(double ScrubTime);
 	void ReadTrace(uint64 FrameIndex);
-	void ReadTrace(
-		const TraceServices::IAnalysisSession& Session,
-		const TraceServices::IFrameProvider& FrameProvider,
-		const TraceServices::FFrame& Frame
-		);
+
+	//~ Begin UE::StateTreeDebugger::ITraceReader interface
+	virtual UE::StateTreeDebugger::FInstanceEventCollection* GetOrCreateEventCollection(FStateTreeInstanceDebugId InstanceId) override;
+	virtual void OnTraceEventProcessed(FStateTreeInstanceDebugId InstanceId, const FStateTreeTraceEventVariantType& Event) override;
+	//~ End UE::StateTreeDebugger::ITraceReader interface
 
 	/**
 	 * Tests event for given instance id against breakpoints.
@@ -357,8 +411,6 @@ private:
 
 	UE::Trace::FStoreClient* GetStoreClient() const;
 
-	bool ProcessEvent(const FStateTreeInstanceDebugId InstanceId, const TraceServices::FFrame& Frame, const FStateTreeTraceEventVariantType& Event);
-	void AddEvents(double StartTime, double EndTime, const TraceServices::IFrameProvider& FrameProvider, const IStateTreeTraceProvider& StateTreeTraceProvider);
 	void UpdateMetadata(FTraceDescriptor& TraceDescriptor) const;
 
 	/** Module used to access the store client and analysis sessions .*/
diff --git a/Engine/Plugins/Runtime/StateTree/Source/StateTreeModule/Public/Debugger/StateTreeDebuggerTypes.h b/Engine/Plugins/Runtime/StateTree/Source/StateTreeModule/Public/Debugger/StateTreeDebuggerTypes.h
--- a/Engine/Plugins/Runtime/StateTree/Source/StateTreeModule/Public/Debugger/StateTreeDebuggerTypes.h
+++ b/Engine/Plugins/Runtime/StateTree/Source/StateTreeModule/Public/Debugger/StateTreeDebuggerTypes.h
@@ -63,7 +63,7 @@ struct FFrameSpan
 struct FInstanceDescriptor : TSharedFromThis<FInstanceDescriptor>
 {
 	FInstanceDescriptor() = default;
-	STATETREEMODULE_API FInstanceDescriptor(const UStateTree* InStateTree, const FStateTreeInstanceDebugId InId, const FString& InName, const TRange<double> InLifetime);
+	STATETREEMODULE_API FInstanceDescriptor(const UStateTree* InStateTree, const FStateTreeInstanceDebugId InId, const FString& InName, const TRange<double>& InLifetime);
 
 	STATETREEMODULE_API bool IsValid() const;
 
