StateTree: Fix bad index in BindDelegate. It used the datahandle index instead of the node index.

diff --git a/Engine/Plugins/Runtime/StateTree/Source/StateTreeModule/Private/StateTreeExecutionContext.cpp b/Engine/Plugins/Runtime/StateTree/Source/StateTreeModule/Private/StateTreeExecutionContext.cpp
--- a/Engine/Plugins/Runtime/StateTree/Source/StateTreeModule/Private/StateTreeExecutionContext.cpp
+++ b/Engine/Plugins/Runtime/StateTree/Source/StateTreeModule/Private/StateTreeExecutionContext.cpp
@@ -1661,9 +1661,8 @@ void FStateTreeExecutionContext::BindDelegate(const FStateTreeDelegateListener&
 		return;
 	}
 
-	const int32 ActiveStateIndex = CurrentFrame->ActiveStates.IndexOfReverse(CurrentlyProcessedState);
-	const UE::StateTree::FActiveStateID CurrentlyStateID = ActiveStateIndex != INDEX_NONE ? CurrentFrame->ActiveStates.StateIDs[ActiveStateIndex] : UE::StateTree::FActiveStateID::Invalid;
-	GetExecState().DelegateActiveListeners.Add(Listener, MoveTemp(Delegate), CurrentFrame->FrameID, CurrentlyStateID, FStateTreeIndex16(CurrentNodeDataHandle.GetIndex()));
+	const UE::StateTree::FActiveStateID CurrentlyProcessedStateID = CurrentFrame->ActiveStates.FindStateID(CurrentlyProcessedState);
+	GetExecState().DelegateActiveListeners.Add(Listener, MoveTemp(Delegate), CurrentFrame->FrameID, CurrentlyProcessedStateID, FStateTreeIndex16(CurrentNodeIndex));
 }
 
 // Deprecated
