[AutoRTFM] Implement transactional safety for the StateTree crash handler.

This regression was introduced yesterday at https://p4-swarm.epicgames.net/changes/42822631
Critical sections and atomics cannot be used in AutoRTFM without appropriate safeguards.

diff --git a/Engine/Plugins/Runtime/StateTree/Source/StateTreeModule/Private/CrashReporter/StateTreeCrashReporterHandler.cpp b/Engine/Plugins/Runtime/StateTree/Source/StateTreeModule/Private/CrashReporter/StateTreeCrashReporterHandler.cpp
--- a/Engine/Plugins/Runtime/StateTree/Source/StateTreeModule/Private/CrashReporter/StateTreeCrashReporterHandler.cpp
+++ b/Engine/Plugins/Runtime/StateTree/Source/StateTreeModule/Private/CrashReporter/StateTreeCrashReporterHandler.cpp
@@ -3,10 +3,13 @@
 #include "CrashReporter/StateTreeCrashReporterHandler.h"
 
 #include "Misc/ScopeLock.h"
+#include "Misc/TransactionallySafeCriticalSection.h"
 #include "HAL/PlatformTLS.h"
 
 #include "StateTree.h"
 
+#include <atomic>
+
 #if UE_WITH_STATETREE_CRASHREPORTER
 
 namespace UE::StateTree::Private
@@ -29,8 +32,15 @@ public:
 			, Context(InContext)
 			, ThreadID(InThreadID)
 		{
-			static std::atomic IDGenerator = 0;
-			ID = ++IDGenerator;
+			ID = GenerateID();
+		}
+
+		UE_AUTORTFM_ALWAYS_OPEN
+		static uint32 GenerateID()
+		{
+			// This ID only needs to be unique; it doesn't need to be rolled back if an AutoRTFM transaction fails.
+			static std::atomic<uint32> IDGenerator = 0;
+			return ++IDGenerator;
 		}
 
 		TWeakObjectPtr<const UObject> Owner;
@@ -43,13 +53,13 @@ public:
 public:
 	uint32 PushInfo(TNotNull<const UObject*> InOwner, TNotNull<const UStateTree*> InStateTree, FName InContext)
 	{
-		FScopeLock LockGuard(&CriticalSection);
+		UE::TScopeLock LockGuard(CriticalSection);
 		return Contexts.Emplace_GetRef(InOwner, InStateTree, InContext, FPlatformTLS::GetCurrentThreadId()).ID;
 	}
 
 	void PopInfo(uint32 ID)
 	{
-		FScopeLock LockGuard(&CriticalSection);
+		UE::TScopeLock LockGuard(CriticalSection);
 		const int32 FoundIndex = Contexts.IndexOfByPredicate([ID](const FContext& Other)
 			{
 				return Other.ID == ID;
@@ -63,7 +73,7 @@ public:
 	void DumpInfo(FCrashContextExtendedWriter& Writer)
 	{
 		// Callback from the crash reporter.Avoids memory allocation in case the crash is an OOM.
-		FScopeLock LockGuard(&CriticalSection);
+		UE::TScopeLock LockGuard(CriticalSection);
 
 		for (int32 Index = 0; Index < Contexts.Num(); ++Index)
 		{
@@ -110,7 +120,7 @@ public:
 
 private:
 	TArray<FContext, TInlineAllocator<32>> Contexts;
-	FCriticalSection CriticalSection;
+	FTransactionallySafeCriticalSection CriticalSection;
 	TStringBuilder<1024> Builder;
 };
 
