[StateTreeDebugger] fixed RewindDebugger track using the wrong duration to build the timeline data

diff --git a/Engine/Plugins/Runtime/StateTree/Source/StateTreeEditorModule/Private/Debugger/StateTreeRewindDebuggerTrack.cpp b/Engine/Plugins/Runtime/StateTree/Source/StateTreeEditorModule/Private/Debugger/StateTreeRewindDebuggerTrack.cpp
--- a/Engine/Plugins/Runtime/StateTree/Source/StateTreeEditorModule/Private/Debugger/StateTreeRewindDebuggerTrack.cpp
+++ b/Engine/Plugins/Runtime/StateTree/Source/StateTreeEditorModule/Private/Debugger/StateTreeRewindDebuggerTrack.cpp
@@ -144,15 +144,18 @@ bool FRewindDebuggerTrack::UpdateInternal()
 	const IRewindDebugger* RewindDebugger = IRewindDebugger::Instance();
 	if (const TraceServices::IAnalysisSession* Session = RewindDebugger->GetAnalysisSession())
 	{
-		const double TraceTime = RewindDebugger->CurrentTraceTime();
 		const UStateTree* StateTree = Descriptor->StateTree.Get();
 
-		ReadTrace(*Session, TraceTime, /*ITraceReader*/this, FTraceFilter(InstanceTrackHelper.Get().GetInstanceId()), LastTraceReadTime);
+		ReadTrace(*Session
+			, RewindDebugger->CurrentTraceTime()
+			, /*ITraceReader*/this
+			, FTraceFilter(InstanceTrackHelper.Get().GetInstanceId())
+			, LastTraceReadTime);
 
 		return InstanceTrackHelper.Get().RebuildEventData(StateTree
 			, EventCollection
 			, Descriptor->Lifetime
-			, TraceTime
+			, RewindDebugger->GetScrubTime()
 			, /*Active*/true);
 	}
 
