StateTree: Keep the active node index on the execution frame instead of the execution state. A temporary frame that fails the EnterState will clean up properly. More work to be done there with the transition ID. Nodes that are not activated will not tick or receive ExitState. Add unittest to test for tasks that were not activated (enterstate).
Group TickGlobalEvaluatorsForFrame, StartGlobalsForFrame, and StopGlobalsForFrame to remove duplicated code. Add CopyAllBindingsOnActiveInstances for easily copying all the active bindings.

diff --git a/Engine/Plugins/Runtime/StateTree/Source/StateTreeModule/Private/StateTreeExecutionContext.cpp b/Engine/Plugins/Runtime/StateTree/Source/StateTreeModule/Private/StateTreeExecutionContext.cpp
--- a/Engine/Plugins/Runtime/StateTree/Source/StateTreeModule/Private/StateTreeExecutionContext.cpp
+++ b/Engine/Plugins/Runtime/StateTree/Source/StateTreeModule/Private/StateTreeExecutionContext.cpp
@@ -1178,15 +1178,13 @@ EStateTreeRunStatus FStateTreeExecutionContext::Start(FStartParameters Parameter
 	SetUpdatePhaseInExecutionState(Exec, EStateTreeUpdatePhase::StartTree);
 
 	// Start evaluators and global tasks. Fail the execution if any global task fails.
-	FStateTreeIndex16 LastInitializedTaskIndex;
-	const EStateTreeRunStatus GlobalTasksRunStatus = StartEvaluatorsAndGlobalTasks(LastInitializedTaskIndex);
+	const EStateTreeRunStatus GlobalTasksRunStatus = StartEvaluatorsAndGlobalTasks();
 	if (GlobalTasksRunStatus == EStateTreeRunStatus::Running)
 	{
-		// First tick.
-		// Tasks are not ticked here, since their behavior is that EnterState() (called above) is treated as a tick.
-		//@todo Check the result of TickEvaluatorsAndGlobalTasks and early exit if it is not running
-		constexpr bool bTickGlobalTasks = false;
-		TickEvaluatorsAndGlobalTasks(0.0f, bTickGlobalTasks);
+		// Exception with Start. Tick the evaluators.
+		constexpr float DeltaTime = 0.0f;
+		constexpr FStateTreeExecutionFrame* ParentInitFrame = nullptr;
+		TickGlobalEvaluatorsForFrameOnActiveInstances(DeltaTime, ParentInitFrame, InitFrame);
 
 		// Initialize to unset running state.
 		Exec.TreeRunStatus = EStateTreeRunStatus::Running;
@@ -1237,7 +1235,7 @@ EStateTreeRunStatus FStateTreeExecutionContext::Start(FStartParameters Parameter
 	}
 	else
 	{
-		StopEvaluatorsAndGlobalTasks(GlobalTasksRunStatus, LastInitializedTaskIndex);
+		StopEvaluatorsAndGlobalTasks(GlobalTasksRunStatus);
 
 		STATETREE_LOG(VeryVerbose, TEXT("%hs: Global tasks completed the StateTree %s on start in status '%s'."),
 			__FUNCTION__, *GetNameSafe(&Owner), *GetFullNameSafe(&RootStateTree), *UEnum::GetDisplayValueAsText(GlobalTasksRunStatus).ToString());
@@ -2327,6 +2325,92 @@ FStateTreeDataView FStateTreeExecutionContext::AddTemporaryInstance(const FState
 	return NewInstance;
 }
 
+void FStateTreeExecutionContext::CopyAllBindingsOnActiveInstances(const ECopyBindings CopyType)
+{
+	FStateTreeExecutionState& Exec = GetExecState();
+	for (int32 FrameIndex = 0; FrameIndex < Exec.ActiveFrames.Num(); FrameIndex++)
+	{
+		FStateTreeExecutionFrame* CurrentParentFrame = FrameIndex > 0 ? &Exec.ActiveFrames[FrameIndex - 1] : nullptr;
+		FStateTreeExecutionFrame& CurrentFrame = Exec.ActiveFrames[FrameIndex];
+		const UStateTree* CurrentStateTree = CurrentFrame.StateTree;
+		const int32 CurrentActiveTaskNodeIndex = CurrentFrame.ActiveTaskNodeIndex.AsInt32();
+
+		FCurrentlyProcessedFrameScope FrameScope(*this, CurrentParentFrame, CurrentFrame);
+
+		const bool bShouldCallOnEvaluatorsAndGlobalTasks = CurrentFrame.bIsGlobalFrame;
+		if (bShouldCallOnEvaluatorsAndGlobalTasks)
+		{
+			const int32 EvaluatorEnd = CurrentStateTree->EvaluatorsBegin + CurrentStateTree->EvaluatorsNum;
+			for (int32 EvalIndex = CurrentStateTree->EvaluatorsBegin; EvalIndex < EvaluatorEnd; ++EvalIndex)
+			{
+				const FStateTreeEvaluatorBase& Eval = CurrentStateTree->Nodes[EvalIndex].Get<const FStateTreeEvaluatorBase>();
+				if (Eval.BindingsBatch.IsValid())
+				{
+					const FStateTreeDataView EvalInstanceView = GetDataView(CurrentParentFrame, CurrentFrame, Eval.InstanceDataHandle);
+					FNodeInstanceDataScope DataScope(*this, &Eval, EvalIndex, Eval.InstanceDataHandle, EvalInstanceView);
+					CopyBatchOnActiveInstances(CurrentParentFrame, CurrentFrame, EvalInstanceView, Eval.BindingsBatch);
+				}
+			}
+
+			const int32 GlobalTasksEnd = CurrentStateTree->GlobalTasksBegin + CurrentStateTree->GlobalTasksNum;
+			for (int32 TaskIndex = CurrentStateTree->GlobalTasksBegin; TaskIndex < GlobalTasksEnd; ++TaskIndex)
+			{
+				if (TaskIndex <= CurrentActiveTaskNodeIndex)
+				{
+					const FStateTreeTaskBase& Task = CurrentStateTree->Nodes[TaskIndex].Get<const FStateTreeTaskBase>();
+					const bool bTaskRequestCopy = (CopyType == ECopyBindings::ExitState && Task.bShouldCopyBoundPropertiesOnExitState)
+						|| (CopyType == ECopyBindings::EnterState)
+						|| (CopyType == ECopyBindings::Tick && Task.bShouldCopyBoundPropertiesOnTick);
+					if (bTaskRequestCopy && Task.BindingsBatch.IsValid())
+					{
+						const FStateTreeDataView TaskInstanceView = GetDataView(CurrentParentFrame, CurrentFrame, Task.InstanceDataHandle);
+						FNodeInstanceDataScope DataScope(*this, &Task, TaskIndex, Task.InstanceDataHandle, TaskInstanceView);
+						CopyBatchOnActiveInstances(CurrentParentFrame, CurrentFrame, TaskInstanceView, Task.BindingsBatch);
+					}
+				}
+			}
+		}
+
+		for (int32 StateIndex = 0; StateIndex < CurrentFrame.ActiveStates.Num(); ++StateIndex)
+		{
+			const FStateTreeStateHandle CurrentStateHandle = CurrentFrame.ActiveStates.States[StateIndex];
+			const UE::StateTree::FActiveStateID CurrentStateID = CurrentFrame.ActiveStates.StateIDs[StateIndex];
+			const FCompactStateTreeState& CurrentState = CurrentStateTree->States[CurrentStateHandle.Index];
+
+			FCurrentlyProcessedStateScope StateScope(*this, CurrentStateHandle);
+
+			if (CurrentState.Type == EStateTreeStateType::Linked
+				|| CurrentState.Type == EStateTreeStateType::LinkedAsset)
+			{
+				if (CurrentState.ParameterDataHandle.IsValid()
+					&& CurrentState.ParameterBindingsBatch.IsValid())
+				{
+					const FStateTreeDataView StateParamsDataView = GetDataView(CurrentParentFrame, CurrentFrame, CurrentState.ParameterDataHandle);
+					CopyBatchOnActiveInstances(CurrentParentFrame, CurrentFrame, StateParamsDataView, CurrentState.ParameterBindingsBatch);
+				}
+			}
+
+			const int32 TasksEnd = CurrentState.TasksBegin + CurrentState.TasksNum;
+			for (int32 TaskIndex = CurrentState.TasksBegin; TaskIndex < TasksEnd; ++TaskIndex)
+			{
+				if (TaskIndex <= CurrentActiveTaskNodeIndex)
+				{
+					const FStateTreeTaskBase& Task = CurrentStateTree->Nodes[TaskIndex].Get<const FStateTreeTaskBase>();
+
+					const bool bTaskRequestCopy = (CopyType == ECopyBindings::ExitState && Task.bShouldCopyBoundPropertiesOnExitState)
+						|| (CopyType == ECopyBindings::EnterState)
+						|| (CopyType == ECopyBindings::Tick && Task.bShouldCopyBoundPropertiesOnTick);
+					if (bTaskRequestCopy && Task.BindingsBatch.IsValid())
+					{
+						const FStateTreeDataView TaskInstanceView = GetDataView(CurrentParentFrame, CurrentFrame, Task.InstanceDataHandle);
+						CopyBatchOnActiveInstances(CurrentParentFrame, CurrentFrame, TaskInstanceView, Task.BindingsBatch);
+					}
+				}
+			}
+		}
+	}
+}
+
 bool FStateTreeExecutionContext::CopyBatchOnActiveInstances(const FStateTreeExecutionFrame* ParentFrame, const FStateTreeExecutionFrame& CurrentFrame, const FStateTreeDataView TargetView, const FStateTreeIndex16 BindingsBatch)
 {
 	const FPropertyBindingCopyInfoBatch& Batch = CurrentFrame.StateTree->PropertyBindings.Super::GetBatch(BindingsBatch);
@@ -2587,8 +2671,6 @@ EStateTreeRunStatus FStateTreeExecutionContext::EnterState(FStateTreeTransitionR
 	CaptureNewStateEvents(Exec.ActiveFrames, Transition.NextActiveFrames, Transition.NextActiveFrameEvents);
 
 	Exec.StateChangeCount++;
-	Exec.EnterStateFailedFrameIndex = FStateTreeIndex16::Invalid; // This will make all tasks to be accepted.
-	Exec.EnterStateFailedTaskIndex = FStateTreeIndex16::Invalid; // This will make all tasks to be accepted.
 	
 	// On target branch means that the state is the target of current transition or child of it.
 	// States which were active before and will remain active, but are not on target branch will not get
@@ -2746,6 +2828,8 @@ EStateTreeRunStatus FStateTreeExecutionContext::EnterState(FStateTreeTransitionR
 
 				FNodeInstanceDataScope DataScope(*this, &Task, AssetTaskIndex, Task.InstanceDataHandle, TaskInstanceView);
 
+				CurrentFrame.ActiveTaskNodeIndex = FStateTreeIndex16(AssetTaskIndex);
+
 				// Copy bound properties.
 				if (Task.BindingsBatch.IsValid())
 				{
@@ -2781,12 +2865,13 @@ EStateTreeRunStatus FStateTreeExecutionContext::EnterState(FStateTreeTransitionR
 					TaskRunStatus = UE::StateTree::ExecutionContext::CastToRunStatus(TaskStatus);
 					UE_STATETREE_DEBUG_TASK_EVENT(this, AssetTaskIndex, TaskInstanceView, EStateTreeTraceEventType::OnEntered, TaskRunStatus);
 
-					Result = UE::StateTree::ExecutionContext::GetPriorityRunStatus(Result, TaskRunStatus);
-					if (TaskRunStatus == EStateTreeRunStatus::Failed && CurrentStateTasksStatus.IsConsideredForCompletion(StateTaskIndex))
+					if (CurrentStateTasksStatus.IsConsideredForCompletion(StateTaskIndex))
 					{
-						Exec.EnterStateFailedFrameIndex = FStateTreeIndex16(FrameIndex);
-						Exec.EnterStateFailedTaskIndex = FStateTreeIndex16(AssetTaskIndex);
-						break;
+						Result = UE::StateTree::ExecutionContext::GetPriorityRunStatus(Result, TaskRunStatus);
+						if (TaskRunStatus == EStateTreeRunStatus::Failed)
+						{
+							break;
+						}
 					}
 				}
 			}
@@ -2810,9 +2895,6 @@ void FStateTreeExecutionContext::ExitState(const FStateTreeTransitionResult& Tra
 {
 	CSV_SCOPED_TIMING_STAT_EXCLUSIVE(StateTree_ExitState);
 
-	check(!GetExecState().LastExitedNodeIndex.IsValid());
-	ON_SCOPE_EXIT { GetExecState().LastExitedNodeIndex = FStateTreeIndex16::Invalid; };
-
 	FStateTreeExecutionState& Exec = GetExecState();
 
 	if (Exec.ActiveFrames.IsEmpty())
@@ -3002,8 +3084,7 @@ void FStateTreeExecutionContext::ExitState(const FStateTreeTransitionResult& Tra
 
 					// Call task completed only if EnterState() was called.
 					// The task order in the tree (BF) allows us to use the comparison.
-					// Relying here that invalid value of Exec.EnterStateFailedTaskIndex == MAX_uint16.
-					if (TaskIndex <= Exec.EnterStateFailedTaskIndex.Get())
+					if (TaskIndex <= CurrentFrame.ActiveTaskNodeIndex.AsInt32())
 					{
 						const FStateTreeTaskBase& Task = CurrentStateTree->Nodes[TaskIndex].Get<const FStateTreeTaskBase>();
 						const FStateTreeDataView TaskInstanceView = GetDataView(CurrentParentFrame, CurrentFrame, Task.InstanceDataHandle);
@@ -3011,28 +3092,29 @@ void FStateTreeExecutionContext::ExitState(const FStateTreeTransitionResult& Tra
 						FNodeInstanceDataScope DataScope(*this, &Task, TaskIndex, Task.InstanceDataHandle, TaskInstanceView);
 
 						// Ignore disabled task
-						if (Task.bTaskEnabled == false)
+						if (Task.bTaskEnabled)
 						{
-							STATETREE_LOG(VeryVerbose, TEXT("%*sSkipped 'ExitState' for disabled Task: '%s'"), UE::StateTree::Debug::IndentSize, TEXT(""), *Task.Name.ToString());
-							continue;
-						}
-
-						const bool bShouldCallStateChange = CurrentTransition.ChangeType == EStateTreeStateChangeType::Changed
+							const bool bShouldCallStateChange = CurrentTransition.ChangeType == EStateTreeStateChangeType::Changed
 									|| (CurrentTransition.ChangeType == EStateTreeStateChangeType::Sustained && Task.bShouldStateChangeOnReselect);
 
-						if (bShouldCallStateChange)
-						{
-							STATETREE_LOG(Verbose, TEXT("%*sTask '%s'.ExitState()"), (FrameIndex + StateIndex + 1)*UE::StateTree::Debug::IndentSize, TEXT(""), *Task.Name.ToString());
-							UE_STATETREE_DEBUG_TASK_EXIT_STATE(this, CurrentStateTree, FStateTreeIndex16(TaskIndex));
+							if (bShouldCallStateChange)
 							{
-								QUICK_SCOPE_CYCLE_COUNTER(StateTree_Task_ExitState);
-								CSV_SCOPED_TIMING_STAT_EXCLUSIVE(StateTree_Task_ExitState);
-								Task.ExitState(*this, CurrentTransition);
+								STATETREE_LOG(Verbose, TEXT("%*sTask '%s'.ExitState()"), (FrameIndex + StateIndex + 1)*UE::StateTree::Debug::IndentSize, TEXT(""), *Task.Name.ToString());
+								UE_STATETREE_DEBUG_TASK_EXIT_STATE(this, CurrentStateTree, FStateTreeIndex16(TaskIndex));
+								{
+									QUICK_SCOPE_CYCLE_COUNTER(StateTree_Task_ExitState);
+									CSV_SCOPED_TIMING_STAT_EXCLUSIVE(StateTree_Task_ExitState);
+									Task.ExitState(*this, CurrentTransition);
+								}
+								UE_STATETREE_DEBUG_TASK_EVENT(this, TaskIndex, TaskInstanceView, EStateTreeTraceEventType::OnExited, Transition.CurrentRunStatus);
 							}
-							UE_STATETREE_DEBUG_TASK_EVENT(this, TaskIndex, TaskInstanceView, EStateTreeTraceEventType::OnExited, Transition.CurrentRunStatus);
-							GetExecState().LastExitedNodeIndex = FStateTreeIndex16(TaskIndex);
+						}
+						else
+						{
+							STATETREE_LOG(VeryVerbose, TEXT("%*sSkipped 'ExitState' for disabled Task: '%s'"), UE::StateTree::Debug::IndentSize, TEXT(""), *Task.Name.ToString());
 						}
 					}
+					CurrentFrame.ActiveTaskNodeIndex = FStateTreeIndex16(TaskIndex - 1);
 				}
 
 				// Call state change events on conditions if needed.
@@ -3082,8 +3164,9 @@ void FStateTreeExecutionContext::ExitState(const FStateTreeTransitionResult& Tra
 			const FExitStateCall& ExitCall = ExitStateCalls[CallIndex--];
 			if (ExitCall.bShouldCall)
 			{
+				CurrentTransition.CurrentState = FStateTreeStateHandle::Invalid;
 				CurrentTransition.ChangeType = ExitCall.ChangeType;
-				CallStopOnEvaluatorsAndGlobalTasks(CurrentParentFrame, CurrentFrame, CurrentTransition);
+				StopGlobalsForFrameOnActiveInstances(CurrentParentFrame, CurrentFrame, CurrentTransition);
 
 				// Delegate Listeners Cleanup
 				GetExecState().DelegateActiveListeners.RemoveAll(CurrentFrame.FrameID);
@@ -3119,68 +3202,64 @@ void FStateTreeExecutionContext::StateCompleted()
 		const FStateTreeExecutionFrame* CurrentParentFrame = FrameIndex > 0 ? &Exec.ActiveFrames[FrameIndex - 1] : nullptr;
 		const FStateTreeExecutionFrame& CurrentFrame = Exec.ActiveFrames[FrameIndex];
 		const UStateTree* CurrentStateTree = CurrentFrame.StateTree;
+		const int32 CurrentActiveTaskNodeIndex = CurrentFrame.ActiveTaskNodeIndex.AsInt32();
 
 		FCurrentlyProcessedFrameScope FrameScope(*this, CurrentParentFrame, CurrentFrame);
 
-		if (FrameIndex <= Exec.EnterStateFailedFrameIndex.Get())
+		for (int32 StateIndex = CurrentFrame.ActiveStates.Num() - 1; StateIndex >= 0; StateIndex--)
 		{
-			for (int32 StateIndex = CurrentFrame.ActiveStates.Num() - 1; StateIndex >= 0; StateIndex--)
-			{
-				const FStateTreeStateHandle CurrentHandle = CurrentFrame.ActiveStates[StateIndex];
-				const FCompactStateTreeState& State = CurrentStateTree->States[CurrentHandle.Index];
+			const FStateTreeStateHandle CurrentHandle = CurrentFrame.ActiveStates[StateIndex];
+			const FCompactStateTreeState& State = CurrentStateTree->States[CurrentHandle.Index];
 
-				FCurrentlyProcessedStateScope StateScope(*this, CurrentHandle);
-				UE_STATETREE_DEBUG_STATE_EVENT(this, CurrentHandle, EStateTreeTraceEventType::OnStateCompleted);
-				STATETREE_LOG(Verbose, TEXT("%*sState '%s'"), (FrameIndex + StateIndex + 1)*UE::StateTree::Debug::IndentSize, TEXT(""), *GetSafeStateName(CurrentFrame, CurrentHandle));
+			FCurrentlyProcessedStateScope StateScope(*this, CurrentHandle);
+			UE_STATETREE_DEBUG_STATE_EVENT(this, CurrentHandle, EStateTreeTraceEventType::OnStateCompleted);
+			STATETREE_LOG(Verbose, TEXT("%*sState '%s'"), (FrameIndex + StateIndex + 1)*UE::StateTree::Debug::IndentSize, TEXT(""), *GetSafeStateName(CurrentFrame, CurrentHandle));
 
-				// Notify Tasks
-				for (int32 TaskIndex = (State.TasksBegin + State.TasksNum) - 1; TaskIndex >= State.TasksBegin; TaskIndex--)
+			// Notify Tasks
+			for (int32 TaskIndex = (State.TasksBegin + State.TasksNum) - 1; TaskIndex >= State.TasksBegin; TaskIndex--)
+			{
+				// Call task completed only if EnterState() was called.
+				if (TaskIndex <= CurrentActiveTaskNodeIndex)
 				{
-					// Call task completed only if EnterState() was called.
-					// The task order in the tree (BF) allows us to use the comparison.
-					// Relying here that invalid value of Exec.EnterStateFailedTaskIndex == MAX_uint16.
-					if (TaskIndex <= Exec.EnterStateFailedTaskIndex.Get())
-					{
-						const FStateTreeTaskBase& Task = CurrentStateTree->Nodes[TaskIndex].Get<const FStateTreeTaskBase>();
-						const FStateTreeDataView TaskInstanceView = GetDataView(CurrentParentFrame, CurrentFrame, Task.InstanceDataHandle);
-						FNodeInstanceDataScope DataScope(*this, &Task, TaskIndex, Task.InstanceDataHandle, TaskInstanceView);
-
-						// Ignore disabled task
-						if (Task.bTaskEnabled == false)
-						{
-							STATETREE_LOG(VeryVerbose, TEXT("%*sSkipped 'StateCompleted' for disabled Task: '%s'"), UE::StateTree::Debug::IndentSize, TEXT(""), *Task.Name.ToString());
-							continue;
-						}
+					const FStateTreeTaskBase& Task = CurrentStateTree->Nodes[TaskIndex].Get<const FStateTreeTaskBase>();
+					const FStateTreeDataView TaskInstanceView = GetDataView(CurrentParentFrame, CurrentFrame, Task.InstanceDataHandle);
+					FNodeInstanceDataScope DataScope(*this, &Task, TaskIndex, Task.InstanceDataHandle, TaskInstanceView);
 
-						STATETREE_LOG(Verbose, TEXT("%*sTask '%s'.StateCompleted()"), (FrameIndex + StateIndex + 1)*UE::StateTree::Debug::IndentSize, TEXT(""), *Task.Name.ToString());
-						Task.StateCompleted(*this, Exec.LastTickStatus, CurrentFrame.ActiveStates);
+					// Ignore disabled task
+					if (Task.bTaskEnabled == false)
+					{
+						STATETREE_LOG(VeryVerbose, TEXT("%*sSkipped 'StateCompleted' for disabled Task: '%s'"), UE::StateTree::Debug::IndentSize, TEXT(""), *Task.Name.ToString());
+						continue;
 					}
+
+					STATETREE_LOG(Verbose, TEXT("%*sTask '%s'.StateCompleted()"), (FrameIndex + StateIndex + 1)*UE::StateTree::Debug::IndentSize, TEXT(""), *Task.Name.ToString());
+					Task.StateCompleted(*this, Exec.LastTickStatus, CurrentFrame.ActiveStates);
 				}
+			}
 				
-				// Call state change events on conditions if needed.
-				if (State.bHasStateChangeConditions)
+			// Call state change events on conditions if needed.
+			if (State.bHasStateChangeConditions)
+			{
+				for (int32 ConditionIndex = (State.EnterConditionsBegin + State.EnterConditionsNum) - 1; ConditionIndex >= State.EnterConditionsBegin; ConditionIndex--)
 				{
-					for (int32 ConditionIndex = (State.EnterConditionsBegin + State.EnterConditionsNum) - 1; ConditionIndex >= State.EnterConditionsBegin; ConditionIndex--)
+					const FStateTreeConditionBase& Cond = CurrentFrame.StateTree->Nodes[ConditionIndex].Get<const FStateTreeConditionBase>();
+					if (Cond.bHasShouldCallStateChangeEvents)
 					{
-						const FStateTreeConditionBase& Cond = CurrentFrame.StateTree->Nodes[ConditionIndex].Get<const FStateTreeConditionBase>();
-						if (Cond.bHasShouldCallStateChangeEvents)
-						{
-							const FStateTreeDataView ConditionInstanceView = GetDataView(CurrentParentFrame, CurrentFrame, Cond.InstanceDataHandle);
-							FNodeInstanceDataScope DataScope(*this, &Cond, ConditionIndex, Cond.InstanceDataHandle, ConditionInstanceView);
+						const FStateTreeDataView ConditionInstanceView = GetDataView(CurrentParentFrame, CurrentFrame, Cond.InstanceDataHandle);
+						FNodeInstanceDataScope DataScope(*this, &Cond, ConditionIndex, Cond.InstanceDataHandle, ConditionInstanceView);
 
-							if (Cond.BindingsBatch.IsValid())
-							{
-								// Use validated copy, since we test in situations where the sources are not always valid (e.g. enter conditions may try to access inactive parent state).
-								CopyBatchOnActiveInstances(CurrentParentFrame, CurrentFrame, ConditionInstanceView, Cond.BindingsBatch);
-							}
+						if (Cond.BindingsBatch.IsValid())
+						{
+							// Use validated copy, since we test in situations where the sources are not always valid (e.g. enter conditions may try to access inactive parent state).
+							CopyBatchOnActiveInstances(CurrentParentFrame, CurrentFrame, ConditionInstanceView, Cond.BindingsBatch);
+						}
 							
-							Cond.StateCompleted(*this, Exec.LastTickStatus, CurrentFrame.ActiveStates);
+						Cond.StateCompleted(*this, Exec.LastTickStatus, CurrentFrame.ActiveStates);
 
-							// Reset copied properties that might contain object references.
-							if (Cond.BindingsBatch.IsValid())
-							{
-								CurrentFrame.StateTree->PropertyBindings.Super::ResetObjects(Cond.BindingsBatch, ConditionInstanceView);
-							}
+						// Reset copied properties that might contain object references.
+						if (Cond.BindingsBatch.IsValid())
+						{
+							CurrentFrame.StateTree->PropertyBindings.Super::ResetObjects(Cond.BindingsBatch, ConditionInstanceView);
 						}
 					}
 				}
@@ -3189,6 +3268,61 @@ void FStateTreeExecutionContext::StateCompleted()
 	}
 }
 
+void FStateTreeExecutionContext::TickGlobalEvaluatorsForFrameOnActiveInstances(const float DeltaTime, const FStateTreeExecutionFrame* ParentFrame, const FStateTreeExecutionFrame& Frame)
+{
+	constexpr bool bOnActiveInstances = true;
+	TickGlobalEvaluatorsForFrameInternal<bOnActiveInstances>(DeltaTime, ParentFrame, Frame);
+}
+
+void FStateTreeExecutionContext::TickGlobalEvaluatorsForFrameWithValidation(const float DeltaTime, const FStateTreeExecutionFrame* ParentFrame, const FStateTreeExecutionFrame& Frame)
+{
+	constexpr bool bOnActiveInstances = false;
+	TickGlobalEvaluatorsForFrameInternal<bOnActiveInstances>(DeltaTime, ParentFrame, Frame);
+}
+
+template<bool bOnActiveInstances>
+void FStateTreeExecutionContext::TickGlobalEvaluatorsForFrameInternal(const float DeltaTime, const FStateTreeExecutionFrame* ParentFrame, const FStateTreeExecutionFrame& Frame)
+{
+	check(Frame.bIsGlobalFrame);
+
+	const UStateTree* CurrentStateTree = Frame.StateTree;
+	const int32 EvaluatorEnd = CurrentStateTree->EvaluatorsBegin + CurrentStateTree->EvaluatorsNum;
+	if (CurrentStateTree->EvaluatorsBegin < EvaluatorEnd)
+	{
+		FCurrentlyProcessedFrameScope FrameScope(*this, ParentFrame, Frame);
+
+		for (int32 EvalIndex = CurrentStateTree->EvaluatorsBegin; EvalIndex < EvaluatorEnd; ++EvalIndex)
+		{
+			const FStateTreeEvaluatorBase& Eval = CurrentStateTree->Nodes[EvalIndex].Get<const FStateTreeEvaluatorBase>();
+			FStateTreeDataView EvalInstanceView = bOnActiveInstances
+				? GetDataView(ParentFrame, Frame, Eval.InstanceDataHandle)
+				: GetDataViewOrTemporary(ParentFrame, Frame, Eval.InstanceDataHandle);
+			FNodeInstanceDataScope DataScope(*this, &Eval, EvalIndex, Eval.InstanceDataHandle, EvalInstanceView);
+
+			// Copy bound properties.
+			if (Eval.BindingsBatch.IsValid())
+			{
+				if constexpr (bOnActiveInstances)
+				{
+					CopyBatchOnActiveInstances(ParentFrame, Frame, EvalInstanceView, Eval.BindingsBatch);
+				}
+				else
+				{
+					CopyBatchWithValidation(ParentFrame, Frame, EvalInstanceView, Eval.BindingsBatch);
+				}
+			}
+
+			STATETREE_LOG(VeryVerbose, TEXT("  Tick: '%s'"), *Eval.Name.ToString());
+			UE_STATETREE_DEBUG_EVALUATOR_TICK(this, CurrentStateTree, EvalIndex);
+			{
+				QUICK_SCOPE_CYCLE_COUNTER(StateTree_Eval_Tick);
+				Eval.Tick(*this, DeltaTime);
+				UE_STATETREE_DEBUG_EVALUATOR_EVENT(this, EvalIndex, EvalInstanceView, EStateTreeTraceEventType::OnTicked);
+			}
+		}
+	}
+}
+
 EStateTreeRunStatus FStateTreeExecutionContext::TickEvaluatorsAndGlobalTasks(const float DeltaTime, const bool bTickGlobalTasks)
 {
 	// When a global task is completed it completes the tree execution.
@@ -3232,34 +3366,14 @@ EStateTreeRunStatus FStateTreeExecutionContext::TickEvaluatorsAndGlobalTasksForF
 
 	EStateTreeRunStatus Result = EStateTreeRunStatus::Running;
 
-	const UStateTree* CurrentStateTree = CurrentFrame->StateTree;
-
 	// Tick evaluators
-	for (int32 EvalIndex = CurrentStateTree->EvaluatorsBegin; EvalIndex < (CurrentStateTree->EvaluatorsBegin + CurrentStateTree->EvaluatorsNum); EvalIndex++)
-	{
-		const FStateTreeEvaluatorBase& Eval = CurrentStateTree->Nodes[EvalIndex].Get<const FStateTreeEvaluatorBase>();
-		const FStateTreeDataView EvalInstanceView = GetDataView(CurrentParentFrame, *CurrentFrame, Eval.InstanceDataHandle);
-		FNodeInstanceDataScope DataScope(*this, &Eval, EvalIndex, Eval.InstanceDataHandle, EvalInstanceView);
-
-		// Copy bound properties.
-		if (Eval.BindingsBatch.IsValid())
-		{
-			CopyBatchOnActiveInstances(CurrentParentFrame, *CurrentFrame, EvalInstanceView, Eval.BindingsBatch);
-		}
-		STATETREE_LOG(VeryVerbose, TEXT("  Tick: '%s'"), *Eval.Name.ToString());
-		UE_STATETREE_DEBUG_EVALUATOR_TICK(this, CurrentStateTree, EvalIndex);
-		{
-			QUICK_SCOPE_CYCLE_COUNTER(StateTree_Eval_Tick);
-			Eval.Tick(*this, DeltaTime);
-
-			UE_STATETREE_DEBUG_EVALUATOR_EVENT(this, EvalIndex, EvalInstanceView, EStateTreeTraceEventType::OnTicked);
-		}
-	}
+	TickGlobalEvaluatorsForFrameOnActiveInstances(DeltaTime, CurrentParentFrame, *CurrentFrame);
 
 	if (bTickGlobalTasks)
 	{
 		using namespace UE::StateTree;
 
+		const UStateTree* CurrentStateTree = CurrentFrame->StateTree;
 		FTasksCompletionStatus CurrentGlobalTasksStatus = CurrentFrame->ActiveTasksStatus.GetStatus(CurrentStateTree);
 		if (!CurrentGlobalTasksStatus.HasAnyFailed())
 		{
@@ -3289,240 +3403,35 @@ EStateTreeRunStatus FStateTreeExecutionContext::TickEvaluatorsAndGlobalTasksForF
 	return Result;
 }
 
-EStateTreeRunStatus FStateTreeExecutionContext::StartEvaluatorsAndGlobalTasks(FStateTreeIndex16& OutLastInitializedTaskIndex)
+EStateTreeRunStatus FStateTreeExecutionContext::StartGlobalsForFrameOnActiveInstances(const FStateTreeExecutionFrame* ParentFrame, FStateTreeExecutionFrame& Frame, FStateTreeTransitionResult& Transition)
 {
-	CSV_SCOPED_TIMING_STAT_EXCLUSIVE(StateTree_StartEvaluators);
-	UE_STATETREE_DEBUG_SCOPED_PHASE(this, EStateTreeUpdatePhase::StartGlobalTasks);
-
-	STATETREE_LOG(Verbose, TEXT("Start Evaluators & Global tasks"));
-
-	FStateTreeExecutionState& Exec = GetExecState();
-
-	OutLastInitializedTaskIndex = FStateTreeIndex16();
-	EStateTreeRunStatus Result = EStateTreeRunStatus::Running;
-
-	for (int32 FrameIndex = 0; FrameIndex < Exec.ActiveFrames.Num(); FrameIndex++)
-	{
-		FStateTreeExecutionFrame* CurrentParentFrame = FrameIndex > 0 ? &Exec.ActiveFrames[FrameIndex - 1] : nullptr;
-		FStateTreeExecutionFrame& CurrentFrame = Exec.ActiveFrames[FrameIndex];
-		const UE::StateTree::FActiveFrameID CurrentFrameID = CurrentFrame.FrameID;
-		if (CurrentFrame.bIsGlobalFrame)
-		{
-			FCurrentlyProcessedFrameScope FrameScope(*this, CurrentParentFrame, CurrentFrame);
-			
-			const UStateTree* CurrentStateTree = CurrentFrame.StateTree;
-			UE::StateTree::FTasksCompletionStatus CurrentGlobalTasksStatus = CurrentFrame.ActiveTasksStatus.GetStatus(CurrentStateTree);
-
-			// Start evaluators
-			for (int32 EvalIndex = CurrentStateTree->EvaluatorsBegin; EvalIndex < (CurrentStateTree->EvaluatorsBegin + CurrentStateTree->EvaluatorsNum); EvalIndex++)
-			{
-				const FStateTreeEvaluatorBase& Eval = CurrentStateTree->Nodes[EvalIndex].Get<const FStateTreeEvaluatorBase>();
-				const FStateTreeDataView EvalInstanceView = GetDataView(CurrentParentFrame, CurrentFrame, Eval.InstanceDataHandle);
-				FNodeInstanceDataScope DataScope(*this, &Eval, EvalIndex, Eval.InstanceDataHandle, EvalInstanceView);
-
-				// Copy bound properties.
-				if (Eval.BindingsBatch.IsValid())
-				{
-					CopyBatchOnActiveInstances(CurrentParentFrame, CurrentFrame, EvalInstanceView, Eval.BindingsBatch);
-				}
-				STATETREE_LOG(Verbose, TEXT("  Start: '%s'"), *Eval.Name.ToString());
-				UE_STATETREE_DEBUG_EVALUATOR_ENTER_TREE(this, CurrentStateTree, FStateTreeIndex16(EvalIndex));
-				{
-					QUICK_SCOPE_CYCLE_COUNTER(StateTree_Eval_TreeStart);
-					Eval.TreeStart(*this);
-
-					UE_STATETREE_DEBUG_EVALUATOR_EVENT(this, EvalIndex, EvalInstanceView, EStateTreeTraceEventType::OnTreeStarted);
-				}
-			}
-
-			// Start Global tasks
-			// Even if we call Enter/ExitState() on global tasks, they do not enter any specific state.
-			const FStateTreeTransitionResult Transition = {}; // Empty transition
-
-			for (int32 GlobalTaskIndex = 0; GlobalTaskIndex < CurrentStateTree->GlobalTasksNum; ++GlobalTaskIndex)
-			{
-				const int32 AssetTaskIndex = CurrentStateTree->GlobalTasksBegin + GlobalTaskIndex;
-				const FStateTreeTaskBase& Task =  CurrentStateTree->Nodes[AssetTaskIndex].Get<const FStateTreeTaskBase>();
-				const FStateTreeDataView TaskInstanceView = GetDataView(CurrentParentFrame, CurrentFrame, Task.InstanceDataHandle);
-				FNodeInstanceDataScope DataScope(*this, &Task, AssetTaskIndex, Task.InstanceDataHandle, TaskInstanceView);
-
-				// Copy bound properties.
-				if (Task.BindingsBatch.IsValid())
-				{
-					CopyBatchOnActiveInstances(CurrentParentFrame, CurrentFrame, TaskInstanceView, Task.BindingsBatch);
-				}
-
-				// Ignore disabled task
-				if (Task.bTaskEnabled == false)
-				{
-					STATETREE_LOG(VeryVerbose, TEXT("%*sSkipped 'EnterState' for disabled Task: '%s'"), UE::StateTree::Debug::IndentSize, TEXT(""), *Task.Name.ToString());
-					continue;
-				}
-
-				STATETREE_LOG(Verbose, TEXT("  Start: '%s'"), *Task.Name.ToString());
-				UE_STATETREE_DEBUG_TASK_ENTER_STATE(this, CurrentStateTree, FStateTreeIndex16(AssetTaskIndex));
-
-				EStateTreeRunStatus TaskRunStatus = EStateTreeRunStatus::Unset;
-				{
-					QUICK_SCOPE_CYCLE_COUNTER(StateTree_Task_TreeStart);
-
-					TaskRunStatus = Task.EnterState(*this, Transition);
-				}
-
-				UE::StateTree::ETaskCompletionStatus TaskStatus = UE::StateTree::ExecutionContext::CastToTaskStatus(TaskRunStatus);
-				TaskStatus = CurrentGlobalTasksStatus.SetStatusWithPriority(GlobalTaskIndex, TaskStatus);
-
-				TaskRunStatus = UE::StateTree::ExecutionContext::CastToRunStatus(TaskStatus);
-				UE_STATETREE_DEBUG_TASK_EVENT(this, AssetTaskIndex, TaskInstanceView, EStateTreeTraceEventType::OnEntered, TaskRunStatus);
-
-				Result = UE::StateTree::ExecutionContext::GetPriorityRunStatus(Result, TaskRunStatus);
-				if (TaskRunStatus == EStateTreeRunStatus::Failed && CurrentGlobalTasksStatus.IsConsideredForCompletion(GlobalTaskIndex))
-				{
-					OutLastInitializedTaskIndex = FStateTreeIndex16(AssetTaskIndex);
-				}
-			}
-		}
-	}
-
-	return Result;
+	constexpr bool bOnActiveInstances = true;
+	return StartGlobalsForFrameInternal<bOnActiveInstances>(ParentFrame, Frame, Transition);
 }
 
-void FStateTreeExecutionContext::StopEvaluatorsAndGlobalTasks(const EStateTreeRunStatus CompletionStatus, const FStateTreeIndex16 LastInitializedTaskIndex)
+EStateTreeRunStatus FStateTreeExecutionContext::StartGlobalsForFrameWithValidation(const FStateTreeExecutionFrame* ParentFrame, FStateTreeExecutionFrame& Frame, FStateTreeTransitionResult& Transition)
 {
-	CSV_SCOPED_TIMING_STAT_EXCLUSIVE(StateTree_StopEvaluators);
-	UE_STATETREE_DEBUG_SCOPED_PHASE(this, EStateTreeUpdatePhase::StopGlobalTasks);
-
-	STATETREE_LOG(Verbose, TEXT("Stop Evaluators & Global Tasks"));
-
-	FStateTreeExecutionState& Exec = GetExecState();
-
-	// Update bindings
-	for (int32 FrameIndex = 0; FrameIndex < Exec.ActiveFrames.Num(); FrameIndex++)
-	{
-		FStateTreeExecutionFrame* CurrentParentFrame = FrameIndex > 0 ? &Exec.ActiveFrames[FrameIndex - 1] : nullptr;
-		FStateTreeExecutionFrame& CurrentFrame = Exec.ActiveFrames[FrameIndex];
-		if (CurrentFrame.bIsGlobalFrame)
-		{
-			FCurrentlyProcessedFrameScope FrameScope(*this, CurrentParentFrame, CurrentFrame);
-			
-			const UStateTree* CurrentStateTree = CurrentFrame.StateTree;
-
-			for (int32 EvalIndex = CurrentStateTree->EvaluatorsBegin; EvalIndex < (CurrentStateTree->EvaluatorsBegin + CurrentStateTree->EvaluatorsNum); EvalIndex++)
-			{
-				const FStateTreeEvaluatorBase& Eval = CurrentStateTree->Nodes[EvalIndex].Get<const FStateTreeEvaluatorBase>();
-				const FStateTreeDataView EvalInstanceView = GetDataView(CurrentParentFrame, CurrentFrame, Eval.InstanceDataHandle);
-				FNodeInstanceDataScope DataScope(*this, &Eval, EvalIndex, Eval.InstanceDataHandle, EvalInstanceView);
-
-				// Copy bound properties.
-				if (Eval.BindingsBatch.IsValid())
-				{
-					CopyBatchOnActiveInstances(CurrentParentFrame, CurrentFrame, EvalInstanceView, Eval.BindingsBatch);
-				}
-			}
-
-			for (int32 TaskIndex = CurrentStateTree->GlobalTasksBegin; TaskIndex < (CurrentStateTree->GlobalTasksBegin + CurrentStateTree->GlobalTasksNum); TaskIndex++)
-			{
-				const FStateTreeTaskBase& Task = CurrentStateTree->Nodes[TaskIndex].Get<const FStateTreeTaskBase>();
-				const FStateTreeDataView TaskInstanceView = GetDataView(CurrentParentFrame, CurrentFrame, Task.InstanceDataHandle);
-				FNodeInstanceDataScope DataScope(*this, &Task, TaskIndex, Task.InstanceDataHandle, TaskInstanceView);
-
-				// Copy bound properties.
-				if (Task.BindingsBatch.IsValid() && Task.bShouldCopyBoundPropertiesOnExitState)
-				{
-					CopyBatchOnActiveInstances(CurrentParentFrame, CurrentFrame, TaskInstanceView, Task.BindingsBatch);
-				}
-			}
-		}
-	}
-
-	// Call in reverse order.
-	FStateTreeTransitionResult Transition;
-	Transition.TargetState = FStateTreeStateHandle::FromCompletionStatus(CompletionStatus);
-	Transition.CurrentRunStatus = CompletionStatus;
-
-	bool bIsLastGlobalFrame = true;
-	for (int32 FrameIndex = Exec.ActiveFrames.Num() - 1; FrameIndex >= 0; FrameIndex--)
-	{
-		const FStateTreeExecutionFrame* CurrentParentFrame = FrameIndex > 0 ? &Exec.ActiveFrames[FrameIndex - 1] : nullptr;
-		const FStateTreeExecutionFrame& CurrentFrame = Exec.ActiveFrames[FrameIndex];
-		if (CurrentFrame.bIsGlobalFrame)
-		{
-			// LastInitializedTaskIndex belongs to the last frame.
-			const FStateTreeIndex16 LastTaskToBeStopped = bIsLastGlobalFrame ? LastInitializedTaskIndex : FStateTreeIndex16::Invalid;
-			CallStopOnEvaluatorsAndGlobalTasks(CurrentParentFrame, CurrentFrame, Transition, LastTaskToBeStopped);
-			bIsLastGlobalFrame = false;
-		}
-	}
+	constexpr bool bOnActiveInstances = false;
+	return StartGlobalsForFrameInternal<bOnActiveInstances>(ParentFrame, Frame, Transition);
 }
 
-void FStateTreeExecutionContext::CallStopOnEvaluatorsAndGlobalTasks(const FStateTreeExecutionFrame* ParentFrame, const FStateTreeExecutionFrame& Frame, const FStateTreeTransitionResult& Transition, const FStateTreeIndex16 LastInitializedTaskIndex /*= FStateTreeIndex16()*/)
-{
-	check(Frame.bIsGlobalFrame);
-
-	ON_SCOPE_EXIT { GetExecState().LastExitedNodeIndex = FStateTreeIndex16::Invalid; };
-
-	FCurrentlyProcessedFrameScope FrameScope(*this, ParentFrame, Frame);
-	const UStateTree* CurrentStateTree = Frame.StateTree;
-
-	for (int32 TaskIndex = (CurrentStateTree->GlobalTasksBegin + CurrentStateTree->GlobalTasksNum) - 1;  TaskIndex >= CurrentStateTree->GlobalTasksBegin; TaskIndex--)
-	{
-		const FStateTreeTaskBase& Task =  CurrentStateTree->Nodes[TaskIndex].Get<const FStateTreeTaskBase>();
-		const FStateTreeDataView TaskInstanceView = GetDataView(ParentFrame, Frame, Task.InstanceDataHandle);
-		FNodeInstanceDataScope DataScope(*this, &Task, TaskIndex, Task.InstanceDataHandle, TaskInstanceView);
-
-		// Ignore disabled task
-		if (Task.bTaskEnabled == false)
-		{
-			STATETREE_LOG(VeryVerbose, TEXT("%*sSkipped 'ExitState' for disabled Task: '%s'"), UE::StateTree::Debug::IndentSize, TEXT(""), *Task.Name.ToString());
-			continue;
-		}
-
-		// Relying here that invalid value of LastInitializedTaskIndex == MAX_uint16.
-		if (TaskIndex <= LastInitializedTaskIndex.Get())
-		{
-			STATETREE_LOG(Verbose, TEXT("  Stop: '%s'"), *Task.Name.ToString());
-			UE_STATETREE_DEBUG_TASK_EXIT_STATE(this, CurrentStateTree, FStateTreeIndex16(TaskIndex));
-			{
-				QUICK_SCOPE_CYCLE_COUNTER(StateTree_Task_TreeStop);
-				Task.ExitState(*this, Transition);
-			}
-			UE_STATETREE_DEBUG_TASK_EVENT(this, TaskIndex, TaskInstanceView, EStateTreeTraceEventType::OnExited, Transition.CurrentRunStatus);
-			GetExecState().LastExitedNodeIndex = FStateTreeIndex16(TaskIndex);
-		}
-	}
-
-	for (int32 EvalIndex = (CurrentStateTree->EvaluatorsBegin + CurrentStateTree->EvaluatorsNum) - 1; EvalIndex >= CurrentStateTree->EvaluatorsBegin; EvalIndex--)
-	{
-		const FStateTreeEvaluatorBase& Eval = CurrentStateTree->Nodes[EvalIndex].Get<const FStateTreeEvaluatorBase>();
-		const FStateTreeDataView EvalInstanceView = GetDataView(ParentFrame, Frame, Eval.InstanceDataHandle);
-		FNodeInstanceDataScope DataScope(*this, &Eval, EvalIndex, Eval.InstanceDataHandle, EvalInstanceView);
-
-		STATETREE_LOG(Verbose, TEXT("  Stop: '%s'"), *Eval.Name.ToString());
-		UE_STATETREE_DEBUG_EVALUATOR_EXIT_TREE(this, CurrentStateTree, FStateTreeIndex16(EvalIndex));
-		{
-			QUICK_SCOPE_CYCLE_COUNTER(StateTree_Eval_TreeStop);
-			Eval.TreeStop(*this);
-
-			UE_STATETREE_DEBUG_EVALUATOR_EVENT(this, EvalIndex, EvalInstanceView, EStateTreeTraceEventType::OnTreeStopped);
-		}
-		GetExecState().LastExitedNodeIndex = FStateTreeIndex16(EvalIndex);
-	}
-}
-
-EStateTreeRunStatus FStateTreeExecutionContext::StartTemporaryEvaluatorsAndGlobalTasks(const FStateTreeExecutionFrame* CurrentParentFrame, FStateTreeExecutionFrame& CurrentFrame)
+template<bool bOnActiveInstances>
+EStateTreeRunStatus FStateTreeExecutionContext::StartGlobalsForFrameInternal(const FStateTreeExecutionFrame* CurrentParentFrame, FStateTreeExecutionFrame& CurrentFrame, FStateTreeTransitionResult& Transition)
 {
 	if (!CurrentFrame.bIsGlobalFrame)
 	{
-		return EStateTreeRunStatus::Failed;
+		return EStateTreeRunStatus::Running;
 	}
 
 	CSV_SCOPED_TIMING_STAT_EXCLUSIVE(StateTree_StartEvaluators);
-	STATETREE_LOG(Verbose, TEXT("Start Temporary Evaluators & Global tasks while trying to select linked asset: %s"), *GetNameSafe(CurrentFrame.StateTree));
 
 	FStateTreeExecutionState& Exec = GetExecState();
 
 	FCurrentlyProcessedFrameScope FrameScope(*this, CurrentParentFrame, CurrentFrame);
-	UE_STATETREE_DEBUG_SCOPED_PHASE(this, EStateTreeUpdatePhase::StartGlobalTasksForSelection);
+	if constexpr (!bOnActiveInstances)
+	{
+		UE_STATETREE_DEBUG_ENTER_PHASE(this, EStateTreeUpdatePhase::StartGlobalTasksForSelection);
+	}
 
 	EStateTreeRunStatus Result = EStateTreeRunStatus::Running;
 	const UE::StateTree::FActiveFrameID CurrentFrameID = CurrentFrame.FrameID;
@@ -3533,18 +3442,30 @@ EStateTreeRunStatus FStateTreeExecutionContext::StartTemporaryEvaluatorsAndGloba
 	for (int32 EvalIndex = CurrentStateTree->EvaluatorsBegin; EvalIndex < (CurrentStateTree->EvaluatorsBegin + CurrentStateTree->EvaluatorsNum); EvalIndex++)
 	{
 		const FStateTreeEvaluatorBase& Eval = CurrentStateTree->Nodes[EvalIndex].Get<const FStateTreeEvaluatorBase>();
-		FStateTreeDataView EvalInstanceView = GetDataViewOrTemporary(CurrentParentFrame, CurrentFrame, Eval.InstanceDataHandle);
-		if (!EvalInstanceView.IsValid())
+		FStateTreeDataView EvalInstanceView = bOnActiveInstances
+			? GetDataView(CurrentParentFrame, CurrentFrame, Eval.InstanceDataHandle)
+			: GetDataViewOrTemporary(CurrentParentFrame, CurrentFrame, Eval.InstanceDataHandle);
+		if constexpr (!bOnActiveInstances)
 		{
-			EvalInstanceView = AddTemporaryInstance(CurrentFrame, FStateTreeIndex16(EvalIndex), Eval.InstanceDataHandle, CurrentFrame.StateTree->DefaultInstanceData.GetStruct(Eval.InstanceTemplateIndex.Get()));
-			check(EvalInstanceView.IsValid());
+			if (!EvalInstanceView.IsValid())
+			{
+				EvalInstanceView = AddTemporaryInstance(CurrentFrame, FStateTreeIndex16(EvalIndex), Eval.InstanceDataHandle, CurrentFrame.StateTree->DefaultInstanceData.GetStruct(Eval.InstanceTemplateIndex.Get()));
+				check(EvalInstanceView.IsValid());
+			}
 		}
-		
+
 		FNodeInstanceDataScope DataScope(*this, &Eval, EvalIndex, Eval.InstanceDataHandle, EvalInstanceView);
 		// Copy bound properties with the temporary buffer.
 		if (Eval.BindingsBatch.IsValid())
 		{
-			CopyBatchWithValidation(CurrentParentFrame, CurrentFrame, EvalInstanceView, Eval.BindingsBatch);
+			if constexpr (bOnActiveInstances)
+			{
+				CopyBatchOnActiveInstances(CurrentParentFrame, CurrentFrame, EvalInstanceView, Eval.BindingsBatch);
+			}
+			else
+			{
+				CopyBatchWithValidation(CurrentParentFrame, CurrentFrame, EvalInstanceView, Eval.BindingsBatch);
+			}
 		}
 
 		STATETREE_LOG(Verbose, TEXT("  Start: '%s'"), *Eval.Name.ToString());
@@ -3559,12 +3480,12 @@ EStateTreeRunStatus FStateTreeExecutionContext::StartTemporaryEvaluatorsAndGloba
 
 	// Start Global tasks
 	// Even if we call Enter/ExitState() on global tasks, they do not enter any specific state.
-	const FStateTreeTransitionResult Transition = {}; // Empty transition
 
 	for (int32 GlobalTaskIndex = 0; GlobalTaskIndex < CurrentStateTree->GlobalTasksNum; ++GlobalTaskIndex)
 	{
 		const int32 AssetTaskIndex = CurrentStateTree->GlobalTasksBegin + GlobalTaskIndex;
-		const FStateTreeTaskBase& Task =  CurrentStateTree->Nodes[AssetTaskIndex].Get<const FStateTreeTaskBase>();
+		const FStateTreeTaskBase& Task = CurrentStateTree->Nodes[AssetTaskIndex].Get<const FStateTreeTaskBase>();
+
 		// Ignore disabled task
 		if (Task.bTaskEnabled == false)
 		{
@@ -3572,19 +3493,33 @@ EStateTreeRunStatus FStateTreeExecutionContext::StartTemporaryEvaluatorsAndGloba
 			continue;
 		}
 
-		FStateTreeDataView TaskDataView = GetDataViewOrTemporary(CurrentParentFrame, CurrentFrame, Task.InstanceDataHandle);
-		if (!TaskDataView.IsValid())
+		FStateTreeDataView TaskDataView = bOnActiveInstances
+			? GetDataView(CurrentParentFrame, CurrentFrame, Task.InstanceDataHandle)
+			: GetDataViewOrTemporary(CurrentParentFrame, CurrentFrame, Task.InstanceDataHandle);
+		if constexpr (!bOnActiveInstances)
 		{
-			TaskDataView = AddTemporaryInstance(CurrentFrame, FStateTreeIndex16(AssetTaskIndex), Task.InstanceDataHandle, CurrentFrame.StateTree->DefaultInstanceData.GetStruct(Task.InstanceTemplateIndex.Get()));
-			check(TaskDataView.IsValid())
+			if (!TaskDataView.IsValid())
+			{
+				TaskDataView = AddTemporaryInstance(CurrentFrame, FStateTreeIndex16(AssetTaskIndex), Task.InstanceDataHandle, CurrentFrame.StateTree->DefaultInstanceData.GetStruct(Task.InstanceTemplateIndex.Get()));
+				check(TaskDataView.IsValid())
+			}
 		}
 
 		FNodeInstanceDataScope DataScope(*this, &Task, AssetTaskIndex, Task.InstanceDataHandle, TaskDataView);
 
+		CurrentFrame.ActiveTaskNodeIndex = FStateTreeIndex16(AssetTaskIndex);
+
 		// Copy bound properties with the temporary buffer.
 		if (Task.BindingsBatch.IsValid())
 		{
-			CopyBatchWithValidation(CurrentParentFrame, CurrentFrame, TaskDataView, Task.BindingsBatch);
+			if constexpr (bOnActiveInstances)
+			{
+				CopyBatchOnActiveInstances(CurrentParentFrame, CurrentFrame, TaskDataView, Task.BindingsBatch);
+			}
+			else
+			{
+				CopyBatchWithValidation(CurrentParentFrame, CurrentFrame, TaskDataView, Task.BindingsBatch);
+			}
 		}
 
 		STATETREE_LOG(Verbose, TEXT("  Start: '%s'"), *Task.Name.ToString());
@@ -3602,8 +3537,50 @@ EStateTreeRunStatus FStateTreeExecutionContext::StartTemporaryEvaluatorsAndGloba
 		TaskRunStatus = UE::StateTree::ExecutionContext::CastToRunStatus(TaskStatus);
 		UE_STATETREE_DEBUG_TASK_EVENT(this, AssetTaskIndex, TaskDataView, EStateTreeTraceEventType::OnEntered, TaskRunStatus);
 
-		Result = UE::StateTree::ExecutionContext::GetPriorityRunStatus(Result, TaskRunStatus);
-		if (TaskRunStatus == EStateTreeRunStatus::Failed && CurrentTasksStatus.IsConsideredForCompletion(GlobalTaskIndex))
+		if (CurrentTasksStatus.IsConsideredForCompletion(GlobalTaskIndex))
+		{
+			Result = UE::StateTree::ExecutionContext::GetPriorityRunStatus(Result, TaskRunStatus);
+			if (Result == EStateTreeRunStatus::Failed)
+			{
+				return Result;
+			}
+		}
+	}
+
+	if constexpr (!bOnActiveInstances)
+	{
+		UE_STATETREE_DEBUG_EXIT_PHASE(this, EStateTreeUpdatePhase::StartGlobalTasksForSelection);
+	}
+
+	return Result;
+}
+
+// Deprecated
+EStateTreeRunStatus FStateTreeExecutionContext::StartEvaluatorsAndGlobalTasks(FStateTreeIndex16& OutLastInitializedTaskIndex)
+{
+	return StartEvaluatorsAndGlobalTasks();
+}
+
+EStateTreeRunStatus FStateTreeExecutionContext::StartEvaluatorsAndGlobalTasks()
+{
+	UE_STATETREE_DEBUG_SCOPED_PHASE(this, EStateTreeUpdatePhase::StartGlobalTasks);
+
+	STATETREE_LOG(Verbose, TEXT("Start Evaluators & Global tasks"));
+
+	FStateTreeExecutionState& Exec = GetExecState();
+
+	EStateTreeRunStatus Result = EStateTreeRunStatus::Running;
+	FStateTreeTransitionResult Transition{};
+	Transition.TargetState = FStateTreeStateHandle::Root;
+	Transition.CurrentRunStatus = EStateTreeRunStatus::Running;
+
+	for (int32 FrameIndex = 0; FrameIndex < Exec.ActiveFrames.Num(); FrameIndex++)
+	{
+		const FStateTreeExecutionFrame* CurrentParentFrame = FrameIndex > 0 ? &Exec.ActiveFrames[FrameIndex - 1] : nullptr;
+		FStateTreeExecutionFrame& CurrentFrame = Exec.ActiveFrames[FrameIndex];
+		EStateTreeRunStatus FrameResult = StartGlobalsForFrameOnActiveInstances(CurrentParentFrame, CurrentFrame, Transition);
+		Result = UE::StateTree::ExecutionContext::GetPriorityRunStatus(Result, FrameResult);
+		if (FrameResult == EStateTreeRunStatus::Failed)
 		{
 			break;
 		}
@@ -3612,77 +3589,161 @@ EStateTreeRunStatus FStateTreeExecutionContext::StartTemporaryEvaluatorsAndGloba
 	return Result;
 }
 
-void FStateTreeExecutionContext::StopTemporaryEvaluatorsAndGlobalTasks(const FStateTreeExecutionFrame* CurrentParentFrame, const FStateTreeExecutionFrame& CurrentFrame)
+EStateTreeRunStatus FStateTreeExecutionContext::StartTemporaryEvaluatorsAndGlobalTasks(const FStateTreeExecutionFrame* CurrentParentFrame, FStateTreeExecutionFrame& CurrentFrame)
 {
-	//@todo only stop the evaluators and tasks that were started in StartTemporaryEvaluatorsAndGlobalTasks
-	STATETREE_LOG(Verbose, TEXT("Stop Temporary Evaluators & Global tasks"));
+	STATETREE_LOG(Verbose, TEXT("Start Temporary Evaluators & Global tasks while trying to select linked asset: %s"), *GetNameSafe(CurrentFrame.StateTree));
 
-	FCurrentlyProcessedFrameScope FrameScope(*this, CurrentParentFrame, CurrentFrame);
-	UE_STATETREE_DEBUG_SCOPED_PHASE(this, EStateTreeUpdatePhase::StopGlobalTasksForSelection);
+	FStateTreeTransitionResult Transition;
+	Transition.ChangeType = EStateTreeStateChangeType::Changed;
+	Transition.CurrentRunStatus = EStateTreeRunStatus::Running;
+	return StartGlobalsForFrameWithValidation(CurrentParentFrame, CurrentFrame, Transition);
+}
 
-	// Create temporary transition to stop the unused global tasks and evaluators.
-	constexpr  EStateTreeRunStatus CompletionStatus = EStateTreeRunStatus::Stopped; 
+// Deprecated
+void FStateTreeExecutionContext::StopEvaluatorsAndGlobalTasks(const EStateTreeRunStatus CompletionStatus, const FStateTreeIndex16 LastInitializedTaskIndex)
+{
+	StopEvaluatorsAndGlobalTasks(CompletionStatus);
+}
+
+void FStateTreeExecutionContext::StopEvaluatorsAndGlobalTasks(const EStateTreeRunStatus CompletionStatus)
+{
+	CSV_SCOPED_TIMING_STAT_EXCLUSIVE(StateTree_StopEvaluators);
+	UE_STATETREE_DEBUG_SCOPED_PHASE(this, EStateTreeUpdatePhase::StopGlobalTasks);
+
+	STATETREE_LOG(Verbose, TEXT("Stop Evaluators & Global Tasks"));
+
+	FStateTreeExecutionState& Exec = GetExecState();
+
+	// Update bindings
+	CopyAllBindingsOnActiveInstances(ECopyBindings::ExitState);
+
+	// Call in reverse order.
 	FStateTreeTransitionResult Transition;
 	Transition.TargetState = FStateTreeStateHandle::FromCompletionStatus(CompletionStatus);
 	Transition.CurrentRunStatus = CompletionStatus;
 
-	TArrayView<FStateTreeTemporaryInstanceData> TempInstances = Storage.GetMutableTemporaryInstances();
-	for (int32 Index = TempInstances.Num() - 1; Index >= 0; Index--)
+	for (int32 FrameIndex = Exec.ActiveFrames.Num() - 1; FrameIndex >= 0; FrameIndex--)
 	{
-		FStateTreeTemporaryInstanceData& TempInstance = TempInstances[Index];
-		if (TempInstance.FrameID != CurrentFrame.FrameID)
+		const FStateTreeExecutionFrame* CurrentParentFrame = FrameIndex > 0 ? &Exec.ActiveFrames[FrameIndex - 1] : nullptr;
+		FStateTreeExecutionFrame& CurrentFrame = Exec.ActiveFrames[FrameIndex];
+		if (CurrentFrame.bIsGlobalFrame)
 		{
-			continue;
+			StopGlobalsForFrameOnActiveInstances(CurrentParentFrame, CurrentFrame, Transition);
 		}
+	}
+}
+
+void FStateTreeExecutionContext::StopGlobalsForFrameOnActiveInstances(const FStateTreeExecutionFrame* ParentFrame, FStateTreeExecutionFrame& Frame, const FStateTreeTransitionResult& Transition)
+{
+	constexpr bool bOnActiveInstances = true;
+	StopGlobalsForFrameInternal<bOnActiveInstances>(ParentFrame, Frame, Transition);
+}
 
-		if (TempInstance.OwnerNodeIndex.IsValid()
-			&& TempInstance.Instance.IsValid())
+void FStateTreeExecutionContext::StopGlobalsForFrameWithValidation(const FStateTreeExecutionFrame* ParentFrame, FStateTreeExecutionFrame& Frame, const FStateTreeTransitionResult& Transition)
+{
+	constexpr bool bOnActiveInstances = false;
+	StopGlobalsForFrameInternal<bOnActiveInstances>(ParentFrame, Frame, Transition);
+}
+
+template<bool bOnActiveInstances>
+void FStateTreeExecutionContext::StopGlobalsForFrameInternal(const FStateTreeExecutionFrame* ParentFrame, FStateTreeExecutionFrame& Frame, const FStateTreeTransitionResult& Transition)
+{
+	if (!Frame.bIsGlobalFrame)
+	{
+		return;
+	}
+
+	FCurrentlyProcessedFrameScope FrameScope(*this, ParentFrame, Frame);
+	if constexpr (bOnActiveInstances)
+	{
+		UE_STATETREE_DEBUG_ENTER_PHASE(this, EStateTreeUpdatePhase::StopGlobalTasksForSelection);
+	}
+
+	const UStateTree* CurrentStateTree = Frame.StateTree;
+
+	const int32 GlobalTasksEnd = CurrentStateTree->GlobalTasksBegin + CurrentStateTree->GlobalTasksNum;
+	for (int32 TaskIndex = GlobalTasksEnd - 1; TaskIndex >= CurrentStateTree->GlobalTasksBegin; --TaskIndex)
+	{
+		const FStateTreeTaskBase& Task = CurrentStateTree->Nodes[TaskIndex].Get<const FStateTreeTaskBase>();
+		const FStateTreeDataView TaskInstanceView = bOnActiveInstances
+			? GetDataView(ParentFrame, Frame, Task.InstanceDataHandle)
+			: GetDataViewOrTemporary(ParentFrame, Frame, Task.InstanceDataHandle);
+		FNodeInstanceDataScope DataScope(*this, &Task, TaskIndex, Task.InstanceDataHandle, TaskInstanceView);
+
+		if (TaskIndex <= Frame.ActiveTaskNodeIndex.AsInt32())
 		{
-			FStateTreeDataView NodeInstanceView;
-			if (FStateTreeInstanceObjectWrapper* Wrapper = TempInstance.Instance.GetMutablePtr<FStateTreeInstanceObjectWrapper>())
+			// Ignore disabled task
+			if (Task.bTaskEnabled == false)
 			{
-				NodeInstanceView = FStateTreeDataView(Wrapper->InstanceObject);
+				STATETREE_LOG(VeryVerbose, TEXT("%*sSkipped 'ExitState' for disabled Task: '%s'"), UE::StateTree::Debug::IndentSize, TEXT(""), *Task.Name.ToString());
 			}
 			else
 			{
-				NodeInstanceView = FStateTreeDataView(TempInstance.Instance);
-			}
-
-
-			FConstStructView NodeView = CurrentFrame.StateTree->Nodes[TempInstance.OwnerNodeIndex.Get()];
-			if (const FStateTreeTaskBase* Task = NodeView.GetPtr<const FStateTreeTaskBase>())
-			{
-				FNodeInstanceDataScope DataScope(*this, Task, TempInstance.OwnerNodeIndex.Get(), TempInstance.DataHandle, NodeInstanceView);
-
-				STATETREE_LOG(Verbose, TEXT("  Stop: '%s'"), *Task->Name.ToString());
-				UE_STATETREE_DEBUG_TASK_EXIT_STATE(this, CurrentFrame.StateTree, FStateTreeIndex16(TempInstance.OwnerNodeIndex.Get()));
+				STATETREE_LOG(Verbose, TEXT("  Stop: '%s'"), *Task.Name.ToString());
+				UE_STATETREE_DEBUG_TASK_EXIT_STATE(this, CurrentStateTree, FStateTreeIndex16(TaskIndex));
 				{
 					QUICK_SCOPE_CYCLE_COUNTER(StateTree_Task_TreeStop);
-					Task->ExitState(*this, Transition);
+					Task.ExitState(*this, Transition);
 				}
-				UE_STATETREE_DEBUG_TASK_EVENT(this, TempInstance.OwnerNodeIndex.Get(), NodeInstanceView, EStateTreeTraceEventType::OnExited, Transition.CurrentRunStatus);
+				UE_STATETREE_DEBUG_TASK_EVENT(this, TaskIndex, TaskInstanceView, EStateTreeTraceEventType::OnExited, Transition.CurrentRunStatus);
 			}
-			else if (const FStateTreeEvaluatorBase* Eval = NodeView.GetPtr<const FStateTreeEvaluatorBase>())
-			{
-				FNodeInstanceDataScope DataScope(*this, Eval, TempInstance.OwnerNodeIndex.Get(), TempInstance.DataHandle, NodeInstanceView);
+		}
+		Frame.ActiveTaskNodeIndex = FStateTreeIndex16(TaskIndex - 1);
+	}
+	Frame.ActiveTaskNodeIndex = FStateTreeIndex16::Invalid;
 
-				STATETREE_LOG(Verbose, TEXT("  Stop: '%s'"), *Eval->Name.ToString());
-				UE_STATETREE_DEBUG_EVALUATOR_EXIT_TREE(this, CurrentFrame.StateTree, FStateTreeIndex16(TempInstance.OwnerNodeIndex.Get()));
-				{
-					QUICK_SCOPE_CYCLE_COUNTER(StateTree_Eval_TreeStop);
-					Eval->TreeStop(*this);
+	const int32 EvaluatorsEnd = CurrentStateTree->EvaluatorsBegin + CurrentStateTree->EvaluatorsNum;
+	for (int32 EvalIndex = EvaluatorsEnd - 1; EvalIndex >= CurrentStateTree->EvaluatorsBegin; --EvalIndex)
+	{
+		const FStateTreeEvaluatorBase& Eval = CurrentStateTree->Nodes[EvalIndex].Get<const FStateTreeEvaluatorBase>();
+		const FStateTreeDataView EvalInstanceView = GetDataView(ParentFrame, Frame, Eval.InstanceDataHandle);
+		FNodeInstanceDataScope DataScope(*this, &Eval, EvalIndex, Eval.InstanceDataHandle, EvalInstanceView);
 
-					UE_STATETREE_DEBUG_EVALUATOR_EVENT(this, TempInstance.OwnerNodeIndex.Get(), NodeInstanceView, EStateTreeTraceEventType::OnTreeStopped);
-				}
-			}
+		STATETREE_LOG(Verbose, TEXT("  Stop: '%s'"), *Eval.Name.ToString());
+		UE_STATETREE_DEBUG_EVALUATOR_EXIT_TREE(this, CurrentStateTree, FStateTreeIndex16(EvalIndex));
+		{
+			QUICK_SCOPE_CYCLE_COUNTER(StateTree_Eval_TreeStop);
+			Eval.TreeStop(*this);
+
+			UE_STATETREE_DEBUG_EVALUATOR_EVENT(this, EvalIndex, EvalInstanceView, EStateTreeTraceEventType::OnTreeStopped);
 		}
 	}
+
+	if constexpr (bOnActiveInstances)
+	{
+		UE_STATETREE_DEBUG_EXIT_PHASE(this, EStateTreeUpdatePhase::StopGlobalTasksForSelection);
+	}
+}
+
+// Deprecated
+void FStateTreeExecutionContext::CallStopOnEvaluatorsAndGlobalTasks(const FStateTreeExecutionFrame* ParentFrame, const FStateTreeExecutionFrame& Frame, const FStateTreeTransitionResult& Transition, const FStateTreeIndex16 LastInitializedTaskIndex /*= FStateTreeIndex16()*/)
+{
+	constexpr bool bOnActiveInstances = true;
+	StopGlobalsForFrameInternal<bOnActiveInstances>(ParentFrame, const_cast<FStateTreeExecutionFrame&>(Frame), Transition);
+}
+
+// Deprecated
+void FStateTreeExecutionContext::StopTemporaryEvaluatorsAndGlobalTasks(const FStateTreeExecutionFrame* CurrentParentFrame, const FStateTreeExecutionFrame& CurrentFrame)
+{
+	StopTemporaryEvaluatorsAndGlobalTasks(CurrentParentFrame, const_cast<FStateTreeExecutionFrame&>(CurrentFrame), EStateTreeRunStatus::Running);
+}
+
+void FStateTreeExecutionContext::StopTemporaryEvaluatorsAndGlobalTasks(const FStateTreeExecutionFrame* CurrentParentFrame, FStateTreeExecutionFrame& CurrentFrame, EStateTreeRunStatus StartResult)
+{
+	STATETREE_LOG(Verbose, TEXT("Stop Temporary Evaluators & Global tasks"));
+
+	// Create temporary transition to stop the unused global tasks and evaluators.
+	FStateTreeTransitionResult Transition;
+	Transition.TargetState = FStateTreeStateHandle::FromCompletionStatus(StartResult);
+	Transition.CurrentRunStatus = StartResult;
+	Transition.ChangeType = EStateTreeStateChangeType::Changed;
+	StopGlobalsForFrameWithValidation(CurrentParentFrame, CurrentFrame, Transition);
 }
 
 EStateTreeRunStatus FStateTreeExecutionContext::TickTasks(const float DeltaTime)
 {
 	// When a task is completed it also completes the state and triggers the completion transition (because LastTickStatus is set).
-	// A task can complete async. See CompletedStates.
+	// A task can complete async.
 	// When a task fails, stop ticking the following tasks.
 	// When no task ticks, then the leaf completes.
 
@@ -3868,6 +3929,7 @@ FStateTreeExecutionContext::FTickTaskResult FStateTreeExecutionContext::TickTask
 	const bool bCopyBoundPropertiesOnNonTickedTask = ExecutionContext::Private::bCopyBoundPropertiesOnNonTickedTask;
 	const UStateTree* CurrentStateTree = Args.Frame->StateTree;
 	const FActiveFrameID CurrentFrameID = Args.Frame->FrameID;
+	const int32 CurrentActiveTaskNodeIndex = Args.Frame->ActiveTaskNodeIndex.AsInt32();
 	check(CurrentStateTree);
 
 	for (int32 OwnerTaskIndex = 0; OwnerTaskIndex < Args.TasksNum; ++OwnerTaskIndex)
@@ -3882,6 +3944,13 @@ FStateTreeExecutionContext::FTickTaskResult FStateTreeExecutionContext::TickTask
 			continue;
 		}
 
+		if (AssetTaskIndex > CurrentActiveTaskNodeIndex)
+		{
+			STATETREE_LOG(VeryVerbose, TEXT("%*sSkipped 'Tick' for task that didn't get the EnterState Task: '%s'"), Debug::IndentSize, TEXT(""), *Task.Name.ToString());
+			bShouldTickTasks = false;
+			break;
+		}
+
 		const FStateTreeDataView TaskInstanceView = GetDataView(Args.ParentFrame, *Args.Frame, Task.InstanceDataHandle);
 		FNodeInstanceDataScope DataScope(*this, &Task, AssetTaskIndex, Task.InstanceDataHandle, TaskInstanceView);
 
@@ -4473,6 +4542,7 @@ bool FStateTreeExecutionContext::TriggerTransitions()
 		{
 			FStateTreeExecutionFrame& CurrentFrame = Exec.ActiveFrames[FrameIndex];
 			const UStateTree* CurrentStateTree = CurrentFrame.StateTree;
+			const int32 CurrentActiveTaskNodeIndex = CurrentFrame.ActiveTaskNodeIndex.AsInt32();
 
 			for (int32 StateIndex = CurrentFrame.ActiveStates.Num() - 1; StateIndex >= 0; StateIndex--)
 			{
@@ -4493,7 +4563,7 @@ bool FStateTreeExecutionContext::TriggerTransitions()
 					for (int32 TaskIndex = (State.TasksBegin + State.TasksNum) - 1; TaskIndex >= State.TasksBegin; TaskIndex--)
 					{
 						const FStateTreeTaskBase& Task = CurrentStateTree->Nodes[TaskIndex].Get<const FStateTreeTaskBase>();
-						if (Task.bShouldAffectTransitions && Task.bTaskEnabled)
+						if (Task.bShouldAffectTransitions && Task.bTaskEnabled && TaskIndex <= CurrentActiveTaskNodeIndex)
 						{
 							TransitionHandlers.Emplace(FrameIndex, StateHandle, StateID, FStateTreeIndex16(TaskIndex), Task.TransitionHandlingPriority);
 							bAdded = true;
@@ -5534,6 +5604,7 @@ bool FStateTreeExecutionContext::SelectStateInternal(
 				{
 					NewFrame.FrameID = ExistingFrame->FrameID;
 					NewFrame.ActiveTasksStatus = ExistingFrame->ActiveTasksStatus;
+					NewFrame.ActiveTaskNodeIndex = ExistingFrame->ActiveTaskNodeIndex;
 					NewFrame.ActiveInstanceIndexBase = ExistingFrame->ActiveInstanceIndexBase;
 					NewFrame.GlobalInstanceIndexBase = ExistingFrame->GlobalInstanceIndexBase;
 					NewFrame.StateParameterDataHandle = ExistingFrame->StateParameterDataHandle;
@@ -5622,6 +5693,7 @@ bool FStateTreeExecutionContext::SelectStateInternal(
 			{
 				NewFrame.FrameID = ExistingFrame->FrameID;
 				NewFrame.ActiveTasksStatus = ExistingFrame->ActiveTasksStatus;
+				NewFrame.ActiveTaskNodeIndex = ExistingFrame->ActiveTaskNodeIndex;
 				NewFrame.ActiveInstanceIndexBase = ExistingFrame->ActiveInstanceIndexBase;
 				NewFrame.GlobalInstanceIndexBase = ExistingFrame->GlobalInstanceIndexBase;
 				NewFrame.StateParameterDataHandle = ExistingFrame->StateParameterDataHandle;
@@ -5652,12 +5724,13 @@ bool FStateTreeExecutionContext::SelectStateInternal(
 				NewFrame.StateParameterDataHandle = RootState.ParameterDataHandle;
 
 				// Start global tasks and evaluators temporarily, so that their data is available already during select.
-				if (StartTemporaryEvaluatorsAndGlobalTasks(&CurrentFrame, NewFrame) != EStateTreeRunStatus::Running)
+				const EStateTreeRunStatus StartResult = StartTemporaryEvaluatorsAndGlobalTasks(&CurrentFrame, NewFrame);
+				if (StartResult != EStateTreeRunStatus::Running)
 				{
 					STATETREE_LOG(VeryVerbose, TEXT("%hs: Cannot select state '%s' because cannot start nested tree's '%s' global tasks and evaluators.  '%s' using StateTree '%s'."),
 						__FUNCTION__, *GetSafeStateName(CurrentFrame, NextStateHandle), *GetFullNameSafe(NewFrame.StateTree), *GetNameSafe(&Owner), *GetFullNameSafe(CurrentFrame.StateTree));
 					
-					StopTemporaryEvaluatorsAndGlobalTasks(nullptr, NewFrame);
+					StopTemporaryEvaluatorsAndGlobalTasks(nullptr, NewFrame, StartResult);
 					GetExecState().DelegateActiveListeners.RemoveAll(NewFrame.FrameID);
 
 					break;
@@ -5677,7 +5750,7 @@ bool FStateTreeExecutionContext::SelectStateInternal(
 
 			if (bStartedTemporaryEvaluatorsAndGlobalTasks)
 			{
-				StopTemporaryEvaluatorsAndGlobalTasks(&CurrentFrame, NewFrame);
+				StopTemporaryEvaluatorsAndGlobalTasks(&CurrentFrame, NewFrame, EStateTreeRunStatus::Running);
 				GetExecState().DelegateActiveListeners.RemoveAll(NewFrame.FrameID);
 			}
 				
diff --git a/Engine/Plugins/Runtime/StateTree/Source/StateTreeModule/Private/StateTreeExecutionTypes.cpp b/Engine/Plugins/Runtime/StateTree/Source/StateTreeModule/Private/StateTreeExecutionTypes.cpp
--- a/Engine/Plugins/Runtime/StateTree/Source/StateTreeModule/Private/StateTreeExecutionTypes.cpp
+++ b/Engine/Plugins/Runtime/StateTree/Source/StateTreeModule/Private/StateTreeExecutionTypes.cpp
@@ -354,15 +354,26 @@ void FStateTreeDelegateActiveListeners::BroadcastDelegate(const FStateTreeDelega
 	for (int32 Index = 0; Index < NumListeners; ++Index)
 	{
 		FActiveListener& ActiveListener = Listeners[Index];
-		if (ActiveListener.Listener.GetDispatcher() == Dispatcher
-			&& ActiveListener.OwningNodeIndex.Get() <= Exec.EnterStateFailedTaskIndex.Get()
-			&& ActiveListener.OwningNodeIndex.Get() < Exec.LastExitedNodeIndex.Get())
+		if (ActiveListener.Listener.GetDispatcher() == Dispatcher)
 		{
 			if (const FStateTreeExecutionFrame* ExectionFrame = Exec.FindActiveFrame(ActiveListener.FrameID))
 			{
-				if (!ActiveListener.StateID.IsValid() || ExectionFrame->ActiveStates.Contains(ActiveListener.StateID))
+				// Is the node active and is the state active.
+				if (ActiveListener.OwningNodeIndex.Get() <= ExectionFrame->ActiveTaskNodeIndex.Get())
 				{
-					ActiveListener.Delegate.ExecuteIfBound();
+					if (!ActiveListener.StateID.IsValid())
+					{
+						// It's a global task, no need to check for the state.
+						ActiveListener.Delegate.ExecuteIfBound();
+					}
+					else
+					{
+						const int32 FoundStateIndex = ExectionFrame->ActiveStates.IndexOfReverse(ActiveListener.StateID);
+						if (FoundStateIndex != INDEX_NONE && FoundStateIndex < ExectionFrame->NumCurrentlyActiveStates)
+						{
+							ActiveListener.Delegate.ExecuteIfBound();
+						}
+					}
 				}
 			}
 		}
diff --git a/Engine/Plugins/Runtime/StateTree/Source/StateTreeModule/Public/StateTreeExecutionContext.h b/Engine/Plugins/Runtime/StateTree/Source/StateTreeModule/Public/StateTreeExecutionContext.h
--- a/Engine/Plugins/Runtime/StateTree/Source/StateTreeModule/Public/StateTreeExecutionContext.h
+++ b/Engine/Plugins/Runtime/StateTree/Source/StateTreeModule/Public/StateTreeExecutionContext.h
@@ -39,6 +39,7 @@ namespace UE::StateTree::ExecutionContext
 	UE_API UE::StateTree::ETaskCompletionStatus CastToTaskStatus(EStateTreeRunStatus InStatus);
 	UE_API EStateTreeRunStatus CastToRunStatus(UE::StateTree::ETaskCompletionStatus InStatus);
 }
+
 /**
  * Delegate used by the execution context to collect external data views for a given StateTree asset.
  * The caller is expected to iterate over the ExternalDataDescs array, find the matching external data,
@@ -56,7 +57,7 @@ DECLARE_DELEGATE_RetVal_FourParams(bool, FOnCollectStateTreeExternalData, const
 /**
  * Read-only execution context to interact with the state tree instance data. Only const and read accesses are available.
  * Multiple FStateTreeReadOnlyExecutionContext can coexist on different threads as long no other (minimal, weak, regular) execution context exists.
- * The user is responsible for preventing invalid multithread access.
+ * The user is responsible for preventing invalid multi-threaded access.
  */
 struct FStateTreeReadOnlyExecutionContext
 {
@@ -831,30 +832,85 @@ protected:
 	 */
 	UE_API void StateCompleted();
 
+private:
+	template<bool bOnActiveInstances>
+	void TickGlobalEvaluatorsForFrameInternal(float DeltaTime, const FStateTreeExecutionFrame* ParentFrame, const FStateTreeExecutionFrame& Frame);
+
+protected:
+	/**
+	 * Stop the frame's evaluators.
+	 * Should be used only on active node instances, assumes valid data handles and does not consider temporary node instances.
+	 */
+	UE_API void TickGlobalEvaluatorsForFrameOnActiveInstances(float DeltaTime, const FStateTreeExecutionFrame* ParentFrame, const FStateTreeExecutionFrame& Frame);
+
+	/**
+	 * Tick the frame's evaluators.
+	 * This version validates the data handles and looks up temporary instances.
+	 */
+	UE_API void TickGlobalEvaluatorsForFrameWithValidation(float DeltaTime, const FStateTreeExecutionFrame* ParentFrame, const FStateTreeExecutionFrame& Frame);
+
 	/**
 	 * Tick evaluators and global tasks by delta time.
 	 */
-	UE_API EStateTreeRunStatus TickEvaluatorsAndGlobalTasks(const float DeltaTime, bool bTickGlobalTasks = true);
+	UE_API EStateTreeRunStatus TickEvaluatorsAndGlobalTasks(float DeltaTime, bool bTickGlobalTasks = true);
 
 	/**
-	 * Tick evaluators and global tasks of a specific frame.
+	 * Tick the frame's evaluators and the frame's tasks.
 	 */
-	UE_API EStateTreeRunStatus TickEvaluatorsAndGlobalTasksForFrame(const float DeltaTime, const bool bTickGlobalTasks, const int32 FrameIndex, const FStateTreeExecutionFrame* CurrentParentFrame, const TNotNull<FStateTreeExecutionFrame*> CurrentFrame);
+	UE_API EStateTreeRunStatus TickEvaluatorsAndGlobalTasksForFrame(float DeltaTime, bool bTickGlobalTasks, int32 FrameIndex, const FStateTreeExecutionFrame* CurrentParentFrame, const TNotNull<FStateTreeExecutionFrame*> CurrentFrame);
 
+private:
+	template<bool bOnActiveInstances>
+	EStateTreeRunStatus StartGlobalsForFrameInternal(const FStateTreeExecutionFrame* ParentFrame, FStateTreeExecutionFrame& Frame, FStateTreeTransitionResult& Transition);
+
+protected:
 	/**
-	 * Starts evaluators and global tasks.
-	 * @return run status returned by the global tasks.
+	 * Start the frame's evaluators and the frame's tasks.
+	 * Should be used only on active node instances, assumes valid data handles and does not consider temporary node instances.
 	 */
+	UE_API EStateTreeRunStatus StartGlobalsForFrameOnActiveInstances(const FStateTreeExecutionFrame* ParentFrame, FStateTreeExecutionFrame& Frame, FStateTreeTransitionResult& Transition);
+
+	/**
+	 * Start the frame's evaluators and the frame's tasks.
+	 * This version validates the data handles and looks up temporary instances.
+	 */
+	UE_API EStateTreeRunStatus StartGlobalsForFrameWithValidation(const FStateTreeExecutionFrame* ParentFrame, FStateTreeExecutionFrame& Frame, FStateTreeTransitionResult& Transition);
+
+	UE_DEPRECATED(5.7, "Use StartEvaluatorsAndGlobalTasks without the OutLastInitializedTaskIndex argument")
 	UE_API EStateTreeRunStatus StartEvaluatorsAndGlobalTasks(FStateTreeIndex16& OutLastInitializedTaskIndex);
 
 	/**
-	 * Stops evaluators and global tasks.
+	 * Starts active global evaluators and global tasks.
+	 * @return run status returned by the global tasks.
+	 */
+	UE_API EStateTreeRunStatus StartEvaluatorsAndGlobalTasks();
+
+private:
+	template<bool bOnActiveInstances>
+	void StopGlobalsForFrameInternal(const FStateTreeExecutionFrame* ParentFrame, FStateTreeExecutionFrame& Frame, const FStateTreeTransitionResult& Transition);
+
+protected:
+	/**
+	 * Stop the frame's evaluators and the frame's tasks.
+	 * Should be used only on active node instances, assumes valid data handles and does not consider temporary node instances.
+	 */
+	UE_API void StopGlobalsForFrameOnActiveInstances(const FStateTreeExecutionFrame* ParentFrame, FStateTreeExecutionFrame& Frame, const FStateTreeTransitionResult& Transition);
+
+	/**
+	 * Stop the frame's evaluators and the frame's tasks.
+	 * This version validates the data handles and looks up temporary instances.
 	 */
-	UE_API void StopEvaluatorsAndGlobalTasks(const EStateTreeRunStatus CompletionStatus, const FStateTreeIndex16 LastInitializedTaskIndex = FStateTreeIndex16());
+	UE_API void StopGlobalsForFrameWithValidation(const FStateTreeExecutionFrame* ParentFrame, FStateTreeExecutionFrame& Frame, const FStateTreeTransitionResult& Transition);
+
+	UE_DEPRECATED(5.7, "Use StopEvaluatorsAndGlobalTasks without LastInitializedTaskIndex")
+	UE_API void StopEvaluatorsAndGlobalTasks(const EStateTreeRunStatus CompletionStatus, const FStateTreeIndex16 LastInitializedTaskIndex);
 
 	/**
-	 * Stops evaluators and global tasks of the given frame. Expects nodes data to be already bound.
+	 * Stops active global evaluators and active global tasks.
 	 */
+	UE_API void StopEvaluatorsAndGlobalTasks(const EStateTreeRunStatus CompletionStatus);
+
+	UE_DEPRECATED(5.7, "Use StopGlobalsForFrameOnActiveInstances")
 	UE_API void CallStopOnEvaluatorsAndGlobalTasks(const FStateTreeExecutionFrame* ParentFrame, const FStateTreeExecutionFrame& Frame, const FStateTreeTransitionResult& Transition, const FStateTreeIndex16 LastInitializedTaskIndex = FStateTreeIndex16());
 
 	/** Starts temporary instances of global evaluators and tasks for a given frame. */
@@ -866,9 +922,12 @@ protected:
 		return StartTemporaryEvaluatorsAndGlobalTasks(CurrentParentFrame, const_cast<FStateTreeExecutionFrame&>(CurrentFrame));
 	}
 
-	/** Stops temporary global evaluators and tasks for the provided frame. */
+	UE_DEPRECATED(5.7, "Use the non const version of StopTemporaryEvaluatorsAndGlobalTasks")
 	UE_API void StopTemporaryEvaluatorsAndGlobalTasks(const FStateTreeExecutionFrame* CurrentParentFrame, const FStateTreeExecutionFrame& CurrentFrame);
 
+	/** Stops temporary frame's evaluators and frame's tasks for the provided frame. */
+	UE_API void StopTemporaryEvaluatorsAndGlobalTasks(const FStateTreeExecutionFrame* CurrentParentFrame, FStateTreeExecutionFrame& CurrentFrame, EStateTreeRunStatus StartResult);
+
 	/**
 	 * Ticks tasks of all active states starting from current state by delta time.
 	 * @return Run status returned by the tasks.
@@ -913,7 +972,7 @@ private:
 protected:
 	/**
 	 * Checks all conditions at given range.
-	 * Should be used only on active instances, assumes valid handles and does not consider temporary instances.
+	 * Should be used only on active node instances, assumes valid data handles and does not consider temporary node instances.
 	 * @return True if all conditions pass.
 	 */
 	UE_API bool TestAllConditionsOnActiveInstances(const FStateTreeExecutionFrame* CurrentParentFrame, const FStateTreeExecutionFrame& CurrentFrame, FStateTreeStateHandle CurrentStateHandle, const int32 ConditionsOffset, const int32 ConditionsNum);
@@ -937,10 +996,16 @@ protected:
 	UE_DEPRECATED(5.7, "Use EvaluateUtilityWithValidation. It can get the data from a temporary frame in state selection.")
 	UE_API float EvaluateUtility(const FStateTreeExecutionFrame* CurrentParentFrame, const FStateTreeExecutionFrame& CurrentFrame, const int32 ConsiderationsOffset, const int32 ConsiderationsNum, const float StateWeight);
 
-	/* Evaluate all function at given range. Should be used only on active instances, assumes valid handles and does not consider temporary instances. */
+	/*
+	 * Evaluate all function at given range.
+	 * Should be used only on active node instances, assumes valid data handles and does not consider temporary node instances.
+	 */
 	UE_API void EvaluatePropertyFunctionsOnActiveInstances(const FStateTreeExecutionFrame* CurrentParentFrame, const FStateTreeExecutionFrame& CurrentFrame, FStateTreeIndex16 FuncsBegin, uint16 FuncsNum);
 
-	/* Evaluate all function at given range. This version validates the data handles and looks up temporary instances. */
+	/*
+	 * Evaluate all function at given range.
+	 * This version validates the data handles and looks up temporary instances.
+	 */
 	UE_API void EvaluatePropertyFunctionsWithValidation(const FStateTreeExecutionFrame* CurrentParentFrame, const FStateTreeExecutionFrame& CurrentFrame, FStateTreeIndex16 FuncsBegin, uint16 FuncsNum);
 
 	/**
@@ -1077,7 +1142,20 @@ PRAGMA_ENABLE_DEPRECATION_WARNINGS
 	 */
 	UE_API FStateTreeDataView AddTemporaryInstance(const FStateTreeExecutionFrame& Frame, const FStateTreeIndex16 OwnerNodeIndex, const FStateTreeDataHandle DataHandle, FConstStructView NewInstanceData);
 
-	/** Copies a batch of properties to the data in TargetView. Should be used only on active instances, assumes valid handles and does not consider temporary instances. */
+	enum class ECopyBindings : uint8
+	{
+		EnterState,
+		Tick,
+		ExitState,
+	};
+
+	/** Copy the binding on all valid active instances. */
+	UE_API void CopyAllBindingsOnActiveInstances(ECopyBindings CopyType);
+
+	/**
+	 * Copies a batch of properties to the data in TargetView.
+	 * Should be used only on active node instances, assumes valid data handles and does not consider temporary node instances.
+	 */
 	UE_API bool CopyBatchOnActiveInstances(const FStateTreeExecutionFrame* ParentFrame, const FStateTreeExecutionFrame& CurrentFrame, const FStateTreeDataView TargetView, const FStateTreeIndex16 BindingsBatch);
 
 	/** Copies a batch of properties to the data in TargetView. This version validates the data handles and looks up temporary instances. */
diff --git a/Engine/Plugins/Runtime/StateTree/Source/StateTreeModule/Public/StateTreeExecutionTypes.h b/Engine/Plugins/Runtime/StateTree/Source/StateTreeModule/Public/StateTreeExecutionTypes.h
--- a/Engine/Plugins/Runtime/StateTree/Source/StateTreeModule/Public/StateTreeExecutionTypes.h
+++ b/Engine/Plugins/Runtime/StateTree/Source/StateTreeModule/Public/StateTreeExecutionTypes.h
@@ -939,6 +939,13 @@ struct FStateTreeExecutionFrame
 	/** Unique frame id for this frame. Can be used to identify the frame. */
 	UE::StateTree::FActiveFrameID FrameID;
 
+	/**
+	 * The task node index that was "entered".
+	 * Used during the Enter and Exit phase. A node can fail EnterState.
+	 * Nodes after ActiveNodeIndex do not receive ExitState, because they didn't receive EnterState.
+	 */
+	FStateTreeIndex16 ActiveTaskNodeIndex = FStateTreeIndex16::Invalid;
+
 	/** First index of the external data for this frame. */
 	UPROPERTY()
 	FStateTreeIndex16 ExternalDataBaseIndex = FStateTreeIndex16::Invalid;
@@ -1083,18 +1090,6 @@ public:
 	UPROPERTY(Transient)
 	TInstancedStruct<FStateTreeExecutionExtension> ExecutionExtension;
 
-	/** The index of the task that failed during enter state. Exit state uses it to call ExitState() symmetrically. */
-	UPROPERTY()
-	FStateTreeIndex16 EnterStateFailedFrameIndex = FStateTreeIndex16::Invalid;
-
-	/** The index of the frame that failed during enter state. Exit state uses it to call ExitState() symmetrically. */
-	UPROPERTY()
-	FStateTreeIndex16 EnterStateFailedTaskIndex = FStateTreeIndex16::Invalid;
-
-	/** The index of the node that just finished exiting. */
-	UPROPERTY()
-	FStateTreeIndex16 LastExitedNodeIndex = FStateTreeIndex16::Invalid;
-
 	/** Result of last TickTasks */
 	UPROPERTY()
 	EStateTreeRunStatus LastTickStatus = EStateTreeRunStatus::Failed;
@@ -1136,6 +1131,18 @@ PRAGMA_DISABLE_DEPRECATION_WARNINGS
 
 	UE_DEPRECATED(5.6, "CurrentExecutionContext is not needed anymore. Use FrameID and StateID.")
 	FStateTreeExecutionContext* CurrentExecutionContext = nullptr;
+
+	UE_DEPRECATED(5.7, "Use FStateTreeExecutionFrame::ActiveTaskNodeIndex instead.")
+	UPROPERTY()
+	FStateTreeIndex16 EnterStateFailedFrameIndex = FStateTreeIndex16::Invalid;
+
+	UE_DEPRECATED(5.7, "Use FStateTreeExecutionFrame::ActiveTaskNodeIndex instead.")
+	UPROPERTY()
+	FStateTreeIndex16 EnterStateFailedTaskIndex = FStateTreeIndex16::Invalid;
+
+	UE_DEPRECATED(5.7, "Use FStateTreeExecutionFrame::ActiveTaskNodeIndex instead.")
+	UPROPERTY()
+	FStateTreeIndex16 LastExitedNodeIndex = FStateTreeIndex16::Invalid;
 PRAGMA_ENABLE_DEPRECATION_WARNINGS
 #endif //WITH_EDITORONLY_DATA
 };
diff --git a/Engine/Plugins/Runtime/StateTree/Source/StateTreeTestSuite/Private/StateTreeTest.cpp b/Engine/Plugins/Runtime/StateTree/Source/StateTreeTestSuite/Private/StateTreeTest.cpp
--- a/Engine/Plugins/Runtime/StateTree/Source/StateTreeTestSuite/Private/StateTreeTest.cpp
+++ b/Engine/Plugins/Runtime/StateTree/Source/StateTreeTestSuite/Private/StateTreeTest.cpp
@@ -298,8 +298,21 @@ struct FStateTreeTest_FailEnterState : FStateTreeTestBase
 		AITEST_TRUE(TEXT("StateTree exec status should be failed"), Exec.GetLastTickStatus() == EStateTreeRunStatus::Failed);
 		Exec.LogClear();
 
+		// It will try 5 times to reenter the same states.
+		Exec.Tick(0.01f);
+		AITEST_TRUE(TEXT("StateTree TaskRoot should tick"), Exec.Expect(TaskRoot.GetName(), TickStr));
+		AITEST_TRUE(TEXT("StateTree Task1 should tick"), Exec.Expect(Task1.GetName(), TickStr));
+		AITEST_FALSE(TEXT("StateTree Task2 should not tick because it completed"), Exec.Expect(Task2.GetName(), TickStr));
+		AITEST_FALSE(TEXT("StateTree Task3 should not tick because it didn't enter state"), Exec.Expect(Task3.GetName(), TickStr));
+		AITEST_FALSE(TEXT("StateTree Task3 should not exit state"), Exec.Expect(Task3.GetName(), ExitStateStr));
+		Exec.LogClear();
+
 		// Stop and exit state
 		Status = Exec.Stop();
+		AITEST_FALSE(TEXT("StateTree Task1 should not state complete"), Exec.Expect(Task1.GetName(), StateCompletedStr));
+		AITEST_FALSE(TEXT("StateTree Task2 should not state complete"), Exec.Expect(Task2.GetName(), StateCompletedStr));
+		AITEST_FALSE(TEXT("StateTree Task3 should not state complete"), Exec.Expect(Task3.GetName(), StateCompletedStr));
+		AITEST_FALSE(TEXT("StateTree TaskRoot should not state complete"), Exec.Expect(TaskRoot.GetName(), StateCompletedStr));
 		AITEST_TRUE(TEXT("StateTree TaskRoot should exit state"), Exec.Expect(TaskRoot.GetName(), ExitStateStr));
 		AITEST_TRUE(TEXT("StateTree Task1 should exit state"), Exec.Expect(Task1.GetName(), ExitStateStr));
 		AITEST_TRUE(TEXT("StateTree Task2 should exit state"), Exec.Expect(Task2.GetName(), ExitStateStr));
