[StateTreeDebugger]
- Refactored DebuggerView to use extrapolated world time to provide smoother UI when analyzing live recording (e.g., PIE) without enough produced events to force refresh the tracks
- StateTreeDebugger and RewindDebugger now listen to ShutdownPIE (instead of EndPie) to be able to process event emitted during EndPlay

diff --git a/Engine/Plugins/Animation/GameplayInsights/Source/RewindDebugger/Private/RewindDebugger.cpp b/Engine/Plugins/Animation/GameplayInsights/Source/RewindDebugger/Private/RewindDebugger.cpp
--- a/Engine/Plugins/Animation/GameplayInsights/Source/RewindDebugger/Private/RewindDebugger.cpp
+++ b/Engine/Plugins/Animation/GameplayInsights/Source/RewindDebugger/Private/RewindDebugger.cpp
@@ -97,8 +97,9 @@ FRewindDebugger::FRewindDebugger()
 	FEditorDelegates::PreBeginPIE.AddRaw(this, &FRewindDebugger::OnPIEStarted);
 	FEditorDelegates::PausePIE.AddRaw(this, &FRewindDebugger::OnPIEPaused);
 	FEditorDelegates::ResumePIE.AddRaw(this, &FRewindDebugger::OnPIEResumed);
-	FEditorDelegates::EndPIE.AddRaw(this, &FRewindDebugger::OnPIEStopped);
 	FEditorDelegates::SingleStepPIE.AddRaw(this, &FRewindDebugger::OnPIESingleStepped);
+	// Using ShutdownPIE instead of EndPIE to make sure all traces emitted during world EndPlay get processed before disabling channels
+	FEditorDelegates::ShutdownPIE.AddRaw(this, &FRewindDebugger::OnPIEStopped);
 
 	SelectedObjectName.OnPropertyChanged = SelectedObjectName.OnPropertyChanged.CreateLambda([this](FString Target)
 		{
@@ -141,7 +142,7 @@ FRewindDebugger::~FRewindDebugger()
 	FEditorDelegates::PostPIEStarted.RemoveAll(this);
 	FEditorDelegates::PausePIE.RemoveAll(this);
 	FEditorDelegates::ResumePIE.RemoveAll(this);
-	FEditorDelegates::EndPIE.RemoveAll(this);
+	FEditorDelegates::ShutdownPIE.RemoveAll(this);
 	FEditorDelegates::SingleStepPIE.RemoveAll(this);
 
 	FTSTicker::GetCoreTicker().RemoveTicker(TickerHandle);
diff --git a/Engine/Plugins/Runtime/StateTree/Source/StateTreeEditorModule/Private/Debugger/SStateTreeDebuggerView.cpp b/Engine/Plugins/Runtime/StateTree/Source/StateTreeEditorModule/Private/Debugger/SStateTreeDebuggerView.cpp
--- a/Engine/Plugins/Runtime/StateTree/Source/StateTreeEditorModule/Private/Debugger/SStateTreeDebuggerView.cpp
+++ b/Engine/Plugins/Runtime/StateTree/Source/StateTreeEditorModule/Private/Debugger/SStateTreeDebuggerView.cpp
@@ -197,7 +197,8 @@ void GenerateElementForProperties(const TCHAR* TypeAsText, const TCHAR* ValueAsT
 SStateTreeDebuggerView::SStateTreeDebuggerView()
 {
 	FEditorDelegates::BeginPIE.AddRaw(this, &SStateTreeDebuggerView::OnPIEStarted);
-	FEditorDelegates::EndPIE.AddRaw(this, &SStateTreeDebuggerView::OnPIEStopped);
+	// Using ShutdownPIE instead of EndPIE to make sure all traces emitted during world EndPlay get processed before disabling channels
+	FEditorDelegates::ShutdownPIE.AddRaw(this, &SStateTreeDebuggerView::OnPIEStopped);
 	FEditorDelegates::PausePIE.AddRaw(this, &SStateTreeDebuggerView::OnPIEPaused);
 	FEditorDelegates::ResumePIE.AddRaw(this, &SStateTreeDebuggerView::OnPIEResumed);
 	FEditorDelegates::SingleStepPIE.AddRaw(this, &SStateTreeDebuggerView::OnPIESingleStepped);
@@ -217,7 +218,7 @@ SStateTreeDebuggerView::~SStateTreeDebuggerView()
 	Debugger->ClearSelection();
 
 	FEditorDelegates::BeginPIE.RemoveAll(this);
-	FEditorDelegates::EndPIE.RemoveAll(this);
+	FEditorDelegates::ShutdownPIE.RemoveAll(this);
 	FEditorDelegates::PausePIE.RemoveAll(this);
 	FEditorDelegates::ResumePIE.RemoveAll(this);
 	FEditorDelegates::SingleStepPIE.RemoveAll(this);
@@ -251,16 +252,19 @@ void SStateTreeDebuggerView::StopRecording()
 
 		// Update max duration from current recording until track data gets reset
 		check(Debugger);
-		MaxTrackRecordingDuration = FMath::Max(MaxTrackRecordingDuration, Debugger->GetRecordingDuration());
+		MaxTrackRecordingDuration = FMath::Max(MaxTrackRecordingDuration, ExtrapolatedRecordedWorldTime);
 
 		// Mark all tracks from the stopped session as stale to have different look.
 		for (const TSharedPtr<RewindDebugger::FRewindDebuggerTrack>& DebugTrack : InstanceOwnerTracks)
 		{
 			if (FStateTreeDebuggerBaseTrack* StateTreeTrack = static_cast<FStateTreeDebuggerBaseTrack*>(DebugTrack.Get()))
 			{
-				StateTreeTrack->MarkAsStale();
+				StateTreeTrack->MarkAsStale(ExtrapolatedRecordedWorldTime);
 			}
 		}
+
+		// Reset the extrapolated time for next runs
+		ExtrapolatedRecordedWorldTime = 0;
 	}
 }
 
@@ -707,21 +711,36 @@ void SStateTreeDebuggerView::Tick(const FGeometry& AllottedGeometry, const doubl
 
 	check(Debugger);
 
-	const double RecordingDuration = Debugger->GetRecordingDuration();
-	const bool bHasMoreRecentData = LastUpdatedTrackRecordingDuration != RecordingDuration;
-	LastUpdatedTrackRecordingDuration = RecordingDuration;
-	MaxTrackRecordingDuration = FMath::Max(MaxTrackRecordingDuration, RecordingDuration);
+	const double LastProcessedRecordedWorldTime = Debugger->GetLastProcessedRecordedWorldTime();
+	const bool bHasMoreRecentData = LastUpdatedTrackRecordingDuration != LastProcessedRecordedWorldTime;
+	LastUpdatedTrackRecordingDuration = LastProcessedRecordedWorldTime;
+
+	if (bHasMoreRecentData)
+	{
+		// Prevent extrapolated value from going back, can only catch up
+		ExtrapolatedRecordedWorldTime = FMath::Max(LastProcessedRecordedWorldTime, ExtrapolatedRecordedWorldTime);
+	}
+
+	MaxTrackRecordingDuration = FMath::Max(MaxTrackRecordingDuration, MaxTrackRecordingDuration);
 	
 	if ((Debugger->IsAnalysisSessionActive() && !Debugger->IsAnalysisSessionPaused())
 		|| bHasMoreRecentData)
 	{
+		// When no events are processed this frame we force timeline to progress to make UI feel more responsive,
+		// This is useful when debugging a live recording where events are produced in realtime and can't be read faster from a saved recording.
+		// Note that we update regardless of autoscroll to maintain a valid value if it gets enabled.
+		if (!bHasMoreRecentData)
+		{
+			ExtrapolatedRecordedWorldTime += InDeltaTime;
+		}
+
 		if (bAutoScroll)
 		{
 			// Stick to most recent data if auto scroll is enabled.
 			// Autoscroll is disabled when paused.
 			// This allows the user to pause the analysis, inspect the data, and continue and the autoscroll will catch up with latest.
 			// Complementary logic in OnTimeLineScrubPositionChanged().
-			Debugger->SetScrubTime(Debugger->GetRecordingDuration());
+			Debugger->SetScrubTime(ExtrapolatedRecordedWorldTime);
 		}
 		else
 		{
@@ -1204,7 +1223,7 @@ void SStateTreeDebuggerView::OnNewInstance(FStateTreeInstanceDebugId InstanceId)
 	if (FStateTreeDebuggerOwnerTrack* OwnerTrack = static_cast<FStateTreeDebuggerOwnerTrack*>(ExistingOwnerTrack->Get()))
 	{
 		const FString TrackName = FString::Printf(TEXT("Execution #%d"), OwnerTrack->NumSubTracks()+1);
-		const TSharedPtr<FStateTreeDebuggerInstanceTrack> SubTrack = MakeShared<FStateTreeDebuggerInstanceTrack>(Debugger.ToSharedRef(), InstanceId, FText::FromString(TrackName), ViewRange);
+		const TSharedPtr<FStateTreeDebuggerInstanceTrack> SubTrack = MakeShared<FStateTreeDebuggerInstanceTrack>(SharedThis(this), Debugger.ToSharedRef(), InstanceId, FText::FromString(TrackName), ViewRange);
 		OwnerTrack->AddSubTrack(SubTrack);
 
 		// Look at current selection; if nothing selected or stale track then select new track
diff --git a/Engine/Plugins/Runtime/StateTree/Source/StateTreeEditorModule/Private/Debugger/SStateTreeDebuggerView.h b/Engine/Plugins/Runtime/StateTree/Source/StateTreeEditorModule/Private/Debugger/SStateTreeDebuggerView.h
--- a/Engine/Plugins/Runtime/StateTree/Source/StateTreeEditorModule/Private/Debugger/SStateTreeDebuggerView.h
+++ b/Engine/Plugins/Runtime/StateTree/Source/StateTreeEditorModule/Private/Debugger/SStateTreeDebuggerView.h
@@ -48,6 +48,16 @@ public:
 
 	void Construct(const FArguments& InArgs, const UStateTree& InStateTree, const TSharedRef<FStateTreeViewModel>& InStateTreeViewModel, const TSharedRef<FUICommandList>& InCommandList);
 
+	[[nodiscard]] TSharedPtr<FStateTreeDebugger> GetDebugger() const
+	{
+		return Debugger;
+	}
+
+	[[nodiscard]] double GetExtrapolatedRecordedWorldTime() const
+	{
+		return ExtrapolatedRecordedWorldTime;
+	}
+
 private:
 	virtual FReply OnPreviewKeyDown(const FGeometry& MyGeometry, const FKeyEvent& InKeyEvent) override;
 	virtual void Tick(const FGeometry& AllottedGeometry, const double InCurrentTime, const float InDeltaTime) override;
@@ -73,12 +83,23 @@ private:
 
 	bool CanUseScrubButtons() const;
 
-	bool CanStartRecording() const { return !IsRecording(); }
+	bool CanStartRecording() const
+	{
+		return !IsRecording();
+	}
+
 	void StartRecording();
 
-	bool IsRecording() const { return bRecording; }
+	bool IsRecording() const
+	{
+		return bRecording;
+	}
+
+	bool CanStopRecording() const
+	{
+		return IsRecording();
+	}
 
-	bool CanStopRecording() const { return IsRecording(); }	
 	void StopRecording();
 
 	bool CanResumeDebuggerAnalysis() const;
@@ -157,6 +178,13 @@ private:
 	/** The recording duration time the UI was last updated from the debugger. Used to detect if new data has been collect while the UI was inactive. */
 	double LastUpdatedTrackRecordingDuration = 0;
 
+	/**
+	 * Extrapolated world time used to make the timeline progress until the debugger
+	 * receives a new event that will update its last processed world time.
+	 */
+	double ExtrapolatedRecordedWorldTime = 0;
+
+private:
 	/** Indicates that a live session was started (record button or auto record in PIE) to generate StateTree traces. */
 	bool bRecording = false;
 
diff --git a/Engine/Plugins/Runtime/StateTree/Source/StateTreeEditorModule/Private/Debugger/StateTreeDebuggerTrack.cpp b/Engine/Plugins/Runtime/StateTree/Source/StateTreeEditorModule/Private/Debugger/StateTreeDebuggerTrack.cpp
--- a/Engine/Plugins/Runtime/StateTree/Source/StateTreeEditorModule/Private/Debugger/StateTreeDebuggerTrack.cpp
+++ b/Engine/Plugins/Runtime/StateTree/Source/StateTreeEditorModule/Private/Debugger/StateTreeDebuggerTrack.cpp
@@ -4,6 +4,7 @@
 
 #include "Debugger/StateTreeDebuggerTrack.h"
 #include "Debugger/StateTreeDebugger.h"
+#include "SStateTreeDebuggerView.h"
 
 namespace UE::StateTreeDebugger
 {
@@ -21,9 +22,8 @@ FInstanceTrackHelper::FInstanceTrackHelper(const FStateTreeInstanceDebugId InIns
 bool FInstanceTrackHelper::RebuildEventData(
 	const TNotNull<const UStateTree*> InStateTree
 	, const FInstanceEventCollection& InEventCollection
-	, const TRange<double>& InLifetime
-	, const double InRecordingDuration
-	, const bool bInActive) const
+	, const double InInstanceLastActiveTime
+	, const bool bInIsStaleTrack) const
 {
 	if (InEventCollection.IsInvalid())
 	{
@@ -36,11 +36,11 @@ bool FInstanceTrackHelper::RebuildEventData(
 	TimelineData->Points.SetNum(0, EAllowShrinking::No);
 	TimelineData->Windows.SetNum(0);
 
-	auto MakeRandomColor = [bInActive](const uint32 InSeed)->FLinearColor
+	auto MakeRandomColor = [bActive = !bInIsStaleTrack](const uint32 InSeed)->FLinearColor
 		{
 			const FRandomStream Stream(InSeed);
 			const uint8 Hue = static_cast<uint8>(Stream.FRand() * 255.0f);
-			const uint8 SatVal = bInActive ? 196 : 128;
+			const uint8 SatVal = bActive ? 196 : 128;
 			return FLinearColor::MakeFromHSV8(Hue, SatVal, SatVal);
 		};
 
@@ -91,7 +91,7 @@ bool FInstanceTrackHelper::RebuildEventData(
 			// For the last received event we use either the current recording duration if the instance is still active or the last recorded frame time
 			if (StateChangeIndex == (NumStateChanges - 1))
 			{
-				Window.TimeEnd = InLifetime.Contains(InRecordingDuration) ? InRecordingDuration : InEventCollection.FrameSpans.Last().GetWorldTimeEnd();
+				Window.TimeEnd = InInstanceLastActiveTime;
 			}
 			else
 			{
@@ -133,15 +133,17 @@ bool FInstanceTrackHelper::RebuildEventData(
 // FStateTreeDebuggerInstanceTrack
 //----------------------------------------------------------------------//
 FStateTreeDebuggerInstanceTrack::FStateTreeDebuggerInstanceTrack(
+	const TSharedRef<SStateTreeDebuggerView>& InDebuggerView,
 	const TSharedRef<FStateTreeDebugger>& InDebugger,
 	const FStateTreeInstanceDebugId InInstanceId,
 	const FText& InName,
 	const TRange<double>& InViewRange
 	)
 	: FStateTreeDebuggerBaseTrack(FSlateIcon("StateTreeEditorStyle", "StateTreeEditor.Debugger.InstanceTrack", "StateTreeEditor.Debugger.InstanceTrack"), InName)
+	, StateTreeDebuggerView(InDebuggerView)
 	, StateTreeDebugger(InDebugger)
 	, InstanceTrackHelper(new UE::StateTreeDebugger::FInstanceTrackHelper(InInstanceId, InViewRange))
-	, Descriptor(InDebugger->GetDescriptor(InInstanceId))
+	, Descriptor(InDebugger.Get().GetDescriptor(InInstanceId))
 {
 }
 
@@ -153,17 +155,25 @@ void FStateTreeDebuggerInstanceTrack::OnSelected()
 bool FStateTreeDebuggerInstanceTrack::UpdateInternal()
 {
 	const FStateTreeDebugger& Debugger = StateTreeDebugger.Get();
-	const bool bActive = !IsStale();
+
+	double InstanceLastActiveTime = StateTreeDebuggerView.Get().GetExtrapolatedRecordedWorldTime();
+	const bool bIsStaleTrack = IsStale();
+	if (bIsStaleTrack)
+	{
+		// If the instance was no longer active when the track was marked as stale we use its lifetime to build the data
+		// otherwise we build up to the last time the track was active.
+		InstanceLastActiveTime = Descriptor->IsInstanceActive() ? StaleTime : Descriptor.Get()->Lifetime.GetUpperBoundValue();
+	}
+
 	const bool bChanged = InstanceTrackHelper.Get().RebuildEventData(Debugger.GetAsset()
 		, Debugger.GetEventCollection(InstanceTrackHelper.Get().GetInstanceId())
-		, Descriptor->Lifetime
-		, Debugger.GetRecordingDuration()
-		, bActive);
+		, InstanceLastActiveTime
+		, bIsStaleTrack);
 
 	// Tracks can be reactivated when multiple recordings are made in a single PIE session.
 	if (bChanged && IsStale())
 	{
-		bIsStale = false;
+		MarkAsActive();
 	}
 
 	return bChanged;
@@ -201,13 +211,13 @@ void FStateTreeDebuggerOwnerTrack::IterateSubTracksInternal(TFunction<void(TShar
 	}
 }
 
-void FStateTreeDebuggerOwnerTrack::MarkAsStale()
+void FStateTreeDebuggerOwnerTrack::MarkAsStale(const double InStaleTime)
 {
 	for (TSharedPtr<FStateTreeDebuggerInstanceTrack>& Track : SubTracks)
 	{
 		if (FStateTreeDebuggerInstanceTrack* InstanceTrack = Track.Get())
 		{
-			InstanceTrack->MarkAsStale();
+			InstanceTrack->MarkAsStale(InStaleTime);
 		}
 	}
 }
diff --git a/Engine/Plugins/Runtime/StateTree/Source/StateTreeEditorModule/Private/Debugger/StateTreeDebuggerTrack.h b/Engine/Plugins/Runtime/StateTree/Source/StateTreeEditorModule/Private/Debugger/StateTreeDebuggerTrack.h
--- a/Engine/Plugins/Runtime/StateTree/Source/StateTreeEditorModule/Private/Debugger/StateTreeDebuggerTrack.h
+++ b/Engine/Plugins/Runtime/StateTree/Source/StateTreeEditorModule/Private/Debugger/StateTreeDebuggerTrack.h
@@ -8,6 +8,7 @@
 #include "SStateTreeDebuggerEventTimelineView.h"
 #include "Debugger/StateTreeDebuggerTypes.h"
 
+class SStateTreeDebuggerView;
 struct FStateTreeDebugger;
 
 namespace UE::StateTreeDebugger
@@ -28,9 +29,8 @@ struct FInstanceTrackHelper : public TSharedFromThis<FInstanceTrackHelper>
 	bool RebuildEventData(
 		TNotNull<const UStateTree*> InStateTree
 		, const FInstanceEventCollection& InEventCollection
-		, const TRange<double>& InLifetime
-		, const double InRecordingDuration
-		, const bool bInActive) const;
+		, const double InInstanceLastActiveTime
+		, const bool bInIsStaleTrack = false) const;
 
 	FStateTreeInstanceDebugId GetInstanceId() const
 	{
@@ -82,7 +82,7 @@ struct FStateTreeDebuggerBaseTrack : RewindDebugger::FRewindDebuggerTrack
 		return false;
 	}
 
-	virtual void MarkAsStale()
+	virtual void MarkAsStale(const double InStaleTime)
 	{
 	}
 
@@ -114,20 +114,31 @@ protected:
 struct FStateTreeDebuggerInstanceTrack : FStateTreeDebuggerBaseTrack
 {
 	explicit FStateTreeDebuggerInstanceTrack(
+		const TSharedRef<SStateTreeDebuggerView>& InDebuggerView,
 		const TSharedRef<FStateTreeDebugger>& InDebugger,
 		const FStateTreeInstanceDebugId InInstanceId,
 		const FText& InName,
 		const TRange<double>& InViewRange);
 
 	virtual void OnSelected() override;
-	virtual void MarkAsStale() override
+
+	void MarkAsActive()
+	{
+		StaleTime = UE::StateTreeDebugger::FInstanceDescriptor::ActiveInstanceEndTime;
+	}
+
+	virtual void MarkAsStale(const double InStaleTime) override
 	{
-		bIsStale = true;
+		// Do not update timestamp for already stale tracks
+		if (!IsStale())
+		{
+			StaleTime = InStaleTime;
+		}
 	}
 
 	virtual bool IsStale() const override
 	{
-		return bIsStale;
+		return StaleTime != UE::StateTreeDebugger::FInstanceDescriptor::ActiveInstanceEndTime;
 	}
 
 protected:
@@ -136,10 +147,11 @@ protected:
 
 private:
 	
+	TSharedRef<SStateTreeDebuggerView> StateTreeDebuggerView;
 	TSharedRef<FStateTreeDebugger> StateTreeDebugger;
 	TSharedRef<UE::StateTreeDebugger::FInstanceTrackHelper> InstanceTrackHelper;
 	TSharedPtr<const UE::StateTreeDebugger::FInstanceDescriptor> Descriptor;
-	bool bIsStale = false;
+	double StaleTime = UE::StateTreeDebugger::FInstanceDescriptor::ActiveInstanceEndTime;
 };
 
 
@@ -158,7 +170,7 @@ struct FStateTreeDebuggerOwnerTrack : FStateTreeDebuggerBaseTrack
 		return SubTracks.Num();
 	}
 
-	virtual void MarkAsStale() override;
+	virtual void MarkAsStale(double InStaleTime) override;
 	virtual bool IsStale() const override;
 
 protected:
diff --git a/Engine/Plugins/Runtime/StateTree/Source/StateTreeEditorModule/Private/Debugger/StateTreeRewindDebuggerTrack.cpp b/Engine/Plugins/Runtime/StateTree/Source/StateTreeEditorModule/Private/Debugger/StateTreeRewindDebuggerTrack.cpp
--- a/Engine/Plugins/Runtime/StateTree/Source/StateTreeEditorModule/Private/Debugger/StateTreeRewindDebuggerTrack.cpp
+++ b/Engine/Plugins/Runtime/StateTree/Source/StateTreeEditorModule/Private/Debugger/StateTreeRewindDebuggerTrack.cpp
@@ -152,11 +152,7 @@ bool FRewindDebuggerTrack::UpdateInternal()
 			, FTraceFilter(InstanceTrackHelper.Get().GetInstanceId())
 			, LastTraceReadTime);
 
-		return InstanceTrackHelper.Get().RebuildEventData(StateTree
-			, EventCollection
-			, Descriptor->Lifetime
-			, RewindDebugger->GetScrubTime()
-			, /*Active*/true);
+		return InstanceTrackHelper.Get().RebuildEventData(StateTree, EventCollection, RewindDebugger->GetScrubTime());
 	}
 
 	return false;
diff --git a/Engine/Plugins/Runtime/StateTree/Source/StateTreeModule/Private/Debugger/StateTreeDebugger.cpp b/Engine/Plugins/Runtime/StateTree/Source/StateTreeModule/Private/Debugger/StateTreeDebugger.cpp
--- a/Engine/Plugins/Runtime/StateTree/Source/StateTreeModule/Private/Debugger/StateTreeDebugger.cpp
+++ b/Engine/Plugins/Runtime/StateTree/Source/StateTreeModule/Private/Debugger/StateTreeDebugger.cpp
@@ -119,7 +119,9 @@ void FStateTreeDebugger::Tick(const float DeltaTime)
 		}
 	}
 
-	if (bSessionAnalysisPaused == false && StateTreeAsset.IsValid())
+	if (StateTreeAsset.IsValid()
+		&& IsAnalysisSessionActive()
+		&& !IsAnalysisSessionPaused())
 	{
 		SyncToCurrentSessionDuration();
 	}
@@ -127,12 +129,14 @@ void FStateTreeDebugger::Tick(const float DeltaTime)
 
 void FStateTreeDebugger::StopSessionAnalysis()
 {
-	if (const TraceServices::IAnalysisSession* Session = GetAnalysisSession())
+	if (IsAnalysisSessionActive())
 	{
-		Session->Stop(true);
-		AnalysisSession.Reset();
+		// Force one last update to process events emitted while closing the game session (e.g., EndPlay in PIE)
+		SyncToCurrentSessionDuration();
+		AnalysisSession->Stop(/*WaitOnAnalysis*/true);
 	}
 
+	bSessionAnalysisActive = false;
 	bSessionAnalysisPaused = false;
 	HitBreakpoint.Reset();
 }
@@ -353,12 +357,13 @@ bool FStateTreeDebugger::StartSessionAnalysis(const FTraceDescriptor& TraceDescr
 		ITraceServicesModule& TraceServicesModule = FModuleManager::LoadModuleChecked<ITraceServicesModule>("TraceServices");
 		if (const TSharedPtr<TraceServices::IAnalysisService> TraceAnalysisService = TraceServicesModule.GetAnalysisService())
 		{
-			checkf(!AnalysisSession.IsValid(), TEXT("Must make sure that current session was properly stopped before starting a new one otherwise it can cause threading issues"));
+			checkf(!IsAnalysisSessionActive(), TEXT("Must make sure that current session was properly stopped before starting a new one otherwise it can cause threading issues"));
 			AnalysisSession = TraceAnalysisService->StartAnalysis(TraceId, *TraceName, MoveTemp(TraceData));
 		}
 
 		if (AnalysisSession.IsValid())
 		{
+			bSessionAnalysisActive = true;
 			ActiveSessionTraceDescriptor = TraceDescriptor;
 		}
 	}
@@ -807,7 +812,7 @@ bool FStateTreeDebugger::EvaluateBreakpoints(const FStateTreeInstanceDebugId Ins
 		{
 			HitBreakpoint.Index = BreakpointIndex;
 			HitBreakpoint.InstanceId = InstanceId;
-			HitBreakpoint.Time = RecordingDuration;
+			HitBreakpoint.Time = LastProcessedRecordedWorldTime;
 		}
 	}
 
@@ -842,9 +847,9 @@ UE::StateTreeDebugger::FInstanceEventCollection* FStateTreeDebugger::GetOrCreate
 
 void FStateTreeDebugger::OnTraceEventProcessed(const FStateTreeInstanceDebugId InstanceId, const FStateTreeTraceEventVariantType& Event)
 {
-	Visit([&RecordingDuration = RecordingDuration](auto& TypedEvent)
+	Visit([&WorldTime = LastProcessedRecordedWorldTime](auto& TypedEvent)
 		{
-			RecordingDuration = TypedEvent.RecordingWorldTime;
+			WorldTime = TypedEvent.RecordingWorldTime;
 		}, Event);
 
 	EvaluateBreakpoints(InstanceId, Event);
@@ -865,7 +870,7 @@ void FStateTreeDebugger::ResetEventCollections()
 {
 	EventCollections.Reset();
 	SetScrubStateCollectionIndex(INDEX_NONE);
-	RecordingDuration = 0;
+	LastProcessedRecordedWorldTime = 0;
 }
 
 //----------------------------------------------------------------//
diff --git a/Engine/Plugins/Runtime/StateTree/Source/StateTreeModule/Private/Debugger/StateTreeTraceProvider.cpp b/Engine/Plugins/Runtime/StateTree/Source/StateTreeModule/Private/Debugger/StateTreeTraceProvider.cpp
--- a/Engine/Plugins/Runtime/StateTree/Source/StateTreeModule/Private/Debugger/StateTreeTraceProvider.cpp
+++ b/Engine/Plugins/Runtime/StateTree/Source/StateTreeModule/Private/Debugger/StateTreeTraceProvider.cpp
@@ -26,7 +26,7 @@ bool FStateTreeTraceProvider::ReadTimelines(const FStateTreeInstanceDebugId Inst
 		{
 			Callback(InstanceId, *EventsTimelines[*IndexPtr]);
 			return true;
-		}	
+		}
 	}
 	else
 	{
@@ -92,20 +92,31 @@ void FStateTreeTraceProvider::AppendInstanceEvent(
 		{
 			if (const UStateTree* StateTree = WeakStateTree.Get())
 			{
-				Descriptors.Emplace(MakeShared<UE::StateTreeDebugger::FInstanceDescriptor>(
-					StateTree
-					, InInstanceId
-					, InInstanceName
-					, TRange<double>(InWorldRecordingTime, std::numeric_limits<double>::max())));
-
-				check(InstanceIdToDebuggerEntryTimelines.Find(InInstanceId) == nullptr);
-				InstanceIdToDebuggerEntryTimelines.Add(InInstanceId, EventsTimelines.Num());
-		
-				EventsTimelines.Emplace(MakeShared<TraceServices::TPointTimeline<FStateTreeTraceEventVariantType>>(Session.GetLinearAllocator()));
+				const TSharedRef<UE::StateTreeDebugger::FInstanceDescriptor>* Descriptor = Descriptors.FindByPredicate(
+					[InInstanceId](const TSharedRef<const UE::StateTreeDebugger::FInstanceDescriptor>& Descriptor)
+					{
+						return Descriptor.Get().Id == InInstanceId;
+					});
+
+				// Possible to receive new Push when stopping/starting traces during the same session.
+				// In that case we can reuse the existing entries
+				if (Descriptor == nullptr)
+				{
+					Descriptors.Emplace(MakeShared<UE::StateTreeDebugger::FInstanceDescriptor>(
+						StateTree
+						, InInstanceId
+						, InInstanceName
+						, TRange<double>(InWorldRecordingTime, UE::StateTreeDebugger::FInstanceDescriptor::ActiveInstanceEndTime)));
+
+					check(InstanceIdToDebuggerEntryTimelines.Find(InInstanceId) == nullptr);
+					InstanceIdToDebuggerEntryTimelines.Add(InInstanceId, EventsTimelines.Num());
+
+					EventsTimelines.Emplace(MakeShared<TraceServices::TPointTimeline<FStateTreeTraceEventVariantType>>(Session.GetLinearAllocator()));
+				}
 			}
 			else
 			{
-				UE_LOG(LogStateTree, Error, TEXT("Instance event refers to an unloaded asset."));	
+				UE_LOG(LogStateTree, Error, TEXT("Instance event refers to an unloaded asset."));
 			}
 		}
 		else
diff --git a/Engine/Plugins/Runtime/StateTree/Source/StateTreeModule/Public/Debugger/StateTreeDebugger.h b/Engine/Plugins/Runtime/StateTree/Source/StateTreeModule/Public/Debugger/StateTreeDebugger.h
--- a/Engine/Plugins/Runtime/StateTree/Source/StateTreeModule/Public/Debugger/StateTreeDebugger.h
+++ b/Engine/Plugins/Runtime/StateTree/Source/StateTreeModule/Public/Debugger/StateTreeDebugger.h
@@ -228,10 +228,10 @@ struct FStateTreeDebugger : FTickableGameObject, UE::StateTreeDebugger::ITraceRe
 	/** Clears events from all instances. */
 	UE_API void ResetEventCollections();
 
-	/** Returns the recording duration in world recorded time. */
-	double GetRecordingDuration() const
+	/** Returns the world time associated to the last processed event. */
+	double GetLastProcessedRecordedWorldTime() const
 	{
-		return RecordingDuration;
+		return LastProcessedRecordedWorldTime;
 	}
 
 	/** Returns the duration of the analysis session. This is not related to the world simulation time. */
@@ -268,11 +268,18 @@ struct FStateTreeDebugger : FTickableGameObject, UE::StateTreeDebugger::ITraceRe
 			|| AnalysisTransitionType == EAnalysisTransitionType::EditorToEditor;
 	}
 
+	/**
+	 * Indicates whether the analysis session was started and not stopped yet.
+	 */
 	bool IsAnalysisSessionActive() const
 	{
-		return GetAnalysisSession() != nullptr;
+		return bSessionAnalysisActive;
 	}
 
+	/**
+	 * Indicates that the debugger no longer processes new events from the analysis session until it gets resumed.
+	 * This can be an external explicit request or after hitting a breakpoint.
+	 */
 	bool IsAnalysisSessionPaused() const
 	{
 		return bSessionAnalysisPaused;
@@ -449,8 +456,10 @@ private:
 	 */
 	float RetryLoadNextLiveSessionTimer = 0.0f;
 	
-	/** Recording duration of the analysis session in world recorded time. */
-	double RecordingDuration = 0;
+	/**
+	 * Most recent processed world time. This value is updated from received events so it might increase by large deltas.
+	 */
+	double LastProcessedRecordedWorldTime = 0;
 
 	/** Duration of the analysis session. This is not related to the world simulation time. */
 	double AnalysisDuration = 0;
@@ -467,6 +476,12 @@ private:
 	/** List of new instances discovered by processing event in the analysis session. */
 	TArray<FStateTreeInstanceDebugId> NewInstances;
 
+	/**
+	 * Indicates whether the analysis session is valid and active.
+	 * This flag is required since once an analysis session is stopped there is nothing in its API to retrieve its status.
+	 */
+	bool bSessionAnalysisActive = false;
+
 	/**
 	 * Indicates that the debugger no longer process new events from the analysis session until it gets resumed.
 	 * This can be an external explicit request or after hitting a breakpoint.
diff --git a/Engine/Plugins/Runtime/StateTree/Source/StateTreeModule/Public/Debugger/StateTreeDebuggerTypes.h b/Engine/Plugins/Runtime/StateTree/Source/StateTreeModule/Public/Debugger/StateTreeDebuggerTypes.h
--- a/Engine/Plugins/Runtime/StateTree/Source/StateTreeModule/Public/Debugger/StateTreeDebuggerTypes.h
+++ b/Engine/Plugins/Runtime/StateTree/Source/StateTreeModule/Public/Debugger/StateTreeDebuggerTypes.h
@@ -77,6 +77,12 @@ struct FInstanceDescriptor : TSharedFromThis<FInstanceDescriptor>
 		return !(*this == Other);
 	}
 
+	/** Whether the associated instance is still active (i.e., not explicitly stopped by an execution context) */
+	bool IsInstanceActive() const
+	{
+		return Lifetime.GetUpperBoundValue() == ActiveInstanceEndTime;
+	}
+
 	friend FString LexToString(const FInstanceDescriptor& InstanceDesc)
 	{
 		return FString::Printf(TEXT("%s | %s | %s"),
@@ -90,6 +96,7 @@ struct FInstanceDescriptor : TSharedFromThis<FInstanceDescriptor>
 		return GetTypeHash(Desc.Id);
 	}
 
+	static constexpr double ActiveInstanceEndTime = std::numeric_limits<double>::max();
 	TRange<double> Lifetime = TRange<double>(0);
 	TWeakObjectPtr<const UStateTree> StateTree = nullptr;
 	FString Name;
