[StateTreeDebugger] fixed ImGui debugger in cooked builds

diff --git a/Engine/Plugins/Runtime/StateTree/Source/StateTreeModule/StateTreeModule.Build.cs b/Engine/Plugins/Runtime/StateTree/Source/StateTreeModule/StateTreeModule.Build.cs
--- a/Engine/Plugins/Runtime/StateTree/Source/StateTreeModule/StateTreeModule.Build.cs
+++ b/Engine/Plugins/Runtime/StateTree/Source/StateTreeModule/StateTreeModule.Build.cs
@@ -61,8 +61,8 @@ namespace UnrealBuildTool.Rules
 					}
 				);
 				
-				// Allow debugger trace analysis on editor platforms
-				if (Target.Platform.IsInGroup(UnrealPlatformGroup.Desktop) && Target.bBuildEditor)
+				// Allow debugger trace analysis on desktop platforms
+				if (Target.Platform.IsInGroup(UnrealPlatformGroup.Desktop))
 				{
 					PublicDefinitions.Add("WITH_STATETREE_TRACE_DEBUGGER=1");
 					PublicDependencyModuleNames.AddRange(
