[StateTreeDebugger] minor improvements from code-review and fix for cooked builds when closing the process while a debugging session is active
- Updated RewindDebugger delegate registration to use the shared pointer baserd version and fixed one unregistration

diff --git a/Engine/Plugins/Animation/GameplayInsights/Source/RewindDebugger/Private/RewindDebugger.cpp b/Engine/Plugins/Animation/GameplayInsights/Source/RewindDebugger/Private/RewindDebugger.cpp
--- a/Engine/Plugins/Animation/GameplayInsights/Source/RewindDebugger/Private/RewindDebugger.cpp
+++ b/Engine/Plugins/Animation/GameplayInsights/Source/RewindDebugger/Private/RewindDebugger.cpp
@@ -94,12 +94,12 @@ FRewindDebugger::FRewindDebugger()
 		OnPIEStarted(true);
 	}
 
-	FEditorDelegates::PreBeginPIE.AddRaw(this, &FRewindDebugger::OnPIEStarted);
-	FEditorDelegates::PausePIE.AddRaw(this, &FRewindDebugger::OnPIEPaused);
-	FEditorDelegates::ResumePIE.AddRaw(this, &FRewindDebugger::OnPIEResumed);
-	FEditorDelegates::SingleStepPIE.AddRaw(this, &FRewindDebugger::OnPIESingleStepped);
+	PreBeginPIEHandle = FEditorDelegates::PreBeginPIE.AddRaw(this, &FRewindDebugger::OnPIEStarted);
+	PausePIEHandle = FEditorDelegates::PausePIE.AddRaw(this, &FRewindDebugger::OnPIEPaused);
+	ResumePIEHandle = FEditorDelegates::ResumePIE.AddRaw(this, &FRewindDebugger::OnPIEResumed);
+	SingleStepPIEHandle = FEditorDelegates::SingleStepPIE.AddRaw(this, &FRewindDebugger::OnPIESingleStepped);
 	// Using ShutdownPIE instead of EndPIE to make sure all traces emitted during world EndPlay get processed before disabling channels
-	FEditorDelegates::ShutdownPIE.AddRaw(this, &FRewindDebugger::OnPIEStopped);
+	ShutdownPIEHandle = FEditorDelegates::ShutdownPIE.AddRaw(this, &FRewindDebugger::OnPIEStopped);
 
 	SelectedObjectName.OnPropertyChanged = SelectedObjectName.OnPropertyChanged.CreateLambda([this](FString Target)
 		{
@@ -139,11 +139,11 @@ FRewindDebugger::FRewindDebugger()
 
 FRewindDebugger::~FRewindDebugger()
 {
-	FEditorDelegates::PostPIEStarted.RemoveAll(this);
-	FEditorDelegates::PausePIE.RemoveAll(this);
-	FEditorDelegates::ResumePIE.RemoveAll(this);
-	FEditorDelegates::ShutdownPIE.RemoveAll(this);
-	FEditorDelegates::SingleStepPIE.RemoveAll(this);
+	FEditorDelegates::PreBeginPIE.Remove(PreBeginPIEHandle);
+	FEditorDelegates::PausePIE.Remove(PausePIEHandle);
+	FEditorDelegates::ResumePIE.Remove(ResumePIEHandle);
+	FEditorDelegates::SingleStepPIE.Remove(SingleStepPIEHandle);
+	FEditorDelegates::ShutdownPIE.Remove(ShutdownPIEHandle);
 
 	FTSTicker::GetCoreTicker().RemoveTicker(TickerHandle);
 
diff --git a/Engine/Plugins/Animation/GameplayInsights/Source/RewindDebugger/Private/RewindDebugger.h b/Engine/Plugins/Animation/GameplayInsights/Source/RewindDebugger/Private/RewindDebugger.h
--- a/Engine/Plugins/Animation/GameplayInsights/Source/RewindDebugger/Private/RewindDebugger.h
+++ b/Engine/Plugins/Animation/GameplayInsights/Source/RewindDebugger/Private/RewindDebugger.h
@@ -221,6 +221,12 @@ private:
 	void OnPIEStopped(bool bSimulating);
 	void OnPIESingleStepped(bool bSimulating);
 
+	FDelegateHandle PreBeginPIEHandle;
+	FDelegateHandle PausePIEHandle;
+	FDelegateHandle ResumePIEHandle;
+	FDelegateHandle SingleStepPIEHandle;
+	FDelegateHandle ShutdownPIEHandle;
+
 	void SetCurrentScrubTime(double Time);
 
 	TBindableProperty<double> TraceTime;
diff --git a/Engine/Plugins/Runtime/StateTree/Source/StateTreeEditorModule/Private/Debugger/SStateTreeDebuggerView.cpp b/Engine/Plugins/Runtime/StateTree/Source/StateTreeEditorModule/Private/Debugger/SStateTreeDebuggerView.cpp
--- a/Engine/Plugins/Runtime/StateTree/Source/StateTreeEditorModule/Private/Debugger/SStateTreeDebuggerView.cpp
+++ b/Engine/Plugins/Runtime/StateTree/Source/StateTreeEditorModule/Private/Debugger/SStateTreeDebuggerView.cpp
@@ -196,12 +196,12 @@ void GenerateElementForProperties(const TCHAR* TypeAsText, const TCHAR* ValueAsT
 //----------------------------------------------------------------------//
 SStateTreeDebuggerView::SStateTreeDebuggerView()
 {
-	FEditorDelegates::BeginPIE.AddRaw(this, &SStateTreeDebuggerView::OnPIEStarted);
+	BeginPIEHandle = FEditorDelegates::BeginPIE.AddSP(this, &SStateTreeDebuggerView::OnPIEStarted);
+	PausePIEHandle = FEditorDelegates::PausePIE.AddSP(this, &SStateTreeDebuggerView::OnPIEPaused);
+	ResumePIEHandle = FEditorDelegates::ResumePIE.AddSP(this, &SStateTreeDebuggerView::OnPIEResumed);
+	SingleStepPIEHandle =FEditorDelegates::SingleStepPIE.AddSP(this, &SStateTreeDebuggerView::OnPIESingleStepped);
 	// Using ShutdownPIE instead of EndPIE to make sure all traces emitted during world EndPlay get processed before disabling channels
-	FEditorDelegates::ShutdownPIE.AddRaw(this, &SStateTreeDebuggerView::OnPIEStopped);
-	FEditorDelegates::PausePIE.AddRaw(this, &SStateTreeDebuggerView::OnPIEPaused);
-	FEditorDelegates::ResumePIE.AddRaw(this, &SStateTreeDebuggerView::OnPIEResumed);
-	FEditorDelegates::SingleStepPIE.AddRaw(this, &SStateTreeDebuggerView::OnPIESingleStepped);
+	ShutdownPIEHandle = FEditorDelegates::ShutdownPIE.AddSP(this, &SStateTreeDebuggerView::OnPIEStopped);
 }
 
 SStateTreeDebuggerView::~SStateTreeDebuggerView()
@@ -217,11 +217,11 @@ SStateTreeDebuggerView::~SStateTreeDebuggerView()
 
 	Debugger->ClearSelection();
 
-	FEditorDelegates::BeginPIE.RemoveAll(this);
-	FEditorDelegates::ShutdownPIE.RemoveAll(this);
-	FEditorDelegates::PausePIE.RemoveAll(this);
-	FEditorDelegates::ResumePIE.RemoveAll(this);
-	FEditorDelegates::SingleStepPIE.RemoveAll(this);
+	FEditorDelegates::BeginPIE.Remove(BeginPIEHandle);
+	FEditorDelegates::PausePIE.Remove(PausePIEHandle);
+	FEditorDelegates::ResumePIE.Remove(ResumePIEHandle);
+	FEditorDelegates::SingleStepPIE.Remove(SingleStepPIEHandle);
+	FEditorDelegates::ShutdownPIE.Remove(ShutdownPIEHandle);
 }
 
 void SStateTreeDebuggerView::StartRecording()
diff --git a/Engine/Plugins/Runtime/StateTree/Source/StateTreeEditorModule/Private/Debugger/SStateTreeDebuggerView.h b/Engine/Plugins/Runtime/StateTree/Source/StateTreeEditorModule/Private/Debugger/SStateTreeDebuggerView.h
--- a/Engine/Plugins/Runtime/StateTree/Source/StateTreeEditorModule/Private/Debugger/SStateTreeDebuggerView.h
+++ b/Engine/Plugins/Runtime/StateTree/Source/StateTreeEditorModule/Private/Debugger/SStateTreeDebuggerView.h
@@ -170,21 +170,28 @@ private:
 	TAttribute<double> ScrubTimeAttribute;
 
 	/** Range controlled by the timeline widgets and used to adjust cursor position and track content. */
-	TRange<double> ViewRange = TRange<double>(0, 10);
+	TRange<double> ViewRange = TRange<double>(0.0, 10.0);
 
 	/** In case tracks are not reset when a new analysis session is started we keep track of the longest duration to adjust our clamp range. */
-	double MaxTrackRecordingDuration = 0;
+	double MaxTrackRecordingDuration = 0.0;
 
 	/** The recording duration time the UI was last updated from the debugger. Used to detect if new data has been collect while the UI was inactive. */
-	double LastUpdatedTrackRecordingDuration = 0;
+	double LastUpdatedTrackRecordingDuration = 0.0;
 
 	/**
 	 * Extrapolated world time used to make the timeline progress until the debugger
 	 * receives a new event that will update its last processed world time.
 	 */
-	double ExtrapolatedRecordedWorldTime = 0;
+	double ExtrapolatedRecordedWorldTime = 0.0;
 
 private:
+
+	FDelegateHandle BeginPIEHandle;
+	FDelegateHandle PausePIEHandle;
+	FDelegateHandle ResumePIEHandle;
+	FDelegateHandle SingleStepPIEHandle;
+	FDelegateHandle ShutdownPIEHandle;
+
 	/** Indicates that a live session was started (record button or auto record in PIE) to generate StateTree traces. */
 	bool bRecording = false;
 
diff --git a/Engine/Plugins/Runtime/StateTree/Source/StateTreeModule/Private/Debugger/StateTreeDebugger.cpp b/Engine/Plugins/Runtime/StateTree/Source/StateTreeModule/Private/Debugger/StateTreeDebugger.cpp
--- a/Engine/Plugins/Runtime/StateTree/Source/StateTreeModule/Private/Debugger/StateTreeDebugger.cpp
+++ b/Engine/Plugins/Runtime/StateTree/Source/StateTreeModule/Private/Debugger/StateTreeDebugger.cpp
@@ -132,7 +132,12 @@ void FStateTreeDebugger::StopSessionAnalysis()
 	if (IsAnalysisSessionActive())
 	{
 		// Force one last update to process events emitted while closing the game session (e.g., EndPlay in PIE)
-		SyncToCurrentSessionDuration();
+		// Note that we can only perform this if the associated asset is still loaded which might not be the case
+		// during the standalone game shutdown.
+		if (StateTreeAsset.IsValid())
+		{
+			SyncToCurrentSessionDuration();
+		}
 		AnalysisSession->Stop(/*WaitOnAnalysis*/true);
 	}
 
diff --git a/Engine/Plugins/Runtime/StateTree/Source/StateTreeModule/Public/Debugger/StateTreeDebugger.h b/Engine/Plugins/Runtime/StateTree/Source/StateTreeModule/Public/Debugger/StateTreeDebugger.h
--- a/Engine/Plugins/Runtime/StateTree/Source/StateTreeModule/Public/Debugger/StateTreeDebugger.h
+++ b/Engine/Plugins/Runtime/StateTree/Source/StateTreeModule/Public/Debugger/StateTreeDebugger.h
@@ -234,6 +234,12 @@ struct FStateTreeDebugger : FTickableGameObject, UE::StateTreeDebugger::ITraceRe
 		return LastProcessedRecordedWorldTime;
 	}
 
+	UE_DEPRECATED(5.7, "Use GetLastProcessedRecordedWorldTime instead")
+	double GetRecordingDuration() const
+	{
+		return GetLastProcessedRecordedWorldTime();
+	}
+
 	/** Returns the duration of the analysis session. This is not related to the world simulation time. */
 	double GetAnalysisDuration() const
 	{
