[StateTreeDebugger] fixed delegate registration

diff --git a/Engine/Plugins/Runtime/StateTree/Source/StateTreeEditorModule/Private/Debugger/SStateTreeDebuggerView.cpp b/Engine/Plugins/Runtime/StateTree/Source/StateTreeEditorModule/Private/Debugger/SStateTreeDebuggerView.cpp
--- a/Engine/Plugins/Runtime/StateTree/Source/StateTreeEditorModule/Private/Debugger/SStateTreeDebuggerView.cpp
+++ b/Engine/Plugins/Runtime/StateTree/Source/StateTreeEditorModule/Private/Debugger/SStateTreeDebuggerView.cpp
@@ -194,16 +194,6 @@ void GenerateElementForProperties(const TCHAR* TypeAsText, const TCHAR* ValueAsT
 //----------------------------------------------------------------------//
 // SStateTreeDebuggerView
 //----------------------------------------------------------------------//
-SStateTreeDebuggerView::SStateTreeDebuggerView()
-{
-	BeginPIEHandle = FEditorDelegates::BeginPIE.AddSP(this, &SStateTreeDebuggerView::OnPIEStarted);
-	PausePIEHandle = FEditorDelegates::PausePIE.AddSP(this, &SStateTreeDebuggerView::OnPIEPaused);
-	ResumePIEHandle = FEditorDelegates::ResumePIE.AddSP(this, &SStateTreeDebuggerView::OnPIEResumed);
-	SingleStepPIEHandle =FEditorDelegates::SingleStepPIE.AddSP(this, &SStateTreeDebuggerView::OnPIESingleStepped);
-	// Using ShutdownPIE instead of EndPIE to make sure all traces emitted during world EndPlay get processed before disabling channels
-	ShutdownPIEHandle = FEditorDelegates::ShutdownPIE.AddSP(this, &SStateTreeDebuggerView::OnPIEStopped);
-}
-
 SStateTreeDebuggerView::~SStateTreeDebuggerView()
 {
 	UE::StateTree::Delegates::OnTracingStateChanged.RemoveAll(this);
@@ -345,6 +335,13 @@ void SStateTreeDebuggerView::OnPIESingleStepped(bool bSimulating) const
 
 void SStateTreeDebuggerView::Construct(const FArguments& InArgs, const UStateTree& InStateTree, const TSharedRef<FStateTreeViewModel>& InStateTreeViewModel, const TSharedRef<FUICommandList>& InCommandList)
 {
+	BeginPIEHandle = FEditorDelegates::BeginPIE.AddSP(this, &SStateTreeDebuggerView::OnPIEStarted);
+	PausePIEHandle = FEditorDelegates::PausePIE.AddSP(this, &SStateTreeDebuggerView::OnPIEPaused);
+	ResumePIEHandle = FEditorDelegates::ResumePIE.AddSP(this, &SStateTreeDebuggerView::OnPIEResumed);
+	SingleStepPIEHandle =FEditorDelegates::SingleStepPIE.AddSP(this, &SStateTreeDebuggerView::OnPIESingleStepped);
+	// Using ShutdownPIE instead of EndPIE to make sure all traces emitted during world EndPlay get processed before disabling channels
+	ShutdownPIEHandle = FEditorDelegates::ShutdownPIE.AddSP(this, &SStateTreeDebuggerView::OnPIEStopped);
+
 	StateTreeViewModel = InStateTreeViewModel;
 	StateTree = &InStateTree;
 	StateTreeEditorData = Cast<UStateTreeEditorData>(InStateTree.EditorData.Get());
diff --git a/Engine/Plugins/Runtime/StateTree/Source/StateTreeEditorModule/Private/Debugger/SStateTreeDebuggerView.h b/Engine/Plugins/Runtime/StateTree/Source/StateTreeEditorModule/Private/Debugger/SStateTreeDebuggerView.h
--- a/Engine/Plugins/Runtime/StateTree/Source/StateTreeEditorModule/Private/Debugger/SStateTreeDebuggerView.h
+++ b/Engine/Plugins/Runtime/StateTree/Source/StateTreeEditorModule/Private/Debugger/SStateTreeDebuggerView.h
@@ -43,7 +43,6 @@ public:
 	SLATE_BEGIN_ARGS(SStateTreeDebuggerView) {}
 	SLATE_END_ARGS()
 
-	SStateTreeDebuggerView();
 	~SStateTreeDebuggerView();
 
 	void Construct(const FArguments& InArgs, const UStateTree& InStateTree, const TSharedRef<FStateTreeViewModel>& InStateTreeViewModel, const TSharedRef<FUICommandList>& InCommandList);
