[AutoRTFM] Add a "Supports AutoRTFM" flag to various StateTree tests.

This flag will be functional in the `EngineTestAutoRTFM` target. That target will be created once https://p4-swarm.epicgames.net/reviews/42944957 lands. In regular `EngineTest`, this flag is a no-op.

diff --git a/Engine/Plugins/Runtime/StateTree/Source/StateTreeTestSuite/Private/StateTreeBindingTest.cpp b/Engine/Plugins/Runtime/StateTree/Source/StateTreeTestSuite/Private/StateTreeBindingTest.cpp
--- a/Engine/Plugins/Runtime/StateTree/Source/StateTreeTestSuite/Private/StateTreeBindingTest.cpp
+++ b/Engine/Plugins/Runtime/StateTree/Source/StateTreeTestSuite/Private/StateTreeBindingTest.cpp
@@ -149,7 +149,7 @@ struct FStateTreeTest_BindingsCompiler : FStateTreeTestBase
 		return true;
 	}
 };
-IMPLEMENT_AI_INSTANT_TEST(FStateTreeTest_BindingsCompiler, "System.StateTree.Binding.BindingsCompiler");
+IMPLEMENT_STATE_TREE_INSTANT_TEST(FStateTreeTest_BindingsCompiler, "System.StateTree.Binding.BindingsCompiler");
 
 struct FStateTreeTest_PropertyFunctions : FStateTreeTestBase
 {
@@ -226,7 +226,7 @@ struct FStateTreeTest_PropertyFunctions : FStateTreeTestBase
 		return true;
 	}
 };
-IMPLEMENT_AI_INSTANT_TEST(FStateTreeTest_PropertyFunctions, "System.StateTree.Binding.PropertyFunctions");
+IMPLEMENT_STATE_TREE_INSTANT_TEST(FStateTreeTest_PropertyFunctions, "System.StateTree.Binding.PropertyFunctions");
 
 struct FStateTreeTest_CopyObjects : FStateTreeTestBase
 {
@@ -338,7 +338,12 @@ struct FStateTreeTest_CopyObjects : FStateTreeTestBase
 		ObjectA = nullptr;
 		ObjectB = nullptr;
 		Source.Object = nullptr;
-		CollectGarbage(GARBAGE_COLLECTION_KEEPFLAGS);
+
+		// Run a garbage collection pass here when possible. (Garbage collection isn't allowed inside an AutoRTFM transaction.)
+		if (!AutoRTFM::IsTransactional())
+		{
+			CollectGarbage(GARBAGE_COLLECTION_KEEPFLAGS);
+		}
 
 		FStateTreeTest_PropertyCopyObjects TargetC;
 		bool bCopyResultC = true;
@@ -354,7 +359,7 @@ struct FStateTreeTest_CopyObjects : FStateTreeTestBase
 		return true;
 	}
 };
-IMPLEMENT_AI_INSTANT_TEST(FStateTreeTest_CopyObjects, "System.StateTree.Binding.CopyObjects");
+IMPLEMENT_STATE_TREE_INSTANT_TEST(FStateTreeTest_CopyObjects, "System.StateTree.Binding.CopyObjects");
 
 
 struct FStateTreeTest_References : FStateTreeTestBase
@@ -432,7 +437,7 @@ struct FStateTreeTest_References : FStateTreeTestBase
 		return true;
 	}
 };
-IMPLEMENT_AI_INSTANT_TEST(FStateTreeTest_References, "System.StateTree.Binding.References");
+IMPLEMENT_STATE_TREE_INSTANT_TEST(FStateTreeTest_References, "System.StateTree.Binding.References");
 
 struct FStateTreeTest_ReferencesConstness : FStateTreeTestBase
 {
@@ -510,7 +515,7 @@ struct FStateTreeTest_ReferencesConstness : FStateTreeTestBase
 		return true;
 	}
 };
-IMPLEMENT_AI_INSTANT_TEST(FStateTreeTest_ReferencesConstness, "System.StateTree.Binding.ReferencesConstness");
+IMPLEMENT_STATE_TREE_INSTANT_TEST(FStateTreeTest_ReferencesConstness, "System.StateTree.Binding.ReferencesConstness");
 
 struct FStateTreeTest_MutableArray : FStateTreeTestBase
 {
@@ -623,7 +628,7 @@ struct FStateTreeTest_MutableArray : FStateTreeTestBase
 		return true;
 	}
 };
-IMPLEMENT_AI_INSTANT_TEST(FStateTreeTest_MutableArray, "System.StateTree.Binding.MutableArray");
+IMPLEMENT_STATE_TREE_INSTANT_TEST(FStateTreeTest_MutableArray, "System.StateTree.Binding.MutableArray");
 
 struct FStateTreeTest_TransitionTaskWithBinding : FStateTreeTestBase
 {
@@ -753,7 +758,7 @@ struct FStateTreeTest_TransitionTaskWithBinding : FStateTreeTestBase
 		return true;
 	}
 };
-IMPLEMENT_AI_INSTANT_TEST(FStateTreeTest_TransitionTaskWithBinding, "System.StateTree.Binding.TransitionTaskWithBinding");
+IMPLEMENT_STATE_TREE_INSTANT_TEST(FStateTreeTest_TransitionTaskWithBinding, "System.StateTree.Binding.TransitionTaskWithBinding");
 
 } // namespace UE::StateTree::Tests
 
diff --git a/Engine/Plugins/Runtime/StateTree/Source/StateTreeTestSuite/Private/StateTreeDelegateTest.cpp b/Engine/Plugins/Runtime/StateTree/Source/StateTreeTestSuite/Private/StateTreeDelegateTest.cpp
--- a/Engine/Plugins/Runtime/StateTree/Source/StateTreeTestSuite/Private/StateTreeDelegateTest.cpp
+++ b/Engine/Plugins/Runtime/StateTree/Source/StateTreeTestSuite/Private/StateTreeDelegateTest.cpp
@@ -61,7 +61,7 @@ struct FStateTreeTest_Delegate_ConcurrentListeners : FStateTreeTestBase
 		return true;
 	}
 };
-IMPLEMENT_AI_INSTANT_TEST(FStateTreeTest_Delegate_ConcurrentListeners, "System.StateTree.Delegate.ConcurrentListeners");
+IMPLEMENT_STATE_TREE_INSTANT_TEST(FStateTreeTest_Delegate_ConcurrentListeners, "System.StateTree.Delegate.ConcurrentListeners");
 
 struct FStateTreeTest_Delegate_MutuallyExclusiveListeners : FStateTreeTestBase
 {
@@ -119,7 +119,7 @@ struct FStateTreeTest_Delegate_MutuallyExclusiveListeners : FStateTreeTestBase
 		return true;
 	}
 };
-IMPLEMENT_AI_INSTANT_TEST(FStateTreeTest_Delegate_MutuallyExclusiveListeners, "System.StateTree.Delegate.MutuallyExclusiveListeners");
+IMPLEMENT_STATE_TREE_INSTANT_TEST(FStateTreeTest_Delegate_MutuallyExclusiveListeners, "System.StateTree.Delegate.MutuallyExclusiveListeners");
 
 struct FStateTreeTest_Delegate_Transitions : FStateTreeTestBase
 {
@@ -168,7 +168,7 @@ struct FStateTreeTest_Delegate_Transitions : FStateTreeTestBase
 		return true;
 	}
 };
-IMPLEMENT_AI_INSTANT_TEST(FStateTreeTest_Delegate_Transitions, "System.StateTree.Delegate.Transitions");
+IMPLEMENT_STATE_TREE_INSTANT_TEST(FStateTreeTest_Delegate_Transitions, "System.StateTree.Delegate.Transitions");
 
 struct FStateTreeTest_Delegate_Rebroadcasting : FStateTreeTestBase
 {
@@ -211,7 +211,7 @@ struct FStateTreeTest_Delegate_Rebroadcasting : FStateTreeTestBase
 		return true;
 	}
 };
-IMPLEMENT_AI_INSTANT_TEST(FStateTreeTest_Delegate_Rebroadcasting, "System.StateTree.Delegate.Rebroadcasting");
+IMPLEMENT_STATE_TREE_INSTANT_TEST(FStateTreeTest_Delegate_Rebroadcasting, "System.StateTree.Delegate.Rebroadcasting");
 
 struct FStateTreeTest_Delegate_SelfRemoval : FStateTreeTestBase
 {
@@ -259,7 +259,7 @@ struct FStateTreeTest_Delegate_SelfRemoval : FStateTreeTestBase
 		return true;
 	}
 };
-IMPLEMENT_AI_INSTANT_TEST(FStateTreeTest_Delegate_SelfRemoval, "System.StateTree.Delegate.SelfRemoval");
+IMPLEMENT_STATE_TREE_INSTANT_TEST(FStateTreeTest_Delegate_SelfRemoval, "System.StateTree.Delegate.SelfRemoval");
 
 struct FStateTreeTest_Delegate_WithoutRemoval : FStateTreeTestBase
 {
@@ -308,7 +308,7 @@ struct FStateTreeTest_Delegate_WithoutRemoval : FStateTreeTestBase
 		return true;
 	}
 };
-IMPLEMENT_AI_INSTANT_TEST(FStateTreeTest_Delegate_WithoutRemoval, "System.StateTree.Delegate.WithoutRemoval");
+IMPLEMENT_STATE_TREE_INSTANT_TEST(FStateTreeTest_Delegate_WithoutRemoval, "System.StateTree.Delegate.WithoutRemoval");
 
 struct FStateTreeTest_Delegate_GlobalDispatcherAndListener : FStateTreeTestBase
 {
@@ -352,7 +352,7 @@ struct FStateTreeTest_Delegate_GlobalDispatcherAndListener : FStateTreeTestBase
 		return true;
 	}
 };
-IMPLEMENT_AI_INSTANT_TEST(FStateTreeTest_Delegate_GlobalDispatcherAndListener, "System.StateTree.Delegate.GlobalDispatcherAndListener");
+IMPLEMENT_STATE_TREE_INSTANT_TEST(FStateTreeTest_Delegate_GlobalDispatcherAndListener, "System.StateTree.Delegate.GlobalDispatcherAndListener");
 
 struct FStateTreeTest_Delegate_ListeningToDelegateOnExit : FStateTreeTestBase
 {
@@ -389,7 +389,7 @@ struct FStateTreeTest_Delegate_ListeningToDelegateOnExit : FStateTreeTestBase
 		return true;
 	}
 };
-IMPLEMENT_AI_INSTANT_TEST(FStateTreeTest_Delegate_ListeningToDelegateOnExit, "System.StateTree.Delegate.ListeningToDelegateOnExit");
+IMPLEMENT_STATE_TREE_INSTANT_TEST(FStateTreeTest_Delegate_ListeningToDelegateOnExit, "System.StateTree.Delegate.ListeningToDelegateOnExit");
 
 } // namespace UE::StateTree::Tests
 
diff --git a/Engine/Plugins/Runtime/StateTree/Source/StateTreeTestSuite/Private/StateTreeLinkedStateTest.cpp b/Engine/Plugins/Runtime/StateTree/Source/StateTreeTestSuite/Private/StateTreeLinkedStateTest.cpp
--- a/Engine/Plugins/Runtime/StateTree/Source/StateTreeTestSuite/Private/StateTreeLinkedStateTest.cpp
+++ b/Engine/Plugins/Runtime/StateTree/Source/StateTreeTestSuite/Private/StateTreeLinkedStateTest.cpp
@@ -107,7 +107,7 @@ struct FStateTreeTest_FailEnterLinkedAsset : FStateTreeTestBase
 		return true;
 	}
 };
-IMPLEMENT_AI_INSTANT_TEST(FStateTreeTest_FailEnterLinkedAsset, "System.StateTree.LinkedAsset.FailEnter");
+IMPLEMENT_STATE_TREE_INSTANT_TEST(FStateTreeTest_FailEnterLinkedAsset, "System.StateTree.LinkedAsset.FailEnter");
 
 struct FStateTreeTest_EnterAndExitLinkedAsset : FStateTreeTestBase
 {
@@ -234,7 +234,7 @@ struct FStateTreeTest_EnterAndExitLinkedAsset : FStateTreeTestBase
 		return true;
 	}
 };
-IMPLEMENT_AI_INSTANT_TEST(FStateTreeTest_EnterAndExitLinkedAsset, "System.StateTree.LinkedAsset.EnterAndExit");
+IMPLEMENT_STATE_TREE_INSTANT_TEST(FStateTreeTest_EnterAndExitLinkedAsset, "System.StateTree.LinkedAsset.EnterAndExit");
 
 struct FStateTreeTest_EnterAndExitLinkedAsset2 : FStateTreeTestBase
 {
@@ -378,7 +378,7 @@ struct FStateTreeTest_EnterAndExitLinkedAsset2 : FStateTreeTestBase
 		return true;
 	}
 };
-IMPLEMENT_AI_INSTANT_TEST(FStateTreeTest_EnterAndExitLinkedAsset2, "System.StateTree.LinkedAsset.EnterAndExit2");
+IMPLEMENT_STATE_TREE_INSTANT_TEST(FStateTreeTest_EnterAndExitLinkedAsset2, "System.StateTree.LinkedAsset.EnterAndExit2");
 
 struct FStateTreeTest_MultipleSameLinkedAsset : FStateTreeTestBase
 {
@@ -680,7 +680,7 @@ struct FStateTreeTest_MultipleSameLinkedAsset : FStateTreeTestBase
 		return true;
 	}
 };
-IMPLEMENT_AI_INSTANT_TEST(FStateTreeTest_MultipleSameLinkedAsset, "System.StateTree.LinkedAsset.MultipleSameTree");
+IMPLEMENT_STATE_TREE_INSTANT_TEST(FStateTreeTest_MultipleSameLinkedAsset, "System.StateTree.LinkedAsset.MultipleSameTree");
 
 struct FStateTreeTest_EmptyStateWithTickTransitionLinkedAsset : FStateTreeTestBase
 {
@@ -926,12 +926,18 @@ struct FStateTreeTest_EmptyStateWithTickTransitionLinkedAsset : FStateTreeTestBa
 		return true;
 	}
 };
-IMPLEMENT_AI_INSTANT_TEST(FStateTreeTest_EmptyStateWithTickTransitionLinkedAsset, "System.StateTree.LinkedAsset.EmptyStateWithTickTransition");
+IMPLEMENT_STATE_TREE_INSTANT_TEST(FStateTreeTest_EmptyStateWithTickTransitionLinkedAsset, "System.StateTree.LinkedAsset.EmptyStateWithTickTransition");
 
 struct FStateTreeTest_RecursiveLinkedAsset : FStateTreeTestBase
 {
 	virtual bool InstantTest() override
 	{
+		if (AutoRTFM::IsTransactional())
+		{
+			// AddExpectedError is currently not supported in AutoRTFM.
+			return true;
+		}
+
 		//Tree 1
 		//	Root
 		//		StateLinkedTree1 (Tree2) -> Next
@@ -1058,7 +1064,8 @@ struct FStateTreeTest_RecursiveLinkedAsset : FStateTreeTestBase
 		return true;
 	}
 };
-IMPLEMENT_AI_INSTANT_TEST(FStateTreeTest_RecursiveLinkedAsset, "System.StateTree.LinkedAsset.RecursiveLinkedAsset");
+// This test invokes editor-only methods which aren't AutoRTFM-safe, so we skip it when AutoRTFM is on.
+IMPLEMENT_STATE_TREE_INSTANT_TEST(FStateTreeTest_RecursiveLinkedAsset, "System.StateTree.LinkedAsset.RecursiveLinkedAsset");
 
 struct FStateTreeTest_LinkedAssetTransitionSameTick : FStateTreeTestBase
 {
@@ -1199,7 +1206,7 @@ struct FStateTreeTest_LinkedAssetTransitionSameTick : FStateTreeTestBase
 		return true;
 	}
 };
-IMPLEMENT_AI_INSTANT_TEST(FStateTreeTest_LinkedAssetTransitionSameTick, "System.StateTree.LinkedAsset.TransitionSameTick");
+IMPLEMENT_STATE_TREE_INSTANT_TEST(FStateTreeTest_LinkedAssetTransitionSameTick, "System.StateTree.LinkedAsset.TransitionSameTick");
 
 struct FStateTreeTest_Linked_GlobalParameter : FStateTreeTestBase
 {
@@ -1494,7 +1501,7 @@ struct FStateTreeTest_Linked_GlobalParameter : FStateTreeTestBase
 		return true;
 	}
 };
-IMPLEMENT_AI_INSTANT_TEST(FStateTreeTest_Linked_GlobalParameter, "System.StateTree.LinkedAsset.GlobalParameter");
+IMPLEMENT_STATE_TREE_INSTANT_TEST(FStateTreeTest_Linked_GlobalParameter, "System.StateTree.LinkedAsset.GlobalParameter");
 
 struct FStateTreeTest_Linked_FinishGlobalTasks : FStateTreeTestBase
 {
@@ -1685,7 +1692,7 @@ struct FStateTreeTest_Linked_FinishGlobalTasks : FStateTreeTestBase
 		return true;
 	}
 };
-IMPLEMENT_AI_INSTANT_TEST(FStateTreeTest_Linked_FinishGlobalTasks, "System.StateTree.LinkedAsset.FinishGlobalTasks");
+IMPLEMENT_STATE_TREE_INSTANT_TEST(FStateTreeTest_Linked_FinishGlobalTasks, "System.StateTree.LinkedAsset.FinishGlobalTasks");
 
 } // namespace UE::StateTree::Tests
 
diff --git a/Engine/Plugins/Runtime/StateTree/Source/StateTreeTestSuite/Private/StateTreePropertyPathTest.cpp b/Engine/Plugins/Runtime/StateTree/Source/StateTreeTestSuite/Private/StateTreePropertyPathTest.cpp
--- a/Engine/Plugins/Runtime/StateTree/Source/StateTreeTestSuite/Private/StateTreePropertyPathTest.cpp
+++ b/Engine/Plugins/Runtime/StateTree/Source/StateTreeTestSuite/Private/StateTreePropertyPathTest.cpp
@@ -38,7 +38,7 @@ struct FStateTreeTest_PropertyPathOffset : FStateTreeTestBase
 		return true;
 	}
 };
-IMPLEMENT_AI_INSTANT_TEST(FStateTreeTest_PropertyPathOffset, "System.StateTree.PropertyPath.Offset");
+IMPLEMENT_STATE_TREE_INSTANT_TEST(FStateTreeTest_PropertyPathOffset, "System.StateTree.PropertyPath.Offset");
 
 struct FStateTreeTest_PropertyPathParseFail : FStateTreeTestBase
 {
@@ -77,7 +77,7 @@ struct FStateTreeTest_PropertyPathParseFail : FStateTreeTestBase
 		return true;
 	}
 };
-IMPLEMENT_AI_INSTANT_TEST(FStateTreeTest_PropertyPathParseFail, "System.StateTree.PropertyPath.ParseFail");
+IMPLEMENT_STATE_TREE_INSTANT_TEST(FStateTreeTest_PropertyPathParseFail, "System.StateTree.PropertyPath.ParseFail");
 
 struct FStateTreeTest_PropertyPathOffsetFail : FStateTreeTestBase
 {
@@ -101,7 +101,7 @@ struct FStateTreeTest_PropertyPathOffsetFail : FStateTreeTestBase
 		return true;
 	}
 };
-IMPLEMENT_AI_INSTANT_TEST(FStateTreeTest_PropertyPathOffsetFail, "System.StateTree.PropertyPath.OffsetFail");
+IMPLEMENT_STATE_TREE_INSTANT_TEST(FStateTreeTest_PropertyPathOffsetFail, "System.StateTree.PropertyPath.OffsetFail");
 
 struct FStateTreeTest_PropertyPathObject : FStateTreeTestBase
 {
@@ -125,7 +125,7 @@ struct FStateTreeTest_PropertyPathObject : FStateTreeTestBase
 		return true;
 	}
 };
-IMPLEMENT_AI_INSTANT_TEST(FStateTreeTest_PropertyPathObject, "System.StateTree.PropertyPath.Object");
+IMPLEMENT_STATE_TREE_INSTANT_TEST(FStateTreeTest_PropertyPathObject, "System.StateTree.PropertyPath.Object");
 
 struct FStateTreeTest_PropertyPathWrongObject : FStateTreeTestBase
 {
@@ -167,7 +167,7 @@ struct FStateTreeTest_PropertyPathWrongObject : FStateTreeTestBase
 		return true;
 	}
 };
-IMPLEMENT_AI_INSTANT_TEST(FStateTreeTest_PropertyPathWrongObject, "System.StateTree.PropertyPath.WrongObject");
+IMPLEMENT_STATE_TREE_INSTANT_TEST(FStateTreeTest_PropertyPathWrongObject, "System.StateTree.PropertyPath.WrongObject");
 
 struct FStateTreeTest_PropertyPathArray : FStateTreeTestBase
 {
@@ -199,7 +199,7 @@ struct FStateTreeTest_PropertyPathArray : FStateTreeTestBase
 		return true;
 	}
 };
-IMPLEMENT_AI_INSTANT_TEST(FStateTreeTest_PropertyPathArray, "System.StateTree.PropertyPath.Array");
+IMPLEMENT_STATE_TREE_INSTANT_TEST(FStateTreeTest_PropertyPathArray, "System.StateTree.PropertyPath.Array");
 
 struct FStateTreeTest_PropertyPathArrayInvalidIndex : FStateTreeTestBase
 {
@@ -224,7 +224,7 @@ struct FStateTreeTest_PropertyPathArrayInvalidIndex : FStateTreeTestBase
 		return true;
 	}
 };
-IMPLEMENT_AI_INSTANT_TEST(FStateTreeTest_PropertyPathArrayInvalidIndex, "System.StateTree.PropertyPath.ArrayInvalidIndex");
+IMPLEMENT_STATE_TREE_INSTANT_TEST(FStateTreeTest_PropertyPathArrayInvalidIndex, "System.StateTree.PropertyPath.ArrayInvalidIndex");
 
 struct FStateTreeTest_PropertyPathArrayOfStructs : FStateTreeTestBase
 {
@@ -276,7 +276,7 @@ struct FStateTreeTest_PropertyPathArrayOfStructs : FStateTreeTestBase
 		return true;
 	}
 };
-IMPLEMENT_AI_INSTANT_TEST(FStateTreeTest_PropertyPathArrayOfStructs, "System.StateTree.PropertyPath.ArrayOfStructs");
+IMPLEMENT_STATE_TREE_INSTANT_TEST(FStateTreeTest_PropertyPathArrayOfStructs, "System.StateTree.PropertyPath.ArrayOfStructs");
 
 struct FStateTreeTest_PropertyPathArrayOfInstancedObjects : FStateTreeTestBase
 {
@@ -329,7 +329,7 @@ struct FStateTreeTest_PropertyPathArrayOfInstancedObjects : FStateTreeTestBase
 		return true;
 	}
 };
-IMPLEMENT_AI_INSTANT_TEST(FStateTreeTest_PropertyPathArrayOfInstancedObjects, "System.StateTree.PropertyPath.ArrayOfInstancedObjects");
+IMPLEMENT_STATE_TREE_INSTANT_TEST(FStateTreeTest_PropertyPathArrayOfInstancedObjects, "System.StateTree.PropertyPath.ArrayOfInstancedObjects");
 
 } // namespace UE::StateTree::Tests
 
diff --git a/Engine/Plugins/Runtime/StateTree/Source/StateTreeTestSuite/Private/StateTreeStatePathTest.cpp b/Engine/Plugins/Runtime/StateTree/Source/StateTreeTestSuite/Private/StateTreeStatePathTest.cpp
--- a/Engine/Plugins/Runtime/StateTree/Source/StateTreeTestSuite/Private/StateTreeStatePathTest.cpp
+++ b/Engine/Plugins/Runtime/StateTree/Source/StateTreeTestSuite/Private/StateTreeStatePathTest.cpp
@@ -445,7 +445,7 @@ struct FStateTreeTest_StatePath_LinkStates : FStateTreeTestBase
 		return true;
 	}
 };
-IMPLEMENT_AI_INSTANT_TEST(FStateTreeTest_StatePath_LinkStates, "System.StateTree.StatePath.LinkStates");
+IMPLEMENT_STATE_TREE_INSTANT_TEST(FStateTreeTest_StatePath_LinkStates, "System.StateTree.StatePath.LinkStates");
 
 } // namespace UE::StateTree::Tests
 
diff --git a/Engine/Plugins/Runtime/StateTree/Source/StateTreeTestSuite/Private/StateTreeTaskStateTest.cpp b/Engine/Plugins/Runtime/StateTree/Source/StateTreeTestSuite/Private/StateTreeTaskStateTest.cpp
--- a/Engine/Plugins/Runtime/StateTree/Source/StateTreeTestSuite/Private/StateTreeTaskStateTest.cpp
+++ b/Engine/Plugins/Runtime/StateTree/Source/StateTreeTestSuite/Private/StateTreeTaskStateTest.cpp
@@ -170,7 +170,7 @@ struct FStateTreeTest_TasksCompletion_AllAny : FStateTreeTestBase
 		return true;
 	}
 };
-IMPLEMENT_AI_INSTANT_TEST(FStateTreeTest_TasksCompletion_AllAny, "System.StateTree.TasksCompletion.AllAny");
+IMPLEMENT_STATE_TREE_INSTANT_TEST(FStateTreeTest_TasksCompletion_AllAny, "System.StateTree.TasksCompletion.AllAny");
 
 struct FStateTreeTest_TasksCompletion_FailureTasks : FStateTreeTestBase
 {
@@ -251,7 +251,7 @@ struct FStateTreeTest_TasksCompletion_FailureTasks : FStateTreeTestBase
 		return true;
 	}
 };
-IMPLEMENT_AI_INSTANT_TEST(FStateTreeTest_TasksCompletion_FailureTasks, "System.StateTree.TasksCompletion.Failure");
+IMPLEMENT_STATE_TREE_INSTANT_TEST(FStateTreeTest_TasksCompletion_FailureTasks, "System.StateTree.TasksCompletion.Failure");
 
 // test set status with priority
 struct FStateTreeTest_TasksCompletion_Status : FStateTreeTestBase
@@ -432,7 +432,7 @@ struct FStateTreeTest_TasksCompletion_Status : FStateTreeTestBase
 		return true;
 	}
 };
-IMPLEMENT_AI_INSTANT_TEST(FStateTreeTest_TasksCompletion_Status, "System.StateTree.TasksCompletion.Status");
+IMPLEMENT_STATE_TREE_INSTANT_TEST(FStateTreeTest_TasksCompletion_Status, "System.StateTree.TasksCompletion.Status");
 
 } // namespace UE::StateTree::Tests
 
diff --git a/Engine/Plugins/Runtime/StateTree/Source/StateTreeTestSuite/Private/StateTreeTest.cpp b/Engine/Plugins/Runtime/StateTree/Source/StateTreeTestSuite/Private/StateTreeTest.cpp
--- a/Engine/Plugins/Runtime/StateTree/Source/StateTreeTestSuite/Private/StateTreeTest.cpp
+++ b/Engine/Plugins/Runtime/StateTree/Source/StateTreeTestSuite/Private/StateTreeTest.cpp
@@ -73,7 +73,7 @@ struct FStateTreeTest_MakeAndBakeStateTree : FStateTreeTestBase
 		return true;
 	}
 };
-IMPLEMENT_AI_INSTANT_TEST(FStateTreeTest_MakeAndBakeStateTree, "System.StateTree.MakeAndBakeStateTree");
+IMPLEMENT_STATE_TREE_INSTANT_TEST(FStateTreeTest_MakeAndBakeStateTree, "System.StateTree.MakeAndBakeStateTree");
 
 
 struct FStateTreeTest_EmptyStateTree : FStateTreeTestBase
@@ -111,7 +111,7 @@ struct FStateTreeTest_EmptyStateTree : FStateTreeTestBase
 		return true;
 	}
 };
-IMPLEMENT_AI_INSTANT_TEST(FStateTreeTest_EmptyStateTree, "System.StateTree.Empty");
+IMPLEMENT_STATE_TREE_INSTANT_TEST(FStateTreeTest_EmptyStateTree, "System.StateTree.Empty");
 
 struct FStateTreeTest_Sequence : FStateTreeTestBase
 {
@@ -173,7 +173,7 @@ struct FStateTreeTest_Sequence : FStateTreeTestBase
 		return true;
 	}
 };
-IMPLEMENT_AI_INSTANT_TEST(FStateTreeTest_Sequence, "System.StateTree.Sequence");
+IMPLEMENT_STATE_TREE_INSTANT_TEST(FStateTreeTest_Sequence, "System.StateTree.Sequence");
 
 struct FStateTreeTest_Select : FStateTreeTestBase
 {
@@ -247,7 +247,7 @@ struct FStateTreeTest_Select : FStateTreeTestBase
 		return true;
 	}
 };
-IMPLEMENT_AI_INSTANT_TEST(FStateTreeTest_Select, "System.StateTree.Select");
+IMPLEMENT_STATE_TREE_INSTANT_TEST(FStateTreeTest_Select, "System.StateTree.Select");
 
 
 struct FStateTreeTest_FailEnterState : FStateTreeTestBase
@@ -323,7 +323,7 @@ struct FStateTreeTest_FailEnterState : FStateTreeTestBase
 		return true;
 	}
 };
-IMPLEMENT_AI_INSTANT_TEST(FStateTreeTest_FailEnterState, "System.StateTree.FailEnterState");
+IMPLEMENT_STATE_TREE_INSTANT_TEST(FStateTreeTest_FailEnterState, "System.StateTree.FailEnterState");
 
 
 struct FStateTreeTest_Restart : FStateTreeTestBase
@@ -378,7 +378,7 @@ struct FStateTreeTest_Restart : FStateTreeTestBase
 		return true;
 	}
 };
-IMPLEMENT_AI_INSTANT_TEST(FStateTreeTest_Restart, "System.StateTree.Restart");
+IMPLEMENT_STATE_TREE_INSTANT_TEST(FStateTreeTest_Restart, "System.StateTree.Restart");
 
 struct FStateTreeTest_SubTree_ActiveTasks : FStateTreeTestBase
 {
@@ -456,7 +456,7 @@ struct FStateTreeTest_SubTree_ActiveTasks : FStateTreeTestBase
 		return true;
 	}
 };
-IMPLEMENT_AI_INSTANT_TEST(FStateTreeTest_SubTree_ActiveTasks, "System.StateTree.SubTree.ActiveTasks");
+IMPLEMENT_STATE_TREE_INSTANT_TEST(FStateTreeTest_SubTree_ActiveTasks, "System.StateTree.SubTree.ActiveTasks");
 
 struct FStateTreeTest_SubTree_NoActiveTasks : FStateTreeTestBase
 {
@@ -537,7 +537,7 @@ struct FStateTreeTest_SubTree_NoActiveTasks : FStateTreeTestBase
 		return true;
 	}
 };
-IMPLEMENT_AI_INSTANT_TEST(FStateTreeTest_SubTree_NoActiveTasks, "System.StateTree.SubTree.NoActiveTasks");
+IMPLEMENT_STATE_TREE_INSTANT_TEST(FStateTreeTest_SubTree_NoActiveTasks, "System.StateTree.SubTree.NoActiveTasks");
 
 struct FStateTreeTest_SubTree_Condition : FStateTreeTestBase
 {
@@ -619,7 +619,7 @@ struct FStateTreeTest_SubTree_Condition : FStateTreeTestBase
 		return true;
 	}
 };
-IMPLEMENT_AI_INSTANT_TEST(FStateTreeTest_SubTree_Condition, "System.StateTree.SubTree.Condition");
+IMPLEMENT_STATE_TREE_INSTANT_TEST(FStateTreeTest_SubTree_Condition, "System.StateTree.SubTree.Condition");
 
 struct FStateTreeTest_SubTree_CascadedSucceeded : FStateTreeTestBase
 {
@@ -693,7 +693,7 @@ struct FStateTreeTest_SubTree_CascadedSucceeded : FStateTreeTestBase
 		return true;
 	}
 };
-IMPLEMENT_AI_INSTANT_TEST(FStateTreeTest_SubTree_CascadedSucceeded, "System.StateTree.SubTree.CascadedSucceeded");
+IMPLEMENT_STATE_TREE_INSTANT_TEST(FStateTreeTest_SubTree_CascadedSucceeded, "System.StateTree.SubTree.CascadedSucceeded");
 
 
 struct FStateTreeTest_SharedInstanceData : FStateTreeTestBase
@@ -819,7 +819,7 @@ struct FStateTreeTest_SharedInstanceData : FStateTreeTestBase
 		return true;
 	}
 };
-IMPLEMENT_AI_INSTANT_TEST(FStateTreeTest_SharedInstanceData, "System.StateTree.SharedInstanceData");
+IMPLEMENT_STATE_TREE_INSTANT_TEST(FStateTreeTest_SharedInstanceData, "System.StateTree.SharedInstanceData");
 
 struct FStateTreeTest_LastConditionWithIndent : FStateTreeTestBase
 {
@@ -874,7 +874,7 @@ struct FStateTreeTest_LastConditionWithIndent : FStateTreeTestBase
 		return true;
 	}
 };
-IMPLEMENT_AI_INSTANT_TEST(FStateTreeTest_LastConditionWithIndent, "System.StateTree.LastConditionWithIndent");
+IMPLEMENT_STATE_TREE_INSTANT_TEST(FStateTreeTest_LastConditionWithIndent, "System.StateTree.LastConditionWithIndent");
 
 struct FStateTreeTest_StateRequiringEvent : FStateTreeTestBase
 {
@@ -1026,7 +1026,7 @@ struct FStateTreeTest_StateRequiringEvent : FStateTreeTestBase
 		return true;
 	}
 };
-IMPLEMENT_AI_INSTANT_TEST(FStateTreeTest_StateRequiringEvent, "System.StateTree.StateRequiringEvent");
+IMPLEMENT_STATE_TREE_INSTANT_TEST(FStateTreeTest_StateRequiringEvent, "System.StateTree.StateRequiringEvent");
 
 struct FStateTreeTest_Start : FStateTreeTestBase
 {
@@ -1096,7 +1096,7 @@ protected:
 	const FName TaskAName = FName("TaskA");
 	const FName TaskBName = FName("TaskB");
 };
-IMPLEMENT_AI_INSTANT_TEST(FStateTreeTest_Start, "System.StateTree.Start.FirstStateSucceededImmediately");
+IMPLEMENT_STATE_TREE_INSTANT_TEST(FStateTreeTest_Start, "System.StateTree.Start.FirstStateSucceededImmediately");
 
 struct FStateTreeTest_StartScheduledTick : FStateTreeTest_Start
 {
@@ -1118,7 +1118,7 @@ struct FStateTreeTest_StartScheduledTick : FStateTreeTest_Start
 		return StateTree;
 	}
 };
-IMPLEMENT_AI_INSTANT_TEST(FStateTreeTest_StartScheduledTick, "System.StateTree.Start.FirstStateSucceededImmediatelyWithScheduledTick");
+IMPLEMENT_STATE_TREE_INSTANT_TEST(FStateTreeTest_StartScheduledTick, "System.StateTree.Start.FirstStateSucceededImmediatelyWithScheduledTick");
 
 //
 // The stop tests test how the combinations of execution path to stop the tree are reported on ExitState() transition.  
@@ -1225,7 +1225,7 @@ struct FStateTreeTest_Stop_NormalSucceeded : FStateTreeTest_Stop
 		return true;
 	}
 };
-IMPLEMENT_AI_INSTANT_TEST(FStateTreeTest_Stop_NormalSucceeded, "System.StateTree.Stop.NormalSucceeded");
+IMPLEMENT_STATE_TREE_INSTANT_TEST(FStateTreeTest_Stop_NormalSucceeded, "System.StateTree.Stop.NormalSucceeded");
 
 struct FStateTreeTest_Stop_NormalFailed : FStateTreeTest_Stop
 {
@@ -1247,7 +1247,7 @@ struct FStateTreeTest_Stop_NormalFailed : FStateTreeTest_Stop
 		return true;
 	}
 };
-IMPLEMENT_AI_INSTANT_TEST(FStateTreeTest_Stop_NormalFailed, "System.StateTree.Stop.NormalFailed");
+IMPLEMENT_STATE_TREE_INSTANT_TEST(FStateTreeTest_Stop_NormalFailed, "System.StateTree.Stop.NormalFailed");
 
 
 struct FStateTreeTest_Stop_GlobalSucceeded : FStateTreeTest_Stop
@@ -1270,7 +1270,7 @@ struct FStateTreeTest_Stop_GlobalSucceeded : FStateTreeTest_Stop
 		return true;
 	}
 };
-IMPLEMENT_AI_INSTANT_TEST(FStateTreeTest_Stop_GlobalSucceeded, "System.StateTree.Stop.GlobalSucceeded");
+IMPLEMENT_STATE_TREE_INSTANT_TEST(FStateTreeTest_Stop_GlobalSucceeded, "System.StateTree.Stop.GlobalSucceeded");
 
 struct FStateTreeTest_Stop_GlobalFailed : FStateTreeTest_Stop
 {
@@ -1292,7 +1292,7 @@ struct FStateTreeTest_Stop_GlobalFailed : FStateTreeTest_Stop
 		return true;
 	}
 };
-IMPLEMENT_AI_INSTANT_TEST(FStateTreeTest_Stop_GlobalFailed, "System.StateTree.Stop.GlobalFailed");
+IMPLEMENT_STATE_TREE_INSTANT_TEST(FStateTreeTest_Stop_GlobalFailed, "System.StateTree.Stop.GlobalFailed");
 
 
 //
@@ -1363,7 +1363,7 @@ struct FStateTreeTest_Stop_NormalEnterSucceeded : FStateTreeTest_StopEnterNormal
 		return true;
 	}
 };
-IMPLEMENT_AI_INSTANT_TEST(FStateTreeTest_Stop_NormalEnterSucceeded, "System.StateTree.Stop.NormalEnterSucceeded");
+IMPLEMENT_STATE_TREE_INSTANT_TEST(FStateTreeTest_Stop_NormalEnterSucceeded, "System.StateTree.Stop.NormalEnterSucceeded");
 
 struct FStateTreeTest_Stop_NormalEnterFailed : FStateTreeTest_StopEnterNormal
 {
@@ -1385,7 +1385,7 @@ struct FStateTreeTest_Stop_NormalEnterFailed : FStateTreeTest_StopEnterNormal
 		return true;
 	}
 };
-IMPLEMENT_AI_INSTANT_TEST(FStateTreeTest_Stop_NormalEnterFailed, "System.StateTree.Stop.NormalEnterFailed");
+IMPLEMENT_STATE_TREE_INSTANT_TEST(FStateTreeTest_Stop_NormalEnterFailed, "System.StateTree.Stop.NormalEnterFailed");
 
 
 
@@ -1450,7 +1450,7 @@ struct FStateTreeTest_Stop_GlobalEnterSucceeded : FStateTreeTest_StopEnterGlobal
 		return true;
 	}
 };
-IMPLEMENT_AI_INSTANT_TEST(FStateTreeTest_Stop_GlobalEnterSucceeded, "System.StateTree.Stop.GlobalEnterSucceeded");
+IMPLEMENT_STATE_TREE_INSTANT_TEST(FStateTreeTest_Stop_GlobalEnterSucceeded, "System.StateTree.Stop.GlobalEnterSucceeded");
 
 struct FStateTreeTest_Stop_GlobalEnterFailed : FStateTreeTest_StopEnterGlobal
 {
@@ -1472,7 +1472,7 @@ struct FStateTreeTest_Stop_GlobalEnterFailed : FStateTreeTest_StopEnterGlobal
 		return true;
 	}
 };
-IMPLEMENT_AI_INSTANT_TEST(FStateTreeTest_Stop_GlobalEnterFailed, "System.StateTree.Stop.GlobalEnterFailed");
+IMPLEMENT_STATE_TREE_INSTANT_TEST(FStateTreeTest_Stop_GlobalEnterFailed, "System.StateTree.Stop.GlobalEnterFailed");
 
 struct FStateTreeTest_Stop_ExternalStop : FStateTreeTest_Stop
 {
@@ -1538,7 +1538,7 @@ struct FStateTreeTest_Stop_ExternalStop : FStateTreeTest_Stop
 	EStateTreeRunStatus ExpectedStatusAfterStop = EStateTreeRunStatus::Stopped;
 	
 };
-IMPLEMENT_AI_INSTANT_TEST(FStateTreeTest_Stop_ExternalStop, "System.StateTree.Stop.ExternalStop");
+IMPLEMENT_STATE_TREE_INSTANT_TEST(FStateTreeTest_Stop_ExternalStop, "System.StateTree.Stop.ExternalStop");
 
 struct FStateTreeTest_Stop_AlreadyStopped : FStateTreeTest_Stop_ExternalStop
 {
@@ -1563,7 +1563,7 @@ struct FStateTreeTest_Stop_AlreadyStopped : FStateTreeTest_Stop_ExternalStop
 		return true;
 	}
 };
-IMPLEMENT_AI_INSTANT_TEST(FStateTreeTest_Stop_AlreadyStopped, "System.StateTree.Stop.AlreadyStopped");
+IMPLEMENT_STATE_TREE_INSTANT_TEST(FStateTreeTest_Stop_AlreadyStopped, "System.StateTree.Stop.AlreadyStopped");
 
 //
 // The deferred stop tests validates that the tree can be properly stopped if requested in the main entry points (Start, Tick, Stop).  
@@ -1630,7 +1630,7 @@ struct FStateTreeTest_DeferredStop_EnterGlobalTask : FStateTreeTest_DeferredStop
 		return true;
 	}
 };
-IMPLEMENT_AI_INSTANT_TEST(FStateTreeTest_DeferredStop_EnterGlobalTask, "System.StateTree.DeferredStop.EnterGlobalTask");
+IMPLEMENT_STATE_TREE_INSTANT_TEST(FStateTreeTest_DeferredStop_EnterGlobalTask, "System.StateTree.DeferredStop.EnterGlobalTask");
 
 struct FStateTreeTest_DeferredStop_TickGlobalTask : FStateTreeTest_DeferredStop
 {
@@ -1650,7 +1650,7 @@ struct FStateTreeTest_DeferredStop_TickGlobalTask : FStateTreeTest_DeferredStop
 		return true;
 	}
 };
-IMPLEMENT_AI_INSTANT_TEST(FStateTreeTest_DeferredStop_TickGlobalTask, "System.StateTree.DeferredStop.TickGlobalTask");
+IMPLEMENT_STATE_TREE_INSTANT_TEST(FStateTreeTest_DeferredStop_TickGlobalTask, "System.StateTree.DeferredStop.TickGlobalTask");
 
 struct FStateTreeTest_DeferredStop_ExitGlobalTask : FStateTreeTest_DeferredStop
 {
@@ -1671,7 +1671,7 @@ struct FStateTreeTest_DeferredStop_ExitGlobalTask : FStateTreeTest_DeferredStop
 		return true;
 	}
 };
-IMPLEMENT_AI_INSTANT_TEST(FStateTreeTest_DeferredStop_ExitGlobalTask, "System.StateTree.DeferredStop.ExitGlobalTask");
+IMPLEMENT_STATE_TREE_INSTANT_TEST(FStateTreeTest_DeferredStop_ExitGlobalTask, "System.StateTree.DeferredStop.ExitGlobalTask");
 
 struct FStateTreeTest_DeferredStop_EnterTask : FStateTreeTest_DeferredStop
 {
@@ -1688,7 +1688,7 @@ struct FStateTreeTest_DeferredStop_EnterTask : FStateTreeTest_DeferredStop
 		return true;
 	}
 };
-IMPLEMENT_AI_INSTANT_TEST(FStateTreeTest_DeferredStop_EnterTask, "System.StateTree.DeferredStop.EnterTask");
+IMPLEMENT_STATE_TREE_INSTANT_TEST(FStateTreeTest_DeferredStop_EnterTask, "System.StateTree.DeferredStop.EnterTask");
 
 struct FStateTreeTest_DeferredStop_TickTask : FStateTreeTest_DeferredStop
 {
@@ -1708,7 +1708,7 @@ struct FStateTreeTest_DeferredStop_TickTask : FStateTreeTest_DeferredStop
 		return true;
 	}
 };
-IMPLEMENT_AI_INSTANT_TEST(FStateTreeTest_DeferredStop_TickTask, "System.StateTree.DeferredStop.TickTask");
+IMPLEMENT_STATE_TREE_INSTANT_TEST(FStateTreeTest_DeferredStop_TickTask, "System.StateTree.DeferredStop.TickTask");
 
 struct FStateTreeTest_DeferredStop_ExitTask : FStateTreeTest_DeferredStop
 {
@@ -1729,7 +1729,7 @@ struct FStateTreeTest_DeferredStop_ExitTask : FStateTreeTest_DeferredStop
 		return true;
 	}
 };
-IMPLEMENT_AI_INSTANT_TEST(FStateTreeTest_DeferredStop_ExitTask, "System.StateTree.DeferredStop.ExitTask");
+IMPLEMENT_STATE_TREE_INSTANT_TEST(FStateTreeTest_DeferredStop_ExitTask, "System.StateTree.DeferredStop.ExitTask");
 
 struct FStateTreeTest_FinishTasks : FStateTreeTestBase
 {
@@ -1902,7 +1902,7 @@ struct FStateTreeTest_FinishTasks : FStateTreeTestBase
 		return true;
 	}
 };
-IMPLEMENT_AI_INSTANT_TEST(FStateTreeTest_FinishTasks, "System.StateTree.FinishTask");
+IMPLEMENT_STATE_TREE_INSTANT_TEST(FStateTreeTest_FinishTasks, "System.StateTree.FinishTask");
 
 
 // Test nested tree overrides
@@ -2050,13 +2050,19 @@ struct FStateTreeTest_NestedOverride : FStateTreeTestBase
 		return true;
 	}
 };
-IMPLEMENT_AI_INSTANT_TEST(FStateTreeTest_NestedOverride, "System.StateTree.NestedOverride");
+IMPLEMENT_STATE_TREE_INSTANT_TEST(FStateTreeTest_NestedOverride, "System.StateTree.NestedOverride");
 
 // Test parallel tree event priority handling.
 struct FStateTreeTest_RecursiveParallelTask : FStateTreeTestBase
 {
 	virtual bool InstantTest() override
 	{
+		if (AutoRTFM::IsTransactional())
+		{
+			// AddExpectedError is currently not supported in AutoRTFM.
+			return true;
+		}
+
 		//Tree 1
 		//	Root (with task that runs Tree 1)
 
@@ -2095,7 +2101,7 @@ struct FStateTreeTest_RecursiveParallelTask : FStateTreeTestBase
 	}
 };
 
-IMPLEMENT_AI_INSTANT_TEST(FStateTreeTest_RecursiveParallelTask, "System.StateTree.RecursiveParallelTask");
+IMPLEMENT_STATE_TREE_INSTANT_TEST(FStateTreeTest_RecursiveParallelTask, "System.StateTree.RecursiveParallelTask");
 
 // Test parallel tree event priority handling.
 struct FStateTreeTest_ParallelEventPriority : FStateTreeTestBase
@@ -2221,7 +2227,7 @@ struct FStateTreeTest_ParallelEventPriority : FStateTreeTestBase
 		return true;
 	}
 };
-IMPLEMENT_AI_INSTANT_TEST(FStateTreeTest_ParallelEventPriority, "System.StateTree.ParallelEventPriority");
+IMPLEMENT_STATE_TREE_INSTANT_TEST(FStateTreeTest_ParallelEventPriority, "System.StateTree.ParallelEventPriority");
 
 
 struct FStateTreeTest_ParallelEventPriority_Low : FStateTreeTest_ParallelEventPriority
@@ -2231,7 +2237,7 @@ struct FStateTreeTest_ParallelEventPriority_Low : FStateTreeTest_ParallelEventPr
 		ParallelTreePriority = EStateTreeTransitionPriority::Low;
 	}
 };
-IMPLEMENT_AI_INSTANT_TEST(FStateTreeTest_ParallelEventPriority_Low, "System.StateTree.ParallelEventPriority.Low");
+IMPLEMENT_STATE_TREE_INSTANT_TEST(FStateTreeTest_ParallelEventPriority_Low, "System.StateTree.ParallelEventPriority.Low");
 
 struct FStateTreeTest_ParallelEventPriority_High : FStateTreeTest_ParallelEventPriority
 {
@@ -2240,7 +2246,7 @@ struct FStateTreeTest_ParallelEventPriority_High : FStateTreeTest_ParallelEventP
 		ParallelTreePriority = EStateTreeTransitionPriority::High;
 	}
 };
-IMPLEMENT_AI_INSTANT_TEST(FStateTreeTest_ParallelEventPriority_High, "System.StateTree.ParallelEventPriority.High");
+IMPLEMENT_STATE_TREE_INSTANT_TEST(FStateTreeTest_ParallelEventPriority_High, "System.StateTree.ParallelEventPriority.High");
 
 struct FStateTreeTest_SubTreeTransition : FStateTreeTestBase
 {
@@ -2324,7 +2330,7 @@ struct FStateTreeTest_SubTreeTransition : FStateTreeTestBase
 		return true;
 	}
 };
-IMPLEMENT_AI_INSTANT_TEST(FStateTreeTest_SubTreeTransition, "System.StateTree.SubTreeTransition");
+IMPLEMENT_STATE_TREE_INSTANT_TEST(FStateTreeTest_SubTreeTransition, "System.StateTree.SubTreeTransition");
 
 struct FStateTreeTest_Reentrant : FStateTreeTestBase
 {
@@ -2421,7 +2427,7 @@ struct FStateTreeTest_Reentrant : FStateTreeTestBase
 		return true;
 	}
 };
-IMPLEMENT_AI_INSTANT_TEST(FStateTreeTest_Reentrant, "System.StateTree.Reentrant");
+IMPLEMENT_STATE_TREE_INSTANT_TEST(FStateTreeTest_Reentrant, "System.StateTree.Reentrant");
 
 } // namespace UE::StateTree::Tests
 
diff --git a/Engine/Plugins/Runtime/StateTree/Source/StateTreeTestSuite/Private/StateTreeTest.h b/Engine/Plugins/Runtime/StateTree/Source/StateTreeTestSuite/Private/StateTreeTest.h
--- a/Engine/Plugins/Runtime/StateTree/Source/StateTreeTestSuite/Private/StateTreeTest.h
+++ b/Engine/Plugins/Runtime/StateTree/Source/StateTreeTestSuite/Private/StateTreeTest.h
@@ -30,3 +30,5 @@ class UStateTreeTestSchema2 : public UStateTreeSchema
 	GENERATED_BODY()
 };
 
+#define IMPLEMENT_STATE_TREE_INSTANT_TEST(TestClass, PrettyName) \
+	IMPLEMENT_AI_INSTANT_TEST_WITH_FLAGS(TestClass, PrettyName, EAutomationTestFlags::ClientContext | EAutomationTestFlags::EditorContext | EAutomationTestFlags::EngineFilter | EAutomationTestFlags::SupportsAutoRTFM)
diff --git a/Engine/Plugins/Runtime/StateTree/Source/StateTreeTestSuite/Private/StateTreeTransitionTest.cpp b/Engine/Plugins/Runtime/StateTree/Source/StateTreeTestSuite/Private/StateTreeTransitionTest.cpp
--- a/Engine/Plugins/Runtime/StateTree/Source/StateTreeTestSuite/Private/StateTreeTransitionTest.cpp
+++ b/Engine/Plugins/Runtime/StateTree/Source/StateTreeTestSuite/Private/StateTreeTransitionTest.cpp
@@ -94,7 +94,7 @@ struct FStateTreeTest_TransitionPriority : FStateTreeTestBase
 		return true;
 	}
 };
-IMPLEMENT_AI_INSTANT_TEST(FStateTreeTest_TransitionPriority, "System.StateTree.Transition.Priority");
+IMPLEMENT_STATE_TREE_INSTANT_TEST(FStateTreeTest_TransitionPriority, "System.StateTree.Transition.Priority");
 
 struct FStateTreeTest_TransitionPriorityEnterState : FStateTreeTestBase
 {
@@ -160,7 +160,7 @@ struct FStateTreeTest_TransitionPriorityEnterState : FStateTreeTestBase
 		return true;
 	}
 };
-IMPLEMENT_AI_INSTANT_TEST(FStateTreeTest_TransitionPriorityEnterState, "System.StateTree.Transition.PriorityEnterState");
+IMPLEMENT_STATE_TREE_INSTANT_TEST(FStateTreeTest_TransitionPriorityEnterState, "System.StateTree.Transition.PriorityEnterState");
 
 struct FStateTreeTest_TransitionNextSelectableState : FStateTreeTestBase
 {
@@ -250,7 +250,7 @@ struct FStateTreeTest_TransitionNextSelectableState : FStateTreeTestBase
 		return true;
 	}
 };
-IMPLEMENT_AI_INSTANT_TEST(FStateTreeTest_TransitionNextSelectableState, "System.StateTree.Transition.NextSelectableState");
+IMPLEMENT_STATE_TREE_INSTANT_TEST(FStateTreeTest_TransitionNextSelectableState, "System.StateTree.Transition.NextSelectableState");
 
 struct FStateTreeTest_TransitionNextWithParentData : FStateTreeTestBase
 {
@@ -308,7 +308,7 @@ struct FStateTreeTest_TransitionNextWithParentData : FStateTreeTestBase
 		return true;
 	}
 };
-IMPLEMENT_AI_INSTANT_TEST(FStateTreeTest_TransitionNextWithParentData, "System.StateTree.Transition.NextWithParentData");
+IMPLEMENT_STATE_TREE_INSTANT_TEST(FStateTreeTest_TransitionNextWithParentData, "System.StateTree.Transition.NextWithParentData");
 
 struct FStateTreeTest_TransitionGlobalDataView : FStateTreeTestBase
 {
@@ -369,7 +369,7 @@ struct FStateTreeTest_TransitionGlobalDataView : FStateTreeTestBase
 		return true;
 	}
 };
-IMPLEMENT_AI_INSTANT_TEST(FStateTreeTest_TransitionGlobalDataView, "System.StateTree.Transition.GlobalDataView");
+IMPLEMENT_STATE_TREE_INSTANT_TEST(FStateTreeTest_TransitionGlobalDataView, "System.StateTree.Transition.GlobalDataView");
 
 struct FStateTreeTest_TransitionDelay : FStateTreeTestBase
 {
@@ -452,7 +452,7 @@ struct FStateTreeTest_TransitionDelay : FStateTreeTestBase
 		return true;
 	}
 };
-IMPLEMENT_AI_INSTANT_TEST(FStateTreeTest_TransitionDelay, "System.StateTree.Transition.Delay");
+IMPLEMENT_STATE_TREE_INSTANT_TEST(FStateTreeTest_TransitionDelay, "System.StateTree.Transition.Delay");
 
 struct FStateTreeTest_TransitionDelayZero : FStateTreeTestBase
 {
@@ -515,7 +515,7 @@ struct FStateTreeTest_TransitionDelayZero : FStateTreeTestBase
 		return true;
 	}
 };
-IMPLEMENT_AI_INSTANT_TEST(FStateTreeTest_TransitionDelayZero, "System.StateTree.Transition.DelayZero");
+IMPLEMENT_STATE_TREE_INSTANT_TEST(FStateTreeTest_TransitionDelayZero, "System.StateTree.Transition.DelayZero");
 
 struct FStateTreeTest_PassingTransitionEventToStateSelection : FStateTreeTestBase
 {
@@ -618,7 +618,7 @@ struct FStateTreeTest_PassingTransitionEventToStateSelection : FStateTreeTestBas
 		return true;
 	}
 };
-IMPLEMENT_AI_INSTANT_TEST(FStateTreeTest_PassingTransitionEventToStateSelection, "System.StateTree.Transition.PassingTransitionEventToStateSelection");
+IMPLEMENT_STATE_TREE_INSTANT_TEST(FStateTreeTest_PassingTransitionEventToStateSelection, "System.StateTree.Transition.PassingTransitionEventToStateSelection");
 
 struct FStateTreeTest_FollowTransitions : FStateTreeTestBase
 {
@@ -711,7 +711,7 @@ struct FStateTreeTest_FollowTransitions : FStateTreeTestBase
 		return true;
 	}
 };
-IMPLEMENT_AI_INSTANT_TEST(FStateTreeTest_FollowTransitions, "System.StateTree.Transition.FollowTransitions");
+IMPLEMENT_STATE_TREE_INSTANT_TEST(FStateTreeTest_FollowTransitions, "System.StateTree.Transition.FollowTransitions");
 
 struct FStateTreeTest_InfiniteLoop : FStateTreeTestBase
 {
@@ -959,7 +959,7 @@ struct FStateTreeTest_RegularTransitions : FStateTreeTestBase
 		return true;
 	}
 };
-IMPLEMENT_AI_INSTANT_TEST(FStateTreeTest_RegularTransitions, "System.StateTree.Transition.RegularTransitions");
+IMPLEMENT_STATE_TREE_INSTANT_TEST(FStateTreeTest_RegularTransitions, "System.StateTree.Transition.RegularTransitions");
 
 
 struct FStateTreeTest_RequestTransition : FStateTreeTestBase
diff --git a/Engine/Plugins/Runtime/StateTree/Source/StateTreeTestSuite/Private/StateTreeWeakContextTest.cpp b/Engine/Plugins/Runtime/StateTree/Source/StateTreeTestSuite/Private/StateTreeWeakContextTest.cpp
--- a/Engine/Plugins/Runtime/StateTree/Source/StateTreeTestSuite/Private/StateTreeWeakContextTest.cpp
+++ b/Engine/Plugins/Runtime/StateTree/Source/StateTreeTestSuite/Private/StateTreeWeakContextTest.cpp
@@ -219,7 +219,7 @@ struct FStateTreeTest_WeakContext_FinishTask : FStateTreeTestBase
 		return true;
 	}
 };
-IMPLEMENT_AI_INSTANT_TEST(FStateTreeTest_WeakContext_FinishTask, "System.StateTree.WeakContext.FinishTask");
+IMPLEMENT_STATE_TREE_INSTANT_TEST(FStateTreeTest_WeakContext_FinishTask, "System.StateTree.WeakContext.FinishTask");
 } // namespace UE::StateTree::Tests
 
 UE_ENABLE_OPTIMIZATION_SHIP
diff --git a/Engine/Source/Developer/AITestSuite/Public/AITestsCommon.h b/Engine/Source/Developer/AITestSuite/Public/AITestsCommon.h
--- a/Engine/Source/Developer/AITestSuite/Public/AITestsCommon.h
+++ b/Engine/Source/Developer/AITestSuite/Public/AITestsCommon.h
@@ -111,8 +111,8 @@ DEFINE_EXPORTED_LATENT_AUTOMATION_COMMAND_ONE_PARAMETER(AITESTSUITE_API, FAITest
 		return true; \
 	} 
 
-#define IMPLEMENT_AI_INSTANT_TEST(TestClass, PrettyName) \
-	IMPLEMENT_SIMPLE_AUTOMATION_TEST(TestClass##Runner, PrettyName, (EAutomationTestFlags::ClientContext | EAutomationTestFlags::EditorContext | EAutomationTestFlags::EngineFilter)) \
+#define IMPLEMENT_AI_INSTANT_TEST_WITH_FLAGS(TestClass, PrettyName, TestFlags) \
+	IMPLEMENT_SIMPLE_AUTOMATION_TEST(TestClass##Runner, PrettyName, (TestFlags)) \
 	bool TestClass##Runner::RunTest(const FString& Parameters) \
 	{ \
 		bool bSuccess = false; \
@@ -131,6 +131,9 @@ DEFINE_EXPORTED_LATENT_AUTOMATION_COMMAND_ONE_PARAMETER(AITESTSUITE_API, FAITest
 		return bSuccess; \
 	} 
 
+#define IMPLEMENT_AI_INSTANT_TEST(TestClass, PrettyName) \
+	IMPLEMENT_AI_INSTANT_TEST_WITH_FLAGS(TestClass, PrettyName, EAutomationTestFlags::ClientContext | EAutomationTestFlags::EditorContext | EAutomationTestFlags::EngineFilter)
+
 /** 
  *	This macro allows one to implement a whole set of simple tests that share common setups. To use it first implement
  *	a struct that builds the said common setup. Like so:
