StateTree: Use the regular allocator in state tree to track a bug. This is probably temporary.
#rnx

diff --git a/Engine/Plugins/Runtime/StateTree/Source/StateTreeModule/Private/StateTreeExecutionContext.cpp b/Engine/Plugins/Runtime/StateTree/Source/StateTreeModule/Private/StateTreeExecutionContext.cpp
--- a/Engine/Plugins/Runtime/StateTree/Source/StateTreeModule/Private/StateTreeExecutionContext.cpp
+++ b/Engine/Plugins/Runtime/StateTree/Source/StateTreeModule/Private/StateTreeExecutionContext.cpp
@@ -1874,10 +1874,10 @@ void FStateTreeExecutionContext::UpdateInstanceData(TConstArrayView<FStateTreeEx
 		}
 	}
 	
-	TArray<FConstStructView, FNonconcurrentLinearArrayAllocator> InstanceStructs;
+	TArray<FConstStructView> InstanceStructs;
 	InstanceStructs.Reserve(EstimatedNumStructs);
 
-	TArray<FInstancedStruct*, FNonconcurrentLinearArrayAllocator> TempInstanceStructs;
+	TArray<FInstancedStruct*> TempInstanceStructs;
 	TempInstanceStructs.Reserve(EstimatedNumStructs);
 
 	TArray<FCompactStateTreeParameters, TFixedAllocator<FStateSelectionResult::MaxExecutionFrames>> TempParams;
@@ -2683,7 +2683,7 @@ EStateTreeRunStatus FStateTreeExecutionContext::EnterState(FStateTreeTransitionR
 	UE_STATETREE_DEBUG_ENTER_PHASE(this, EStateTreeUpdatePhase::EnterStates);
 
 	// The previous active frames are needed for state enter logic.
-	TArray<FStateTreeExecutionFrame, FNonconcurrentLinearArrayAllocator> PreviousActiveFrames;
+	TArray<FStateTreeExecutionFrame> PreviousActiveFrames;
 	PreviousActiveFrames = Exec.ActiveFrames;
 
 	// Reset the current active frames, new ones are added one by one.
@@ -2920,7 +2920,7 @@ void FStateTreeExecutionContext::ExitState(const FStateTreeTransitionResult& Tra
 		bool bShouldCall = false; 
 	};
 
-	TArray<FExitStateCall, FNonconcurrentLinearArrayAllocator> ExitStateCalls;
+	TArray<FExitStateCall> ExitStateCalls;
 
 	// Track any changed state (i.e., not sustained) to prevent reused subtrees from being considered as sustained states since the whole
 	// tree is reused and should be considered as a new instance.
diff --git a/Engine/Plugins/Runtime/StateTree/Source/StateTreeModule/Public/StateTreeExecutionContext.h b/Engine/Plugins/Runtime/StateTree/Source/StateTreeModule/Public/StateTreeExecutionContext.h
--- a/Engine/Plugins/Runtime/StateTree/Source/StateTreeModule/Public/StateTreeExecutionContext.h
+++ b/Engine/Plugins/Runtime/StateTree/Source/StateTreeModule/Public/StateTreeExecutionContext.h
@@ -1216,7 +1216,7 @@ PRAGMA_ENABLE_DEPRECATION_WARNINGS
 	FStateTreeReferenceOverrides LinkedAssetStateTreeOverrides;
 	
 	/** Data view of the context data. */
-	TArray<FStateTreeDataView, FNonconcurrentLinearArrayAllocator> ContextAndExternalDataViews;
+	TArray<FStateTreeDataView> ContextAndExternalDataViews;
 
 	FOnCollectStateTreeExternalData CollectExternalDataDelegate;
 
@@ -1225,7 +1225,7 @@ PRAGMA_ENABLE_DEPRECATION_WARNINGS
 		const UStateTree* StateTree = nullptr;
 		FStateTreeIndex16 BaseIndex;
 	};
-	TArray<FCollectedExternalDataCache, FNonconcurrentLinearArrayAllocator> CollectedExternalCache;
+	TArray<FCollectedExternalDataCache> CollectedExternalCache;
 
 	bool bActiveExternalDataCollected = false;
 
