[AutoRTFM] Update the remaining two StateTree tests so that they can run in AutoRTFM.

These tests used `AddExpectedError`, which expects a regular expression for its input. AutoRTFM does not support this today because ICU is used to parse the regular expression, and we don't support ICU yet.

However, in this case, we have a simple solution which is actually better all around: the API `AddExpectedErrorPlain` actually works fine and doesn't require ICU at all. This API expects a raw string instead of a regex, which is what we were passing anyway.

diff --git a/Engine/Plugins/Runtime/StateTree/Source/StateTreeTestSuite/Private/StateTreeLinkedStateTest.cpp b/Engine/Plugins/Runtime/StateTree/Source/StateTreeTestSuite/Private/StateTreeLinkedStateTest.cpp
--- a/Engine/Plugins/Runtime/StateTree/Source/StateTreeTestSuite/Private/StateTreeLinkedStateTest.cpp
+++ b/Engine/Plugins/Runtime/StateTree/Source/StateTreeTestSuite/Private/StateTreeLinkedStateTest.cpp
@@ -932,12 +932,6 @@ struct FStateTreeTest_RecursiveLinkedAsset : FStateTreeTestBase
 {
 	virtual bool InstantTest() override
 	{
-		if (AutoRTFM::IsTransactional())
-		{
-			// AddExpectedError is currently not supported in AutoRTFM.
-			return true;
-		}
-
 		//Tree 1
 		//	Root
 		//		StateLinkedTree1 (Tree2) -> Next
@@ -1045,7 +1039,7 @@ struct FStateTreeTest_RecursiveLinkedAsset : FStateTreeTestBase
 			}
 
 			{
-				GetTestRunner().AddExpectedError(TEXT("Trying to recursively enter subtree"), EAutomationExpectedErrorFlags::Contains, 1);
+				GetTestRunner().AddExpectedErrorPlain(TEXT("Trying to recursively enter subtree"), EAutomationExpectedErrorFlags::Contains, 1);
 
 				EStateTreeRunStatus Status = Exec.Start();
 				AITEST_EQUAL(TEXT("Start should complete with Running"), Status, EStateTreeRunStatus::Running);
diff --git a/Engine/Plugins/Runtime/StateTree/Source/StateTreeTestSuite/Private/StateTreeTest.cpp b/Engine/Plugins/Runtime/StateTree/Source/StateTreeTestSuite/Private/StateTreeTest.cpp
--- a/Engine/Plugins/Runtime/StateTree/Source/StateTreeTestSuite/Private/StateTreeTest.cpp
+++ b/Engine/Plugins/Runtime/StateTree/Source/StateTreeTestSuite/Private/StateTreeTest.cpp
@@ -2057,12 +2057,6 @@ struct FStateTreeTest_RecursiveParallelTask : FStateTreeTestBase
 {
 	virtual bool InstantTest() override
 	{
-		if (AutoRTFM::IsTransactional())
-		{
-			// AddExpectedError is currently not supported in AutoRTFM.
-			return true;
-		}
-
 		//Tree 1
 		//	Root (with task that runs Tree 1)
 
@@ -2090,7 +2084,7 @@ struct FStateTreeTest_RecursiveParallelTask : FStateTreeTestBase
 			AITEST_TRUE(TEXT("StateTree should init"), bInitSucceeded);
 
 			{
-				GetTestRunner().AddExpectedError(TEXT("Trying to start a new parallel tree from the same tree"), EAutomationExpectedErrorFlags::Contains, 1);
+ 				GetTestRunner().AddExpectedErrorPlain(TEXT("Trying to start a new parallel tree from the same tree"), EAutomationExpectedErrorFlags::Contains, 1);
 
 				EStateTreeRunStatus Status = Exec.Start();
 				AITEST_EQUAL(TEXT("Start should complete with failed"), Status, EStateTreeRunStatus::Failed);
