StateTree: Add runtime validation to confirm the validity of the statetree's instance data.

diff --git a/Engine/Plugins/Runtime/StateTree/Source/StateTreeModule/Private/Debugger/StateTreeRuntimeValidation.cpp b/Engine/Plugins/Runtime/StateTree/Source/StateTreeModule/Private/Debugger/StateTreeRuntimeValidation.cpp
--- a/Engine/Plugins/Runtime/StateTree/Source/StateTreeModule/Private/Debugger/StateTreeRuntimeValidation.cpp
+++ b/Engine/Plugins/Runtime/StateTree/Source/StateTreeModule/Private/Debugger/StateTreeRuntimeValidation.cpp
@@ -36,6 +36,13 @@ static FAutoConsoleVariableRef CVarRuntimeValidationEnterExitState(
 		"Test that if a node get a ExitState, it did receive an EnterState before.")
 );
 
+bool bRuntimeValidationInstanceData = false;
+static FAutoConsoleVariableRef CVarRuntimeValidationInstanceData(
+	TEXT("StateTree.RuntimeValidation.InstanceData"),
+	bRuntimeValidationInstanceData,
+	TEXT("Test if the state tree instance data and shared instance data are valid.")
+);
+
 FString NodeToString(const UObject* Obj, FGuid Id)
 {
 	TStringBuilderWithBuffer<TCHAR, 128> Buffer;
@@ -88,7 +95,8 @@ void FRuntimeValidationInstanceData::SetContext(const UObject* InNewOwner, const
 		}
 	}
 
-	ValidatesTreeNodes(InNewStateTree);
+	ValidateTreeNodes(InNewStateTree);
+	ValidateInstanceData(InNewStateTree);
 
 	StateTree = NewStateTree;
 	Owner = NewOwner;
@@ -143,7 +151,7 @@ void FRuntimeValidationInstanceData::NodeExitState(FGuid NodeID, FActiveFrameID
 	}
 }
 
-void FRuntimeValidationInstanceData::ValidatesTreeNodes(const UStateTree* InNewStateTree) const
+void FRuntimeValidationInstanceData::ValidateTreeNodes(const UStateTree* InNewStateTree) const
 {
 	if (Private::bRuntimeValidationDoesNewerVersionExists)
 	{
@@ -215,6 +223,31 @@ void FRuntimeValidationInstanceData::ValidatesTreeNodes(const UStateTree* InNewS
 		}
 	}
 }
+
+void FRuntimeValidationInstanceData::ValidateInstanceData(const UStateTree* NewStateTree)
+{
+	if (Private::bRuntimeValidationInstanceData)
+	{
+		if (NewStateTree)
+		{
+			if (!NewStateTree->GetDefaultInstanceData().GetStorage().AreAllInstancesValid())
+			{
+				ensureAlwaysMsgf(false, TEXT("The instance data has invalid data."));
+				Private::bRuntimeValidationInstanceData = false;
+			}
+
+			TSharedPtr<FStateTreeInstanceData> SharedPtr = NewStateTree->GetSharedInstanceData();
+			if (ensureMsgf(SharedPtr, TEXT("The shared instance data is invalid")))
+			{
+				if (!SharedPtr->GetStorage().AreAllInstancesValid())
+				{
+					ensureAlwaysMsgf(false, TEXT("The shared instance data has invalid data."));
+					Private::bRuntimeValidationInstanceData = false;
+				}
+			}
+		}
+	}
+}
 #endif //WITH_STATETREE_DEBUG
 
 /**
diff --git a/Engine/Plugins/Runtime/StateTree/Source/StateTreeModule/Private/Debugger/StateTreeRuntimeValidationInstanceData.h b/Engine/Plugins/Runtime/StateTree/Source/StateTreeModule/Private/Debugger/StateTreeRuntimeValidationInstanceData.h
--- a/Engine/Plugins/Runtime/StateTree/Source/StateTreeModule/Private/Debugger/StateTreeRuntimeValidationInstanceData.h
+++ b/Engine/Plugins/Runtime/StateTree/Source/StateTreeModule/Private/Debugger/StateTreeRuntimeValidationInstanceData.h
@@ -26,7 +26,8 @@ public:
 	void NodeExitState(FGuid NodeID, FActiveFrameID FrameID);
 
 private:
-	void ValidatesTreeNodes(const UStateTree* InNewStateTree) const;
+	void ValidateTreeNodes(const UStateTree* NewStateTree) const;
+	void ValidateInstanceData(const UStateTree* NewStateTree);
 
 private:
 	enum class EState : uint8
diff --git a/Engine/Plugins/Runtime/StateTree/Source/StateTreeModule/Private/StateTree.cpp b/Engine/Plugins/Runtime/StateTree/Source/StateTreeModule/Private/StateTree.cpp
--- a/Engine/Plugins/Runtime/StateTree/Source/StateTreeModule/Private/StateTree.cpp
+++ b/Engine/Plugins/Runtime/StateTree/Source/StateTreeModule/Private/StateTree.cpp
@@ -7,6 +7,7 @@
 #include "StateTreeEvaluatorBase.h"
 #include "StateTreeConditionBase.h"
 #include "StateTreeConsiderationBase.h"
+#include "StateTreeModuleImpl.h"
 #include "AssetRegistry/AssetData.h"
 #include "Misc/ScopeRWLock.h"
 #include "StateTreeDelegates.h"
@@ -21,7 +22,6 @@
 
 #if WITH_EDITOR
 #include "Editor.h"
-#include "StateTreeModuleImpl.h"
 #include "StructUtils/UserDefinedStruct.h"
 #include "Templates/GuardValueAccessors.h"
 #include "UObject/LinkerLoad.h"
@@ -173,6 +173,45 @@ bool UStateTree::HasCompatibleContextData(const UStateTree& Other) const
 	return true;
 }
 
+#if WITH_EDITOR || WITH_STATETREE_DEBUG
+void UStateTree::PostInitProperties()
+{
+	Super::PostInitProperties();
+
+#if WITH_EDITOR
+	if (!HasAnyFlags(RF_ClassDefaultObject))
+	{
+		OnObjectsReinstancedHandle = FStateTreeModule::OnObjectsReinstanced.AddUObject(this, &UStateTree::OnObjectsReinstanced);
+		OnUserDefinedStructReinstancedHandle = FStateTreeModule::OnUserDefinedStructReinstanced.AddUObject(this, &UStateTree::OnUserDefinedStructReinstanced);
+		OnPreBeginPIEHandle = FStateTreeModule::OnPreBeginPIE.AddUObject(this, &UStateTree::OnPreBeginPIE);
+	}
+#endif
+
+#if WITH_STATETREE_DEBUG
+	if (!HasAnyFlags(RF_ClassDefaultObject))
+	{
+		PreGCHandle = FStateTreeModule::OnPreRuntimeValidationInstanceData.AddUObject(this, &UStateTree::HandleRuntimeValidationPreGC);
+		PostGCHandle = FStateTreeModule::OnPostRuntimeValidationInstanceData.AddUObject(this, &UStateTree::HandleRuntimeValidationPostGC);
+	}
+#endif
+}
+
+void UStateTree::BeginDestroy()
+{
+#if WITH_EDITOR
+	FStateTreeModule::OnObjectsReinstanced.Remove(OnObjectsReinstancedHandle);
+	FStateTreeModule::OnUserDefinedStructReinstanced.Remove(OnUserDefinedStructReinstancedHandle);
+	FStateTreeModule::OnPreBeginPIE.Remove(OnPreBeginPIEHandle);
+#endif
+
+#if WITH_STATETREE_DEBUG
+	FStateTreeModule::OnPreRuntimeValidationInstanceData.Remove(PreGCHandle);
+	FStateTreeModule::OnPostRuntimeValidationInstanceData.Remove(PostGCHandle);
+#endif
+
+	Super::BeginDestroy();
+}
+#endif //WITH_EDITOR || WITH_STATETREE_DEBUG
 
 #if WITH_EDITOR
 namespace UE::StateTree::Compiler
@@ -336,27 +375,6 @@ void UStateTree::OnUserDefinedStructReinstanced(const UUserDefinedStruct& UserDe
 	}
 }
 
-void UStateTree::PostInitProperties()
-{
-	Super::PostInitProperties();
-	
-	if (!HasAnyFlags(RF_ClassDefaultObject))
-	{
-		OnObjectsReinstancedHandle = FStateTreeModule::OnObjectsReinstanced.AddUObject(this, &UStateTree::OnObjectsReinstanced);
-		OnUserDefinedStructReinstancedHandle = FStateTreeModule::OnUserDefinedStructReinstanced.AddUObject(this, &UStateTree::OnUserDefinedStructReinstanced);
-		OnPreBeginPIEHandle = FStateTreeModule::OnPreBeginPIE.AddUObject(this, &UStateTree::OnPreBeginPIE);
-	}
-}
-
-void UStateTree::BeginDestroy()
-{
-	FStateTreeModule::OnObjectsReinstanced.Remove(OnObjectsReinstancedHandle);
-	FStateTreeModule::OnUserDefinedStructReinstanced.Remove(OnUserDefinedStructReinstancedHandle);
-	FStateTreeModule::OnPreBeginPIE.Remove(OnPreBeginPIEHandle);
-
-	Super::BeginDestroy();
-}
-
 void UStateTree::PostDuplicate(EDuplicateMode::Type DuplicateMode)
 {
 	Super::PostDuplicate(DuplicateMode);
@@ -1522,3 +1540,88 @@ FString UStateTree::DebugInternalLayoutAsString() const
 }
 #endif // WITH_EDITOR || WITH_STATETREE_DEBUG
 
+#if WITH_STATETREE_DEBUG
+namespace UE::StateTree::Debug::Private
+{
+	IConsoleVariable* FindCVarInstanceDataGC()
+	{
+		IConsoleVariable* FoundVariable = IConsoleManager::Get().FindConsoleVariable(TEXT("StateTree.RuntimeValidation.InstanceDataGC"));
+		check(FoundVariable);
+		return FoundVariable;
+	}
+}
+void UStateTree::HandleRuntimeValidationPreGC()
+{
+	GCObjectDatas.Reset();
+	auto AddInstanceObject = [this](const FStateTreeInstanceStorage& InstanceData, int32 PerThreadSharedIndex, FDebugInstanceData::EContainer Container)
+		{
+			const int32 NumObject = InstanceData.Num();
+			for (int32 Index = 0; Index < NumObject; ++Index)
+			{
+				FConstStructView Instance = InstanceData.GetStruct(Index);
+				if (!Instance.IsValid())
+				{
+					const TCHAR* Msg = (Container == FDebugInstanceData::EContainer::SharedInstance) ? TEXT("A shared") : TEXT("An");
+					ensureAlwaysMsgf(false, TEXT("%s instance data is invalid. InstanceDataIndex:%d. PerThreadSharedIndex:%d."), Msg, Index, PerThreadSharedIndex);
+					UE::StateTree::Debug::Private::FindCVarInstanceDataGC()->Set(false);
+					break;
+				}
+
+				if (const FStateTreeInstanceObjectWrapper* Wrapper = Instance.GetPtr<const FStateTreeInstanceObjectWrapper>())
+				{
+					if (!IsValid(Wrapper->InstanceObject))
+					{
+						const TCHAR* Msg = (Container == FDebugInstanceData::EContainer::SharedInstance) ? TEXT("A shared") : TEXT("An");
+						ensureAlwaysMsgf(false, TEXT("%s instance data object is invalid. InstanceDataIndex:%d. PerThreadSharedIndex:%d."), Msg, Index, PerThreadSharedIndex);
+						UE::StateTree::Debug::Private::FindCVarInstanceDataGC()->Set(false);
+						break;
+					}
+
+					GCObjectDatas.Emplace(FWeakObjectPtr(Wrapper->InstanceObject), Index, PerThreadSharedIndex, Container, FDebugInstanceData::EObjectType::ObjectInstance);
+				}
+				else
+				{
+					GCObjectDatas.Emplace(FWeakObjectPtr(Instance.GetScriptStruct()), Index, PerThreadSharedIndex, Container, FDebugInstanceData::EObjectType::Struct);
+				}
+			}
+		};
+
+	AddInstanceObject(DefaultInstanceData.GetStorage(), INDEX_NONE, FDebugInstanceData::EContainer::DefaultInstance);
+	AddInstanceObject(SharedInstanceData.GetStorage(), INDEX_NONE, FDebugInstanceData::EContainer::SharedInstance);
+
+	{
+		UE::TReadScopeLock ReadLock(PerThreadSharedInstanceDataLock);
+		for (int32 Index = 0; Index < PerThreadSharedInstanceData.Num(); ++Index)
+		{
+			TSharedPtr<FStateTreeInstanceData>& SharedInstance = PerThreadSharedInstanceData[Index];
+			if (ensureMsgf(SharedInstance, TEXT("The shared instance data is invalid. Index %d"), Index))
+			{
+				AddInstanceObject(SharedInstance->GetStorage(), Index, FDebugInstanceData::EContainer::SharedInstance);
+			}
+		}
+	}
+}
+
+void UStateTree::HandleRuntimeValidationPostGC()
+{
+	for (const FDebugInstanceData& Data : GCObjectDatas)
+	{
+		if (Data.Object.IsStale())
+		{
+			if (Data.Type == FDebugInstanceData::EObjectType::ObjectInstance)
+			{
+				const TCHAR* Msg = (Data.Container == FDebugInstanceData::EContainer::SharedInstance) ? TEXT("A shared") : TEXT("An");
+				ensureAlwaysMsgf(false, TEXT("%s instance data object is GCed. InstDataIndex:%d. SharedIndex:%d."), Msg, Data.InstanceDataStructIndex, Data.SharedInstanceDataIndex);
+			}
+			else
+			{
+				const TCHAR* Msg = (Data.Container == FDebugInstanceData::EContainer::SharedInstance) ? TEXT("A shared") : TEXT("An");
+				ensureAlwaysMsgf(false, TEXT("%s instance data struct type is GCed. InstDataIndex:%d. SharedIndex:%d."), Msg, Data.InstanceDataStructIndex, Data.SharedInstanceDataIndex);
+			}
+			UE::StateTree::Debug::Private::FindCVarInstanceDataGC()->Set(false);
+			break;
+		}
+	}
+	GCObjectDatas.Reset();
+}
+#endif
diff --git a/Engine/Plugins/Runtime/StateTree/Source/StateTreeModule/Private/StateTreeModule.cpp b/Engine/Plugins/Runtime/StateTree/Source/StateTreeModule/Private/StateTreeModule.cpp
--- a/Engine/Plugins/Runtime/StateTree/Source/StateTreeModule/Private/StateTreeModule.cpp
+++ b/Engine/Plugins/Runtime/StateTree/Source/StateTreeModule/Private/StateTreeModule.cpp
@@ -33,6 +33,10 @@
 #include "StructUtilsDelegates.h"
 #endif
 
+#if WITH_STATETREE_DEBUG
+#include "UObject/UObjectGlobals.h"
+#endif
+
 #define LOCTEXT_NAMESPACE "StateTree"
 
 #if WITH_EDITOR
@@ -41,6 +45,23 @@ FStateTreeModule::FOnPIEEvent FStateTreeModule::OnPreBeginPIE;
 FStateTreeModule::FOnUserDefinedStructReinstanced FStateTreeModule::OnUserDefinedStructReinstanced;
 #endif
 
+#if WITH_STATETREE_DEBUG
+FTSSimpleMulticastDelegate FStateTreeModule::OnPreRuntimeValidationInstanceData;
+FTSSimpleMulticastDelegate FStateTreeModule::OnPostRuntimeValidationInstanceData;
+#endif
+
+#if WITH_STATETREE_DEBUG
+namespace UE::StateTree::Debug::Private
+{
+	bool bRuntimeValidationInstanceDataGC = false;
+	static FAutoConsoleVariableRef CVarRuntimeValidationInstanceDataGC(
+		TEXT("StateTree.RuntimeValidation.InstanceDataGC"),
+		bRuntimeValidationInstanceDataGC,
+		TEXT("Test after each GC if nodes were properly GCed.")
+	);
+} // namespace UE::StateTree::Debug::Private
+#endif //WITH_STATETREE_DEBUG
+
 #if WITH_STATETREE_TRACE_DEBUGGER
 UE::Trace::FStoreClient* FStateTreeModule::GetStoreClient()
 {
@@ -126,10 +147,20 @@ void FStateTreeModule::StartupModule()
 		return false;
 	 });
 #endif //WITH_EDITOR
+
+#if WITH_STATETREE_DEBUG
+	 PreGCHandle = FCoreUObjectDelegates::GetPreGarbageCollectDelegate().AddRaw(this, &FStateTreeModule::HandlePreGC);
+	 PostGCHandle = FCoreUObjectDelegates::GetPostPurgeGarbageDelegate().AddRaw(this, &FStateTreeModule::HandlePostGC);
+#endif
 }
 
 void FStateTreeModule::ShutdownModule()
 {
+#if WITH_STATETREE_DEBUG
+	FCoreUObjectDelegates::GetPreGarbageCollectDelegate().Remove(PreGCHandle);
+	FCoreUObjectDelegates::GetPostPurgeGarbageDelegate().Remove(PostGCHandle);
+#endif
+
 #if WITH_EDITOR
 	FCoreUObjectDelegates::OnObjectsReinstanced.Remove(OnObjectsReinstancedHandle);
 	UE::StructUtils::Delegates::OnUserDefinedStructReinstanced.Remove(OnUserDefinedStructReinstancedHandle);
@@ -297,6 +328,24 @@ void FStateTreeModule::HandleUserDefinedStructReinstanced(const UUserDefinedStru
 }
 #endif
 
+#if WITH_STATETREE_DEBUG
+void FStateTreeModule::HandlePreGC()
+{
+	if (UE::StateTree::Debug::Private::bRuntimeValidationInstanceDataGC)
+	{
+		OnPreRuntimeValidationInstanceData.Broadcast();
+	}
+}
+
+void FStateTreeModule::HandlePostGC()
+{
+	if (UE::StateTree::Debug::Private::bRuntimeValidationInstanceDataGC)
+	{
+		OnPostRuntimeValidationInstanceData.Broadcast();
+	}
+}
+#endif
+
 IMPLEMENT_MODULE(FStateTreeModule, StateTreeModule)
 
 #undef LOCTEXT_NAMESPACE
diff --git a/Engine/Plugins/Runtime/StateTree/Source/StateTreeModule/Private/StateTreeModuleImpl.h b/Engine/Plugins/Runtime/StateTree/Source/StateTreeModule/Private/StateTreeModuleImpl.h
--- a/Engine/Plugins/Runtime/StateTree/Source/StateTreeModule/Private/StateTreeModuleImpl.h
+++ b/Engine/Plugins/Runtime/StateTree/Source/StateTreeModule/Private/StateTreeModuleImpl.h
@@ -5,6 +5,7 @@
 #include "StateTreeModule.h"
 #include "Debugger/StateTreeTraceModule.h"
 #include "HAL/IConsoleManager.h"
+#include "StateTreeTypes.h"
 #include "UObject/UObjectGlobals.h"
 
 #if WITH_EDITOR
@@ -42,6 +43,11 @@ public:
 	static FOnUserDefinedStructReinstanced OnUserDefinedStructReinstanced;
 #endif
 
+#if WITH_STATETREE_DEBUG
+	static FTSSimpleMulticastDelegate OnPreRuntimeValidationInstanceData;
+	static FTSSimpleMulticastDelegate OnPostRuntimeValidationInstanceData;
+#endif
+
 private:
 #if WITH_EDITOR
 	void HandleObjectsReinstanced(const FReplacementObjectMap& ObjectMap);
@@ -49,6 +55,11 @@ private:
 	void HandleUserDefinedStructReinstanced(const UUserDefinedStruct& UserDefinedStruct);
 #endif
 
+#if WITH_STATETREE_DEBUG
+	void HandlePreGC();
+	void HandlePostGC();
+#endif
+
 #if WITH_STATETREE_TRACE_DEBUGGER
 	/**
 	 * Gets the store client.
@@ -79,5 +90,10 @@ private:
 	FDelegateHandle OnObjectsReinstancedHandle;
 	FDelegateHandle OnUserDefinedStructReinstancedHandle;
 	FDelegateHandle OnPreBeginPIEHandle;
-#endif 
+#endif
+
+#if WITH_STATETREE_DEBUG
+	FDelegateHandle PreGCHandle;
+	FDelegateHandle PostGCHandle;
+#endif
 };
diff --git a/Engine/Plugins/Runtime/StateTree/Source/StateTreeModule/Public/StateTree.h b/Engine/Plugins/Runtime/StateTree/Source/StateTreeModule/Public/StateTree.h
--- a/Engine/Plugins/Runtime/StateTree/Source/StateTreeModule/Public/StateTree.h
+++ b/Engine/Plugins/Runtime/StateTree/Source/StateTreeModule/Public/StateTree.h
@@ -257,13 +257,16 @@ protected:
 	static UE_API void DeclareConstructClasses(TArray<FTopLevelAssetPath>& OutConstructClasses, const UClass* SpecificSubclass);
 #endif
 	UE_API virtual void Serialize(FStructuredArchiveRecord Record) override;
-	
+
+#if WITH_EDITOR || WITH_STATETREE_DEBUG
+	UE_API virtual void BeginDestroy() override;
+	UE_API virtual void PostInitProperties() override;
+#endif
+
 #if WITH_EDITOR
 	using FReplacementObjectMap = TMap<UObject*, UObject*>;
 	UE_API void OnObjectsReinstanced(const FReplacementObjectMap& ObjectMap);
 	UE_API void OnUserDefinedStructReinstanced(const UUserDefinedStruct& UserDefinedStruct);
-	UE_API virtual void PostInitProperties() override;
-	UE_API virtual void BeginDestroy() override;
 	UE_API virtual void PostDuplicate(EDuplicateMode::Type DuplicateMode) override;
 	UE_API virtual void GetAssetRegistryTags(FAssetRegistryTagsContext Context) const override;
 protected:
@@ -430,6 +433,35 @@ private:
 	TSet<FObjectKey> OutOfDateStructs;
 #endif
 
+private:
+#if WITH_STATETREE_DEBUG
+	//~ Info for RuntimeValidation InstanceData GC
+	struct FDebugInstanceData
+	{
+		FWeakObjectPtr Object;
+		int32 InstanceDataStructIndex = INDEX_NONE;
+		int32 SharedInstanceDataIndex = INDEX_NONE;
+		enum class EContainer : uint8
+		{
+			DefaultInstance,
+			SharedInstance,
+		};
+		EContainer Container = EContainer::DefaultInstance;
+		enum class EObjectType : uint8
+		{
+			ObjectInstance,
+			Struct,
+		};
+		EObjectType Type = EObjectType::ObjectInstance;
+	};
+	TArray<FDebugInstanceData> GCObjectDatas;
+	FDelegateHandle PreGCHandle;
+	FDelegateHandle PostGCHandle;
+
+	void HandleRuntimeValidationPreGC();
+	void HandleRuntimeValidationPostGC();
+#endif
+
 	friend struct FStateTreeInstance;
 	friend struct FStateTreeExecutionContext;
 	friend struct FStateTreeTasksCompletionStatus;
