[StateTree] Fix build on linux.

#rnx

diff --git a/Engine/Plugins/Runtime/StateTree/Source/StateTreeModule/Private/StateTree.cpp b/Engine/Plugins/Runtime/StateTree/Source/StateTreeModule/Private/StateTree.cpp
--- a/Engine/Plugins/Runtime/StateTree/Source/StateTreeModule/Private/StateTree.cpp
+++ b/Engine/Plugins/Runtime/StateTree/Source/StateTreeModule/Private/StateTree.cpp
@@ -1550,6 +1550,7 @@ namespace UE::StateTree::Debug::Private
 		return FoundVariable;
 	}
 }
+
 void UStateTree::HandleRuntimeValidationPreGC()
 {
 	GCObjectDatas.Reset();
@@ -1564,9 +1565,15 @@ void UStateTree::HandleRuntimeValidationPreGC()
 					const TCHAR* Msg = (Container == FDebugInstanceData::EContainer::SharedInstance) ? TEXT("A shared") : TEXT("An");
 					ensureAlwaysMsgf(false, TEXT("%s instance data is invalid. InstanceDataIndex:%d. PerThreadSharedIndex:%d."), Msg, Index, PerThreadSharedIndex);
 					UE::StateTree::Debug::Private::FindCVarInstanceDataGC()->Set(false);
+					GCObjectDatas.Reset();
 					break;
 				}
 
+				FDebugInstanceData& Data = GCObjectDatas.AddDefaulted_GetRef();
+				Data.InstanceDataStructIndex = Index;
+				Data.SharedInstanceDataIndex = PerThreadSharedIndex;
+				Data.Container = Container;
+
 				if (const FStateTreeInstanceObjectWrapper* Wrapper = Instance.GetPtr<const FStateTreeInstanceObjectWrapper>())
 				{
 					if (!IsValid(Wrapper->InstanceObject))
@@ -1574,14 +1581,17 @@ void UStateTree::HandleRuntimeValidationPreGC()
 						const TCHAR* Msg = (Container == FDebugInstanceData::EContainer::SharedInstance) ? TEXT("A shared") : TEXT("An");
 						ensureAlwaysMsgf(false, TEXT("%s instance data object is invalid. InstanceDataIndex:%d. PerThreadSharedIndex:%d."), Msg, Index, PerThreadSharedIndex);
 						UE::StateTree::Debug::Private::FindCVarInstanceDataGC()->Set(false);
+						GCObjectDatas.Reset();
 						break;
 					}
 
-					GCObjectDatas.Emplace(FWeakObjectPtr(Wrapper->InstanceObject), Index, PerThreadSharedIndex, Container, FDebugInstanceData::EObjectType::ObjectInstance);
+					Data.Object = FWeakObjectPtr(Wrapper->InstanceObject);
+					Data.Type = FDebugInstanceData::EObjectType::ObjectInstance;
 				}
 				else
 				{
-					GCObjectDatas.Emplace(FWeakObjectPtr(Instance.GetScriptStruct()), Index, PerThreadSharedIndex, Container, FDebugInstanceData::EObjectType::Struct);
+					Data.Object = FWeakObjectPtr(Instance.GetScriptStruct());
+					Data.Type = FDebugInstanceData::EObjectType::Struct;
 				}
 			}
 		};
@@ -1614,7 +1624,7 @@ void UStateTree::HandleRuntimeValidationPostGC()
 				ensureAlwaysMsgf(false, TEXT("%s instance data object is GCed. InstDataIndex:%d. SharedIndex:%d."), Msg, Data.InstanceDataStructIndex, Data.SharedInstanceDataIndex);
 			}
 			else
-			{
+			{ //~V523 disabling identical branch warning
 				const TCHAR* Msg = (Data.Container == FDebugInstanceData::EContainer::SharedInstance) ? TEXT("A shared") : TEXT("An");
 				ensureAlwaysMsgf(false, TEXT("%s instance data struct type is GCed. InstDataIndex:%d. SharedIndex:%d."), Msg, Data.InstanceDataStructIndex, Data.SharedInstanceDataIndex);
 			}
