[StateTree] Disable V523 disabling identical branch warning

#rnx

diff --git a/Engine/Plugins/Runtime/StateTree/Source/StateTreeModule/Private/StateTree.cpp b/Engine/Plugins/Runtime/StateTree/Source/StateTreeModule/Private/StateTree.cpp
--- a/Engine/Plugins/Runtime/StateTree/Source/StateTreeModule/Private/StateTree.cpp
+++ b/Engine/Plugins/Runtime/StateTree/Source/StateTreeModule/Private/StateTree.cpp
@@ -1624,7 +1624,7 @@ void UStateTree::HandleRuntimeValidationPostGC()
 				ensureAlwaysMsgf(false, TEXT("%s instance data object is GCed. InstDataIndex:%d. SharedIndex:%d."), Msg, Data.InstanceDataStructIndex, Data.SharedInstanceDataIndex);
 			}
 			else
-			{ //~V523 disabling identical branch warning
+			{ //-V523 disabling identical branch warning
 				const TCHAR* Msg = (Data.Container == FDebugInstanceData::EContainer::SharedInstance) ? TEXT("A shared") : TEXT("An");
 				ensureAlwaysMsgf(false, TEXT("%s instance data struct type is GCed. InstDataIndex:%d. SharedIndex:%d."), Msg, Data.InstanceDataStructIndex, Data.SharedInstanceDataIndex);
 			}
