[StateTreeDebugger] update trace API to always provide the owner so all events can make sure to trace the parent object for late recording in RewindDebugger

diff --git a/Engine/Plugins/Runtime/StateTree/Source/StateTreeModule/Private/Debugger/StateTreeTrace.cpp b/Engine/Plugins/Runtime/StateTree/Source/StateTreeModule/Private/Debugger/StateTreeTrace.cpp
--- a/Engine/Plugins/Runtime/StateTree/Source/StateTreeModule/Private/Debugger/StateTreeTrace.cpp
+++ b/Engine/Plugins/Runtime/StateTree/Source/StateTreeModule/Private/Debugger/StateTreeTrace.cpp
@@ -204,7 +204,7 @@ struct FInstanceEventBufferedData
 	{
 	}
 
-	void CloseRecording(double WorldTime) const
+	void CloseRecording(const double WorldTime) const
 	{
 		// Output and empty active states event at the last recorded world time to close active states
 		UE_TRACE_LOG(StateTreeDebugger, WorldTimestampEvent, StateTreeDebugChannel) << WorldTimestampEvent.WorldTime(WorldTime);
@@ -596,7 +596,12 @@ FBufferedDataList GBufferedEvents;
  * Pushed or pops an entry on the Phase stack for a given Instance.
  * Will send the Pop events for phases popped if their associated Push events were sent.
  */
-void OutputPhaseScopeEvent(const FStateTreeInstanceDebugId InstanceId, const EStateTreeUpdatePhase Phase, const EStateTreeTraceEventType EventType, const FStateTreeStateHandle StateHandle)
+void OutputPhaseScopeEvent(
+	const UObject* Owner,
+	const FStateTreeInstanceDebugId InstanceId,
+	const EStateTreeUpdatePhase Phase,
+	const EStateTreeTraceEventType EventType,
+	const FStateTreeStateHandle StateHandle)
 {
 	const TOptional<FPhaseTraceStatusPair> RemovedPair = GBufferedEvents.UpdatePhaseScope_AnyThread(InstanceId, Phase, EventType, StateHandle);
 
@@ -708,18 +713,6 @@ void UnregisterGlobalDelegates()
 	GTracingStateChangedDelegateHandle.Reset();
 }
 
-// Deprecated
-FStateTreeIndex16 FindOrAddDebugIdForAsset(const UStateTree* StateTree)
-{
-	return GBufferedEvents.FindOrAddDebugIdForAsset_AnyThread(StateTree);
-}
-
-// Deprecated
-void OutputInstanceLifetimeEvent(const FStateTreeInstanceDebugId InstanceId, const UStateTree* StateTree, const TCHAR* InstanceName, const EStateTreeTraceEventType EventType)
-{
-	OutputInstanceLifetimeEvent(nullptr, InstanceId, StateTree, InstanceName, EventType);
-}
-
 void OutputAssetDebugIdEvent(
 	const UStateTree* StateTree,
 	const FStateTreeIndex16 AssetDebugId
@@ -752,21 +745,18 @@ void OutputInstanceLifetimeEvent(
 {
 	if (UE_TRACE_CHANNELEXPR_IS_ENABLED(StateTreeDebugChannel))
 	{
-		TraceBufferedEvents(InstanceId);
-
-		if (Owner != nullptr)
+		const uint64 ChildId(InstanceId.ToUint64());
+		if (EventType == EStateTreeTraceEventType::Push)
 		{
-			const uint64 ChildId(InstanceId.ToUint64());
-			if (EventType == EStateTreeTraceEventType::Push)
-			{
-				TRACE_CHILD_ELEMENT_LIFETIME_BEGIN(Owner, ChildId);
-			}
-			else if (EventType == EStateTreeTraceEventType::Pop)
-			{
-				TRACE_CHILD_ELEMENT_LIFETIME_END(Owner, ChildId);
-			}
+			TRACE_CHILD_ELEMENT_LIFETIME_BEGIN(Owner, ChildId);
+		}
+		else if (EventType == EStateTreeTraceEventType::Pop)
+		{
+			TRACE_CHILD_ELEMENT_LIFETIME_END(Owner, ChildId);
 		}
 
+		TraceBufferedEvents(InstanceId);
+
 		const FStateTreeIndex16 AssetDebugId = GBufferedEvents.FindOrAddDebugIdForAsset_AnyThread(StateTree);
 
 		UE_TRACE_LOG(StateTreeDebugger, InstanceEvent, StateTreeDebugChannel)
@@ -785,6 +775,7 @@ void OutputInstanceLifetimeEvent(
 }
 
 void OutputInstanceAssetEvent(
+	const UObject* Owner,
 	const FStateTreeInstanceDebugId InstanceId,
 	const UStateTree* StateTree
 	)
@@ -793,6 +784,7 @@ void OutputInstanceAssetEvent(
 
 	if (UE_TRACE_CHANNELEXPR_IS_ENABLED(StateTreeDebugChannel))
 	{
+		TRACE_CHILD_ELEMENT(Owner, InstanceId.ToUint64());
 		TraceBufferedEvents(InstanceId);
 
 		const FStateTreeIndex16 AssetDebugId = GBufferedEvents.FindOrAddDebugIdForAsset_AnyThread(StateTree);
@@ -810,23 +802,26 @@ void OutputInstanceAssetEvent(
 }
 
 void OutputInstanceFrameEvent(
+	const UObject* Owner,
 	const FStateTreeInstanceDebugId InstanceId,
 	const FStateTreeExecutionFrame* Frame
 	)
 {
 	check(Frame != nullptr);
-	OutputInstanceAssetEvent(InstanceId, Frame->StateTree.Get());
+	OutputInstanceAssetEvent(Owner, InstanceId, Frame->StateTree.Get());
 }
 
 void OutputLogEventTrace(
+	const UObject* Owner,
 	const FStateTreeInstanceDebugId InstanceId,
-	ELogVerbosity::Type Verbosity,
+	const ELogVerbosity::Type Verbosity,
 	const TCHAR* Fmt, ...
 	)
 {
 	static TCHAR TraceStaticBuffer[8192];
 	GET_TYPED_VARARGS(TCHAR, TraceStaticBuffer, UE_ARRAY_COUNT(TraceStaticBuffer), UE_ARRAY_COUNT(TraceStaticBuffer) - 1, Fmt, Fmt);
 
+	TRACE_CHILD_ELEMENT(Owner, InstanceId.ToUint64());
 	TraceBufferedEvents(InstanceId);
 
 	UE_TRACE_LOG(StateTreeDebugger, LogEvent, StateTreeDebugChannel)
@@ -838,11 +833,13 @@ void OutputLogEventTrace(
 }
 
 void OutputStateEventTrace(
+	const UObject* Owner,
 	const FStateTreeInstanceDebugId InstanceId,
 	const FStateTreeStateHandle StateHandle,
 	const EStateTreeTraceEventType EventType
 	)
 {
+	TRACE_CHILD_ELEMENT(Owner, InstanceId.ToUint64());
 	TraceBufferedEvents(InstanceId);
 
 	UE_TRACE_LOG(StateTreeDebugger, StateEvent, StateTreeDebugChannel)
@@ -854,14 +851,16 @@ void OutputStateEventTrace(
 }
 
 void OutputTaskEventTrace(
+	const UObject* Owner,
 	const FStateTreeInstanceDebugId InstanceId,
 	FNodeCustomDebugData&& CustomDebugData,
-	FStateTreeIndex16 TaskIdx, FStateTreeDataView DataView,
+	const FStateTreeIndex16 TaskIdx, const FStateTreeDataView DataView,
 	EStateTreeTraceEventType EventType, EStateTreeRunStatus Status)
 {
 	FBufferArchive Archive;
 	SerializeDebugDataToArchive(Archive, MoveTemp(CustomDebugData), DataView);
 
+	TRACE_CHILD_ELEMENT(Owner, InstanceId.ToUint64());
 	TraceBufferedEvents(InstanceId);
 
 	UE_TRACE_LOG(StateTreeDebugger, TaskEvent, StateTreeDebugChannel)
@@ -875,15 +874,17 @@ void OutputTaskEventTrace(
 }
 
 void OutputEvaluatorEventTrace(
+	const UObject* Owner,
 	const FStateTreeInstanceDebugId InstanceId,
 	FNodeCustomDebugData&& CustomDebugData,
-	FStateTreeIndex16 EvaluatorIdx,
-	FStateTreeDataView DataView,
+	const FStateTreeIndex16 EvaluatorIdx,
+	const FStateTreeDataView DataView,
 	EStateTreeTraceEventType EventType)
 {
 	FBufferArchive Archive;
 	SerializeDebugDataToArchive(Archive, MoveTemp(CustomDebugData), DataView);
 
+	TRACE_CHILD_ELEMENT(Owner, InstanceId.ToUint64());
 	TraceBufferedEvents(InstanceId);
 
 	UE_TRACE_LOG(StateTreeDebugger, EvaluatorEvent, StateTreeDebugChannel)
@@ -896,6 +897,7 @@ void OutputEvaluatorEventTrace(
 }
 
 void OutputTransitionEventTrace(
+	const UObject* Owner,
 	const FStateTreeInstanceDebugId InstanceId,
 	const FStateTreeTransitionSource Source,
 	const EStateTreeTraceEventType EventType
@@ -904,12 +906,13 @@ void OutputTransitionEventTrace(
 	// Output an instance Frame/Asset event to put the transition data in the proper context when it gets read.
 	if (const UStateTree* StateTree = Source.Asset.Get())
 	{
-		OutputInstanceAssetEvent(InstanceId, StateTree);
+		OutputInstanceAssetEvent(Owner, InstanceId, StateTree);
 	}
 
 	FBufferArchive Archive;
 	Archive << EventType;
 
+	TRACE_CHILD_ELEMENT(Owner, InstanceId.ToUint64());
 	TraceBufferedEvents(InstanceId);
 
 	UE_TRACE_LOG(StateTreeDebugger, TransitionEvent, StateTreeDebugChannel)
@@ -924,15 +927,17 @@ void OutputTransitionEventTrace(
 }
 
 void OutputConditionEventTrace(
+	const UObject* Owner,
 	const FStateTreeInstanceDebugId InstanceId,
 	FNodeCustomDebugData&& CustomDebugData,
-	FStateTreeIndex16 ConditionIdx,
-	FStateTreeDataView DataView,
+	const FStateTreeIndex16 ConditionIdx,
+	const FStateTreeDataView DataView,
 	EStateTreeTraceEventType EventType)
 {
 	FBufferArchive Archive;
 	SerializeDebugDataToArchive(Archive, MoveTemp(CustomDebugData), DataView);
 
+	TRACE_CHILD_ELEMENT(Owner, InstanceId.ToUint64());
 	TraceBufferedEvents(InstanceId);
 
 	UE_TRACE_LOG(StateTreeDebugger, ConditionEvent, StateTreeDebugChannel)
@@ -945,12 +950,14 @@ void OutputConditionEventTrace(
 }
 
 void OutputActiveStatesEventTrace(
+	const UObject* Owner,
 	const FStateTreeInstanceDebugId InstanceId,
 	const TConstArrayView<FStateTreeExecutionFrame> ActiveFrames
 	)
 {
 	if (UE_TRACE_CHANNELEXPR_IS_ENABLED(StateTreeDebugChannel))
 	{
+		TRACE_CHILD_ELEMENT(Owner, InstanceId.ToUint64());
 		TraceBufferedEvents(InstanceId);
 
 		GBufferedEvents.OutputActiveStates_AnyThread(InstanceId, ActiveFrames);
@@ -961,6 +968,80 @@ void OutputActiveStatesEventTrace(
 	}
 }
 
+// Deprecated
+FStateTreeIndex16 FindOrAddDebugIdForAsset(const UStateTree* StateTree)
+{
+	return GBufferedEvents.FindOrAddDebugIdForAsset_AnyThread(StateTree);
+}
+
+// Deprecated
+void OutputPhaseScopeEvent(const FStateTreeInstanceDebugId InstanceId, const EStateTreeUpdatePhase Phase, const EStateTreeTraceEventType EventType, const FStateTreeStateHandle StateHandle)
+{
+	OutputPhaseScopeEvent(nullptr, InstanceId, Phase, EventType, StateHandle);
+}
+
+// Deprecated
+void OutputInstanceLifetimeEvent(const FStateTreeInstanceDebugId InstanceId, const UStateTree* StateTree, const TCHAR* InstanceName, const EStateTreeTraceEventType EventType)
+{
+	OutputInstanceLifetimeEvent(nullptr, InstanceId, StateTree, InstanceName, EventType);
+}
+
+// Deprecated
+void OutputInstanceAssetEvent(const FStateTreeInstanceDebugId InstanceId, const UStateTree* StateTree)
+{
+	OutputInstanceAssetEvent(nullptr, InstanceId, StateTree);
+}
+// Deprecated
+void OutputInstanceFrameEvent(const FStateTreeInstanceDebugId InstanceId, const FStateTreeExecutionFrame* Frame)
+{
+	OutputInstanceFrameEvent(nullptr, InstanceId, Frame);
+}
+
+// Deprecated
+void OutputLogEventTrace(const FStateTreeInstanceDebugId InstanceId, const ELogVerbosity::Type Verbosity, const TCHAR* Fmt, ...)
+{
+	va_list Args;
+	va_start(Args, Fmt);
+	OutputLogEventTrace(nullptr, InstanceId, Verbosity, Fmt, Args);
+	va_end(Args);
+}
+
+// Deprecated
+void OutputStateEventTrace(const FStateTreeInstanceDebugId InstanceId, const FStateTreeStateHandle StateHandle, const EStateTreeTraceEventType EventType)
+{
+	OutputStateEventTrace(nullptr, InstanceId, StateHandle, EventType);
+}
+
+// Deprecated
+void OutputTaskEventTrace(const FStateTreeInstanceDebugId InstanceId, FNodeCustomDebugData&& DebugData, const FStateTreeIndex16 TaskIdx, const FStateTreeDataView DataView, const EStateTreeTraceEventType EventType, const EStateTreeRunStatus Status)
+{
+	OutputTaskEventTrace(nullptr, InstanceId, MoveTemp(DebugData), TaskIdx, DataView, EventType, Status);
+}
+
+// Deprecated
+void OutputEvaluatorEventTrace(const FStateTreeInstanceDebugId InstanceId, FNodeCustomDebugData&& DebugData, const FStateTreeIndex16 EvaluatorIdx, const FStateTreeDataView DataView, const EStateTreeTraceEventType EventType)
+{
+	OutputEvaluatorEventTrace(nullptr, InstanceId, MoveTemp(DebugData), EvaluatorIdx, DataView, EventType);
+}
+
+// Deprecated
+void OutputConditionEventTrace(const FStateTreeInstanceDebugId InstanceId, FNodeCustomDebugData&& DebugData, const FStateTreeIndex16 ConditionIdx, const FStateTreeDataView DataView, const EStateTreeTraceEventType EventType)
+{
+	OutputConditionEventTrace(nullptr, InstanceId, MoveTemp(DebugData), ConditionIdx, DataView, EventType);
+}
+
+// Deprecated
+void OutputTransitionEventTrace(const FStateTreeInstanceDebugId InstanceId, const FStateTreeTransitionSource TransitionSource, const EStateTreeTraceEventType EventType)
+{
+	OutputTransitionEventTrace(nullptr, InstanceId, TransitionSource, EventType);
+}
+
+// Deprecated
+void OutputActiveStatesEventTrace(const FStateTreeInstanceDebugId InstanceId, const TConstArrayView<FStateTreeExecutionFrame> ActiveFrames)
+{
+	OutputActiveStatesEventTrace(nullptr, InstanceId, ActiveFrames);
+}
+
 } // UE::StateTreeTrace
 
 #endif // WITH_STATETREE_TRACE
diff --git a/Engine/Plugins/Runtime/StateTree/Source/StateTreeModule/Public/Debugger/StateTreeDebug.h b/Engine/Plugins/Runtime/StateTree/Source/StateTreeModule/Public/Debugger/StateTreeDebug.h
--- a/Engine/Plugins/Runtime/StateTree/Source/StateTreeModule/Public/Debugger/StateTreeDebug.h
+++ b/Engine/Plugins/Runtime/StateTree/Source/StateTreeModule/Public/Debugger/StateTreeDebug.h
@@ -13,75 +13,75 @@ enum class EStateTreeTraceEventType : uint8;
 
 #if WITH_STATETREE_DEBUG
 
-#define UE_STATETREE_DEBUG_CONDITION_ENTER_STATE(InExecutionContextPtr, InStateTree, InNodeIndex) ::UE::StateTree::Debug::ConditionEnterState(*(InExecutionContextPtr), ::UE::StateTree::Debug::FNodeReference((InStateTree), FStateTreeIndex16(InNodeIndex)))
-#define UE_STATETREE_DEBUG_CONDITION_TEST_CONDITION(InExecutionContextPtr, InStateTree, InNodeIndex) ::UE::StateTree::Debug::ConditionTest(*(InExecutionContextPtr), ::UE::StateTree::Debug::FNodeReference((InStateTree), FStateTreeIndex16(InNodeIndex)))
-#define UE_STATETREE_DEBUG_CONDITION_EXIT_STATE(InExecutionContextPtr, InStateTree, InNodeIndex) ::UE::StateTree::Debug::ConditionExitState(*(InExecutionContextPtr), ::UE::StateTree::Debug::FNodeReference((InStateTree), FStateTreeIndex16(InNodeIndex)))
+#define UE_STATETREE_DEBUG_CONDITION_ENTER_STATE(InCtx, InStateTree, InNodeIndex) ::UE::StateTree::Debug::ConditionEnterState(*(InCtx), ::UE::StateTree::Debug::FNodeReference((InStateTree), FStateTreeIndex16(InNodeIndex)))
+#define UE_STATETREE_DEBUG_CONDITION_TEST_CONDITION(InCtx, InStateTree, InNodeIndex) ::UE::StateTree::Debug::ConditionTest(*(InCtx), ::UE::StateTree::Debug::FNodeReference((InStateTree), FStateTreeIndex16(InNodeIndex)))
+#define UE_STATETREE_DEBUG_CONDITION_EXIT_STATE(InCtx, InStateTree, InNodeIndex) ::UE::StateTree::Debug::ConditionExitState(*(InCtx), ::UE::StateTree::Debug::FNodeReference((InStateTree), FStateTreeIndex16(InNodeIndex)))
 
-#define UE_STATETREE_DEBUG_EVALUATOR_ENTER_TREE(InExecutionContextPtr, InStateTree, InNodeIndex) ::UE::StateTree::Debug::EvaluatorEnterTree(*(InExecutionContextPtr), ::UE::StateTree::Debug::FNodeReference((InStateTree), FStateTreeIndex16(InNodeIndex)))
-#define UE_STATETREE_DEBUG_EVALUATOR_TICK(InExecutionContextPtr, InStateTree, InNodeIndex) ::UE::StateTree::Debug::EvaluatorTick(*(InExecutionContextPtr), ::UE::StateTree::Debug::FNodeReference((InStateTree), FStateTreeIndex16(InNodeIndex)))
-#define UE_STATETREE_DEBUG_EVALUATOR_EXIT_TREE(InExecutionContextPtr, InStateTree, InNodeIndex) ::UE::StateTree::Debug::EvaluatorExitTree(*(InExecutionContextPtr), ::UE::StateTree::Debug::FNodeReference((InStateTree), FStateTreeIndex16(InNodeIndex)))
+#define UE_STATETREE_DEBUG_EVALUATOR_ENTER_TREE(InCtx, InStateTree, InNodeIndex) ::UE::StateTree::Debug::EvaluatorEnterTree(*(InCtx), ::UE::StateTree::Debug::FNodeReference((InStateTree), FStateTreeIndex16(InNodeIndex)))
+#define UE_STATETREE_DEBUG_EVALUATOR_TICK(InCtx, InStateTree, InNodeIndex) ::UE::StateTree::Debug::EvaluatorTick(*(InCtx), ::UE::StateTree::Debug::FNodeReference((InStateTree), FStateTreeIndex16(InNodeIndex)))
+#define UE_STATETREE_DEBUG_EVALUATOR_EXIT_TREE(InCtx, InStateTree, InNodeIndex) ::UE::StateTree::Debug::EvaluatorExitTree(*(InCtx), ::UE::StateTree::Debug::FNodeReference((InStateTree), FStateTreeIndex16(InNodeIndex)))
 
-#define UE_STATETREE_DEBUG_TASK_ENTER_STATE(InExecutionContextPtr, InStateTree, InNodeIndex) ::UE::StateTree::Debug::TaskEnterState(*(InExecutionContextPtr), ::UE::StateTree::Debug::FNodeReference((InStateTree), FStateTreeIndex16(InNodeIndex)))
-#define UE_STATETREE_DEBUG_TASK_TICK(InExecutionContextPtr, InStateTree, InNodeIndex) ::UE::StateTree::Debug::TaskTick(*(InExecutionContextPtr), ::UE::StateTree::Debug::FNodeReference((InStateTree), FStateTreeIndex16(InNodeIndex)))
-#define UE_STATETREE_DEBUG_TASK_EXIT_STATE(InExecutionContextPtr, InStateTree, InNodeIndex) ::UE::StateTree::Debug::TaskExitState(*(InExecutionContextPtr), ::UE::StateTree::Debug::FNodeReference((InStateTree), FStateTreeIndex16(InNodeIndex)))
+#define UE_STATETREE_DEBUG_TASK_ENTER_STATE(InCtx, InStateTree, InNodeIndex) ::UE::StateTree::Debug::TaskEnterState(*(InCtx), ::UE::StateTree::Debug::FNodeReference((InStateTree), FStateTreeIndex16(InNodeIndex)))
+#define UE_STATETREE_DEBUG_TASK_TICK(InCtx, InStateTree, InNodeIndex) ::UE::StateTree::Debug::TaskTick(*(InCtx), ::UE::StateTree::Debug::FNodeReference((InStateTree), FStateTreeIndex16(InNodeIndex)))
+#define UE_STATETREE_DEBUG_TASK_EXIT_STATE(InCtx, InStateTree, InNodeIndex) ::UE::StateTree::Debug::TaskExitState(*(InCtx), ::UE::StateTree::Debug::FNodeReference((InStateTree), FStateTreeIndex16(InNodeIndex)))
 
-#define UE_STATETREE_DEBUG_SEND_EVENT(InExecutionContextPtr, InStateTree, InTag, InPayload, InOrigin) ::UE::StateTree::Debug::EventSent(*(InExecutionContextPtr), InStateTree, InTag, InPayload, InOrigin)
+#define UE_STATETREE_DEBUG_SEND_EVENT(InCtx, InStateTree, InTag, InPayload, InOrigin) ::UE::StateTree::Debug::EventSent(*(InCtx), InStateTree, InTag, InPayload, InOrigin)
 
-#define UE_STATETREE_DEBUG_ENTER_PHASE(InExecutionContextPtr, InPhase) \
-		::UE::StateTree::Debug::EnterPhase(*(InExecutionContextPtr), (InPhase), FStateTreeStateHandle::Invalid); \
-		TRACE_STATETREE_PHASE_EVENT((InExecutionContextPtr)->GetInstanceDebugId(), (InPhase), EStateTreeTraceEventType::Push, FStateTreeStateHandle::Invalid)
+#define UE_STATETREE_DEBUG_ENTER_PHASE(InCtx, InPhase) \
+		::UE::StateTree::Debug::EnterPhase(*(InCtx), (InPhase), FStateTreeStateHandle::Invalid); \
+		TRACE_STATETREE_PHASE_EVENT((InCtx)->GetOwner(), (InCtx)->GetInstanceDebugId(), (InPhase), EStateTreeTraceEventType::Push, FStateTreeStateHandle::Invalid)
 
-#define UE_STATETREE_DEBUG_EXIT_PHASE(InExecutionContextPtr, InPhase) \
-		::UE::StateTree::Debug::ExitPhase(*(InExecutionContextPtr), (InPhase), FStateTreeStateHandle::Invalid); \
-		TRACE_STATETREE_PHASE_EVENT((InExecutionContextPtr)->GetInstanceDebugId(), (InPhase), EStateTreeTraceEventType::Pop, FStateTreeStateHandle::Invalid)
+#define UE_STATETREE_DEBUG_EXIT_PHASE(InCtx, InPhase) \
+		::UE::StateTree::Debug::ExitPhase(*(InCtx), (InPhase), FStateTreeStateHandle::Invalid); \
+		TRACE_STATETREE_PHASE_EVENT((InCtx)->GetOwner(), (InCtx)->GetInstanceDebugId(), (InPhase), EStateTreeTraceEventType::Pop, FStateTreeStateHandle::Invalid)
 
-#define UE_STATETREE_DEBUG_STATE_EVENT(InExecutionContextPtr, InStateHandle, InEventType) \
-		::UE::StateTree::Debug::StateEvent(*(InExecutionContextPtr), (InStateHandle), (InEventType)); \
-		TRACE_STATETREE_STATE_EVENT((InExecutionContextPtr)->GetInstanceDebugId(), (InStateHandle), (InEventType))
+#define UE_STATETREE_DEBUG_STATE_EVENT(InCtx, InStateHandle, InEventType) \
+		::UE::StateTree::Debug::StateEvent(*(InCtx), (InStateHandle), (InEventType)); \
+		TRACE_STATETREE_STATE_EVENT((InCtx)->GetOwner(), (InCtx)->GetInstanceDebugId(), (InStateHandle), (InEventType))
 
-#define UE_STATETREE_DEBUG_TRANSITION_EVENT(InExecutionContextPtr, InTransitionSource, InEventType) \
-		::UE::StateTree::Debug::TransitionEvent(*(InExecutionContextPtr), (InTransitionSource), (InEventType)); \
-		TRACE_STATETREE_TRANSITION_EVENT((InExecutionContextPtr)->GetInstanceDebugId(), (InTransitionSource), (InEventType))
+#define UE_STATETREE_DEBUG_TRANSITION_EVENT(InCtx, InTransitionSource, InEventType) \
+		::UE::StateTree::Debug::TransitionEvent(*(InCtx), (InTransitionSource), (InEventType)); \
+		TRACE_STATETREE_TRANSITION_EVENT((InCtx)->GetOwner(), (InCtx)->GetInstanceDebugId(), (InTransitionSource), (InEventType))
 
-#define UE_STATETREE_DEBUG_LOG_EVENT(InExecutionContextPtr, InLogVerbosity, InFormat, ...) TRACE_STATETREE_LOG_EVENT((InExecutionContextPtr)->GetInstanceDebugId(), InLogVerbosity, InFormat, ##__VA_ARGS__)
-#define UE_STATETREE_DEBUG_CONDITION_EVENT(InExecutionContextPtr, InIndex, InDataView, InEventType) TRACE_STATETREE_CONDITION_EVENT((InExecutionContextPtr)->GetInstanceDebugId(), (InExecutionContextPtr)->StealNodeCustomDebugTraceData(), FStateTreeIndex16(InIndex), (InDataView), (InEventType));
-#define UE_STATETREE_DEBUG_INSTANCE_EVENT(InExecutionContextPtr, InEventType) TRACE_STATETREE_INSTANCE_EVENT((InExecutionContextPtr)->GetOwner(), (InExecutionContextPtr)->GetInstanceDebugId(), (InExecutionContextPtr)->GetStateTree(), *((InExecutionContextPtr)->GetInstanceDescriptionInternal()), (InEventType));
-#define UE_STATETREE_DEBUG_INSTANCE_FRAME_EVENT(InExecutionContextPtr, InFrame) TRACE_STATETREE_INSTANCE_FRAME_EVENT((InExecutionContextPtr)->GetInstanceDebugId(), (InFrame));
-#define UE_STATETREE_DEBUG_ACTIVE_STATES_EVENT(InExecutionContextPtr, InActiveFrames) TRACE_STATETREE_ACTIVE_STATES_EVENT((InExecutionContextPtr)->GetInstanceDebugId(), InActiveFrames);
-#define UE_STATETREE_DEBUG_TASK_EVENT(InExecutionContextPtr, InIndex, InDataView, InEventType, InStatus) TRACE_STATETREE_TASK_EVENT((InExecutionContextPtr)->GetInstanceDebugId(), (InExecutionContextPtr)->StealNodeCustomDebugTraceData(), FStateTreeIndex16(InIndex), (InDataView), (InEventType), (InStatus));
-#define UE_STATETREE_DEBUG_EVALUATOR_EVENT(InExecutionContextPtr, InIndex, InDataView, InEventType) TRACE_STATETREE_EVALUATOR_EVENT((InExecutionContextPtr)->GetInstanceDebugId(), (InExecutionContextPtr)->StealNodeCustomDebugTraceData(), FStateTreeIndex16(InIndex), (InDataView), (InEventType));
+#define UE_STATETREE_DEBUG_LOG_EVENT(InCtx, InLogVerbosity, InFormat, ...) TRACE_STATETREE_LOG_EVENT((InCtx)->GetOwner(), (InCtx)->GetInstanceDebugId(), InLogVerbosity, InFormat, ##__VA_ARGS__)
+#define UE_STATETREE_DEBUG_CONDITION_EVENT(InCtx, InIndex, InDataView, InEventType) TRACE_STATETREE_CONDITION_EVENT((InCtx)->GetOwner(), (InCtx)->GetInstanceDebugId(), (InCtx)->StealNodeCustomDebugTraceData(), FStateTreeIndex16(InIndex), (InDataView), (InEventType));
+#define UE_STATETREE_DEBUG_INSTANCE_EVENT(InCtx, InEventType) TRACE_STATETREE_INSTANCE_EVENT((InCtx)->GetOwner(), (InCtx)->GetInstanceDebugId(), (InCtx)->GetStateTree(), *((InCtx)->GetInstanceDescriptionInternal()), (InEventType));
+#define UE_STATETREE_DEBUG_INSTANCE_FRAME_EVENT(InCtx, InFrame) TRACE_STATETREE_INSTANCE_FRAME_EVENT((InCtx)->GetOwner(), (InCtx)->GetInstanceDebugId(), (InFrame));
+#define UE_STATETREE_DEBUG_ACTIVE_STATES_EVENT(InCtx, InActiveFrames) TRACE_STATETREE_ACTIVE_STATES_EVENT((InCtx)->GetOwner(), (InCtx)->GetInstanceDebugId(), InActiveFrames);
+#define UE_STATETREE_DEBUG_TASK_EVENT(InCtx, InIndex, InDataView, InEventType, InStatus) TRACE_STATETREE_TASK_EVENT((InCtx)->GetOwner(), (InCtx)->GetInstanceDebugId(), (InCtx)->StealNodeCustomDebugTraceData(), FStateTreeIndex16(InIndex), (InDataView), (InEventType), (InStatus));
+#define UE_STATETREE_DEBUG_EVALUATOR_EVENT(InCtx, InIndex, InDataView, InEventType) TRACE_STATETREE_EVALUATOR_EVENT((InCtx)->GetOwner(), (InCtx)->GetInstanceDebugId(), (InCtx)->StealNodeCustomDebugTraceData(), FStateTreeIndex16(InIndex), (InDataView), (InEventType));
 
 #define UE_STATETREE_ID_NAME PREPROCESSOR_JOIN(InstanceId,__LINE__) \
 
 /** Scope based macros captures the instance ID since it might not be accessible when exiting the scope */
-#define UE_STATETREE_DEBUG_SCOPED_PHASE(InExecutionContextPtr, InPhase) \
-		FStateTreeInstanceDebugId UE_STATETREE_ID_NAME = (InExecutionContextPtr)->GetInstanceDebugId(); \
-		UE_STATETREE_DEBUG_ENTER_PHASE((InExecutionContextPtr), (InPhase)); \
+#define UE_STATETREE_DEBUG_SCOPED_PHASE(InCtx, InPhase) \
+		FStateTreeInstanceDebugId UE_STATETREE_ID_NAME = (InCtx)->GetInstanceDebugId(); \
+		UE_STATETREE_DEBUG_ENTER_PHASE((InCtx), (InPhase)); \
 		ON_SCOPE_EXIT \
 		{ \
-			::UE::StateTree::Debug::ExitPhase(*(InExecutionContextPtr), (InPhase), FStateTreeStateHandle::Invalid); \
-			TRACE_STATETREE_PHASE_EVENT(UE_STATETREE_ID_NAME, (InPhase), EStateTreeTraceEventType::Pop, FStateTreeStateHandle::Invalid); \
+			::UE::StateTree::Debug::ExitPhase(*(InCtx), (InPhase), FStateTreeStateHandle::Invalid); \
+			TRACE_STATETREE_PHASE_EVENT((InCtx)->GetOwner(), UE_STATETREE_ID_NAME, (InPhase), EStateTreeTraceEventType::Pop, FStateTreeStateHandle::Invalid); \
 		}
 
 /** Scope based macros captures the instance ID since it might not be accessible when exiting the scope */
-#define UE_STATETREE_DEBUG_SCOPED_STATE(InExecutionContextPtr, InStateHandle) \
-		FStateTreeInstanceDebugId UE_STATETREE_ID_NAME = (InExecutionContextPtr)->GetInstanceDebugId(); \
-		UE_STATETREE_DEBUG_STATE_EVENT((InExecutionContextPtr), (InStateHandle), EStateTreeTraceEventType::Push); \
+#define UE_STATETREE_DEBUG_SCOPED_STATE(InCtx, InStateHandle) \
+		FStateTreeInstanceDebugId UE_STATETREE_ID_NAME = (InCtx)->GetInstanceDebugId(); \
+		UE_STATETREE_DEBUG_STATE_EVENT((InCtx), (InStateHandle), EStateTreeTraceEventType::Push); \
 		ON_SCOPE_EXIT \
 		{ \
-			::UE::StateTree::Debug::StateEvent(*(InExecutionContextPtr), (InStateHandle), EStateTreeTraceEventType::Pop); \
-			TRACE_STATETREE_STATE_EVENT(UE_STATETREE_ID_NAME, (InStateHandle), EStateTreeTraceEventType::Pop); \
+			::UE::StateTree::Debug::StateEvent(*(InCtx), (InStateHandle), EStateTreeTraceEventType::Pop); \
+			TRACE_STATETREE_STATE_EVENT((InCtx)->GetOwner(), UE_STATETREE_ID_NAME, (InStateHandle), EStateTreeTraceEventType::Pop); \
 		}
 
 /** Scope based macros captures the instance ID since it might not be accessible when exiting the scope */
-#define UE_STATETREE_DEBUG_SCOPED_STATE_PHASE(InExecutionContextPtr, InStateHandle, InPhase) \
-		FStateTreeInstanceDebugId UE_STATETREE_ID_NAME = (InExecutionContextPtr)->GetInstanceDebugId(); \
-		::UE::StateTree::Debug::EnterPhase(*(InExecutionContextPtr), (InPhase), (InStateHandle)); \
-		TRACE_STATETREE_PHASE_EVENT(UE_STATETREE_ID_NAME, (InPhase), EStateTreeTraceEventType::Push, (InStateHandle)); \
+#define UE_STATETREE_DEBUG_SCOPED_STATE_PHASE(InCtx, InStateHandle, InPhase) \
+		FStateTreeInstanceDebugId UE_STATETREE_ID_NAME = (InCtx)->GetInstanceDebugId(); \
+		::UE::StateTree::Debug::EnterPhase(*(InCtx), (InPhase), (InStateHandle)); \
+		TRACE_STATETREE_PHASE_EVENT((InCtx)->GetOwner(), UE_STATETREE_ID_NAME, (InPhase), EStateTreeTraceEventType::Push, (InStateHandle)); \
 		ON_SCOPE_EXIT \
 		{ \
-			::UE::StateTree::Debug::ExitPhase(*(InExecutionContextPtr), (InPhase), (InStateHandle)); \
-			TRACE_STATETREE_PHASE_EVENT(UE_STATETREE_ID_NAME, (InPhase), EStateTreeTraceEventType::Pop, (InStateHandle)); \
+			::UE::StateTree::Debug::ExitPhase(*(InCtx), (InPhase), (InStateHandle)); \
+			TRACE_STATETREE_PHASE_EVENT((InCtx)->GetOwner(), UE_STATETREE_ID_NAME, (InPhase), EStateTreeTraceEventType::Pop, (InStateHandle)); \
 		}
 
 namespace UE::StateTree::Debug
diff --git a/Engine/Plugins/Runtime/StateTree/Source/StateTreeModule/Public/Debugger/StateTreeTrace.h b/Engine/Plugins/Runtime/StateTree/Source/StateTreeModule/Public/Debugger/StateTreeTrace.h
--- a/Engine/Plugins/Runtime/StateTree/Source/StateTreeModule/Public/Debugger/StateTreeTrace.h
+++ b/Engine/Plugins/Runtime/StateTree/Source/StateTreeModule/Public/Debugger/StateTreeTrace.h
@@ -97,21 +97,42 @@ private:
 
 void RegisterGlobalDelegates();
 void UnregisterGlobalDelegates();
+void OutputAssetDebugIdEvent(const UStateTree* StateTree, FStateTreeIndex16 AssetDebugId);
+void OutputPhaseScopeEvent(const UObject* Owner, FStateTreeInstanceDebugId InstanceId, EStateTreeUpdatePhase Phase, EStateTreeTraceEventType EventType, FStateTreeStateHandle StateHandle);
+void OutputInstanceLifetimeEvent(const UObject* Owner, FStateTreeInstanceDebugId InstanceId, const UStateTree* StateTree, const TCHAR* InstanceName, EStateTreeTraceEventType EventType);
+void OutputInstanceAssetEvent(const UObject* Owner, FStateTreeInstanceDebugId InstanceId, const UStateTree* StateTree);
+void OutputInstanceFrameEvent(const UObject* Owner, FStateTreeInstanceDebugId InstanceId, const FStateTreeExecutionFrame* Frame);
+void OutputLogEventTrace(const UObject* Owner, FStateTreeInstanceDebugId InstanceId, ELogVerbosity::Type Verbosity, const TCHAR* Fmt, ...);
+void OutputStateEventTrace(const UObject* Owner, FStateTreeInstanceDebugId InstanceId, FStateTreeStateHandle StateHandle, EStateTreeTraceEventType EventType);
+void OutputTaskEventTrace(const UObject* Owner, FStateTreeInstanceDebugId InstanceId, FNodeCustomDebugData&& DebugData, FStateTreeIndex16 TaskIdx, FStateTreeDataView DataView, EStateTreeTraceEventType EventType, EStateTreeRunStatus Status);
+void OutputEvaluatorEventTrace(const UObject* Owner, FStateTreeInstanceDebugId InstanceId, FNodeCustomDebugData&& DebugData, FStateTreeIndex16 EvaluatorIdx, FStateTreeDataView DataView, EStateTreeTraceEventType EventType);
+void OutputConditionEventTrace(const UObject* Owner, FStateTreeInstanceDebugId InstanceId, FNodeCustomDebugData&& DebugData, FStateTreeIndex16 ConditionIdx, FStateTreeDataView DataView, EStateTreeTraceEventType EventType);
+void OutputTransitionEventTrace(const UObject* Owner, FStateTreeInstanceDebugId InstanceId, FStateTreeTransitionSource TransitionSource, EStateTreeTraceEventType EventType);
+void OutputActiveStatesEventTrace(const UObject* Owner, FStateTreeInstanceDebugId InstanceId, TConstArrayView<FStateTreeExecutionFrame> ActiveFrames);
+
 UE_DEPRECATED(5.6, "This method will no longer be exposed publicly.")
 FStateTreeIndex16 FindOrAddDebugIdForAsset(const UStateTree* StateTree);
+UE_DEPRECATED(5.7, "Use the overload taking a pointer to the UObject owner instead.")
 void OutputPhaseScopeEvent(FStateTreeInstanceDebugId InstanceId, EStateTreeUpdatePhase Phase, EStateTreeTraceEventType EventType, FStateTreeStateHandle StateHandle);
-void OutputAssetDebugIdEvent(const UStateTree* StateTree, FStateTreeIndex16 AssetDebugId);
 UE_DEPRECATED(5.7, "Use the overload taking a pointer to the UObject owner instead.")
 void OutputInstanceLifetimeEvent(FStateTreeInstanceDebugId InstanceId, const UStateTree* StateTree, const TCHAR* InstanceName, EStateTreeTraceEventType EventType);
-void OutputInstanceLifetimeEvent(const UObject* Owner, FStateTreeInstanceDebugId InstanceId, const UStateTree* StateTree, const TCHAR* InstanceName, EStateTreeTraceEventType EventType);
+UE_DEPRECATED(5.7, "Use the overload taking a pointer to the UObject owner instead.")
 void OutputInstanceAssetEvent(FStateTreeInstanceDebugId InstanceId, const UStateTree* StateTree);
+UE_DEPRECATED(5.7, "Use the overload taking a pointer to the UObject owner instead.")
 void OutputInstanceFrameEvent(FStateTreeInstanceDebugId InstanceId, const FStateTreeExecutionFrame* Frame);
+UE_DEPRECATED(5.7, "Use the overload taking a pointer to the UObject owner instead.")
 void OutputLogEventTrace(FStateTreeInstanceDebugId InstanceId, ELogVerbosity::Type Verbosity, const TCHAR* Fmt, ...);
+UE_DEPRECATED(5.7, "Use the overload taking a pointer to the UObject owner instead.")
 void OutputStateEventTrace(FStateTreeInstanceDebugId InstanceId, FStateTreeStateHandle StateHandle, EStateTreeTraceEventType EventType);
+UE_DEPRECATED(5.7, "Use the overload taking a pointer to the UObject owner instead.")
 void OutputTaskEventTrace(FStateTreeInstanceDebugId InstanceId, FNodeCustomDebugData&& DebugData, FStateTreeIndex16 TaskIdx, FStateTreeDataView DataView, EStateTreeTraceEventType EventType, EStateTreeRunStatus Status);
+UE_DEPRECATED(5.7, "Use the overload taking a pointer to the UObject owner instead.")
 void OutputEvaluatorEventTrace(FStateTreeInstanceDebugId InstanceId, FNodeCustomDebugData&& DebugData, FStateTreeIndex16 EvaluatorIdx, FStateTreeDataView DataView, EStateTreeTraceEventType EventType);
+UE_DEPRECATED(5.7, "Use the overload taking a pointer to the UObject owner instead.")
 void OutputConditionEventTrace(FStateTreeInstanceDebugId InstanceId, FNodeCustomDebugData&& DebugData, FStateTreeIndex16 ConditionIdx, FStateTreeDataView DataView, EStateTreeTraceEventType EventType);
+UE_DEPRECATED(5.7, "Use the overload taking a pointer to the UObject owner instead.")
 void OutputTransitionEventTrace(FStateTreeInstanceDebugId InstanceId, FStateTreeTransitionSource TransitionSource, EStateTreeTraceEventType EventType);
+UE_DEPRECATED(5.7, "Use the overload taking a pointer to the UObject owner instead.")
 void OutputActiveStatesEventTrace(FStateTreeInstanceDebugId InstanceId, TConstArrayView<FStateTreeExecutionFrame> ActiveFrames);
 
 } // UE::StateTreeTrace
@@ -119,50 +140,50 @@ void OutputActiveStatesEventTrace(FStateTreeInstanceDebugId InstanceId, TConstAr
 #define TRACE_STATETREE_INSTANCE_EVENT(Owner, InstanceID, StateTree, InstanceName, EventType) \
 	UE::StateTreeTrace::OutputInstanceLifetimeEvent(Owner, InstanceID, StateTree, InstanceName, EventType);
 
-#define TRACE_STATETREE_INSTANCE_FRAME_EVENT(InstanceID, Frame) \
-	UE::StateTreeTrace::OutputInstanceFrameEvent(InstanceID, Frame);
+#define TRACE_STATETREE_INSTANCE_FRAME_EVENT(Owner, InstanceID, Frame) \
+	UE::StateTreeTrace::OutputInstanceFrameEvent(Owner, InstanceID, Frame);
 
-#define TRACE_STATETREE_PHASE_EVENT(InstanceID, Phase, EventType, StateHandle) \
-	UE::StateTreeTrace::OutputPhaseScopeEvent(InstanceID, Phase, EventType, StateHandle); \
+#define TRACE_STATETREE_PHASE_EVENT(Owner, InstanceID, Phase, EventType, StateHandle) \
+	UE::StateTreeTrace::OutputPhaseScopeEvent(Owner, InstanceID, Phase, EventType, StateHandle); \
 
-#define TRACE_STATETREE_LOG_EVENT(InstanceId, TraceVerbosity, Format, ...) \
+#define TRACE_STATETREE_LOG_EVENT(Owner, InstanceId, TraceVerbosity, Format, ...) \
 	if (UE_TRACE_CHANNELEXPR_IS_ENABLED(StateTreeDebugChannel)) \
 	{ \
-		UE::StateTreeTrace::OutputLogEventTrace(InstanceId, ELogVerbosity::TraceVerbosity, Format, ##__VA_ARGS__); \
+		UE::StateTreeTrace::OutputLogEventTrace(Owner, InstanceId, ELogVerbosity::TraceVerbosity, Format, ##__VA_ARGS__); \
 	}
 
-#define TRACE_STATETREE_STATE_EVENT(InstanceId, StateHandle, EventType) \
+#define TRACE_STATETREE_STATE_EVENT(Owner, InstanceId, StateHandle, EventType) \
 	if (UE_TRACE_CHANNELEXPR_IS_ENABLED(StateTreeDebugChannel)) \
 	{ \
-		UE::StateTreeTrace::OutputStateEventTrace(InstanceId, StateHandle, EventType); \
+		UE::StateTreeTrace::OutputStateEventTrace(Owner, InstanceId, StateHandle, EventType); \
 	}
 
-#define TRACE_STATETREE_TASK_EVENT(InstanceId, DebugData, TaskIdx, DataView, EventType, Status) \
+#define TRACE_STATETREE_TASK_EVENT(Owner, InstanceId, DebugData, TaskIdx, DataView, EventType, Status) \
 	if (UE_TRACE_CHANNELEXPR_IS_ENABLED(StateTreeDebugChannel)) \
 	{ \
-		UE::StateTreeTrace::OutputTaskEventTrace(InstanceId, DebugData, TaskIdx, DataView, EventType, Status); \
+		UE::StateTreeTrace::OutputTaskEventTrace(Owner, InstanceId, DebugData, TaskIdx, DataView, EventType, Status); \
 	}
 
-#define TRACE_STATETREE_EVALUATOR_EVENT(InstanceId, DebugData, EvaluatorIdx, DataView, EventType) \
+#define TRACE_STATETREE_EVALUATOR_EVENT(Owner, InstanceId, DebugData, EvaluatorIdx, DataView, EventType) \
 	if (UE_TRACE_CHANNELEXPR_IS_ENABLED(StateTreeDebugChannel)) \
 	{ \
-		UE::StateTreeTrace::OutputEvaluatorEventTrace(InstanceId, DebugData, EvaluatorIdx, DataView, EventType); \
+		UE::StateTreeTrace::OutputEvaluatorEventTrace(Owner, InstanceId, DebugData, EvaluatorIdx, DataView, EventType); \
 	}
 
-#define TRACE_STATETREE_CONDITION_EVENT(InstanceId, DebugData, ConditionIdx, DataView, EventType) \
+#define TRACE_STATETREE_CONDITION_EVENT(Owner, InstanceId, DebugData, ConditionIdx, DataView, EventType) \
 	if (UE_TRACE_CHANNELEXPR_IS_ENABLED(StateTreeDebugChannel)) \
 	{ \
-		UE::StateTreeTrace::OutputConditionEventTrace(InstanceId, DebugData, ConditionIdx, DataView, EventType); \
+		UE::StateTreeTrace::OutputConditionEventTrace(Owner, InstanceId, DebugData, ConditionIdx, DataView, EventType); \
 	}
 
-#define TRACE_STATETREE_TRANSITION_EVENT(InstanceId, TransitionIdx, EventType) \
+#define TRACE_STATETREE_TRANSITION_EVENT(Owner, InstanceId, TransitionIdx, EventType) \
 	if (UE_TRACE_CHANNELEXPR_IS_ENABLED(StateTreeDebugChannel)) \
 	{ \
-		UE::StateTreeTrace::OutputTransitionEventTrace(InstanceId, TransitionIdx, EventType); \
+		UE::StateTreeTrace::OutputTransitionEventTrace(Owner, InstanceId, TransitionIdx, EventType); \
 	}
 
-#define TRACE_STATETREE_ACTIVE_STATES_EVENT(InstanceId, ActivateFrames) \
-		UE::StateTreeTrace::OutputActiveStatesEventTrace(InstanceId, ActivateFrames);
+#define TRACE_STATETREE_ACTIVE_STATES_EVENT(Owner, InstanceId, ActivateFrames) \
+		UE::StateTreeTrace::OutputActiveStatesEventTrace(Owner, InstanceId, ActivateFrames);
 
 #else //WITH_STATETREE_TRACE
 
