Updated headers using UnrealCodeFixup to make sure dllstorage is on methods/staticvars instead of types. Part 4/n

diff --git a/Engine/Plugins/Runtime/StateTree/Source/StateTreeEditorModule/Public/StateTreeEditorData.h b/Engine/Plugins/Runtime/StateTree/Source/StateTreeEditorModule/Public/StateTreeEditorData.h
--- a/Engine/Plugins/Runtime/StateTree/Source/StateTreeEditorModule/Public/StateTreeEditorData.h
+++ b/Engine/Plugins/Runtime/StateTree/Source/StateTreeEditorModule/Public/StateTreeEditorData.h
@@ -8,6 +8,8 @@
 #include "Debugger/StateTreeDebuggerTypes.h"
 #include "StateTreeEditorData.generated.h"
 
+#define UE_API STATETREEEDITORMODULE_API
+
 struct FStateTreeBindableStructDesc;
 
 class UStateTreeSchema;
@@ -53,21 +55,21 @@ enum class EStateTreeVisitor : uint8
 /**
  * Edit time data for StateTree asset. This data gets baked into runtime format before being used by the StateTreeInstance.
  */
-UCLASS(BlueprintType, EditInlineNew, CollapseCategories, Within = "StateTree")
-class STATETREEEDITORMODULE_API UStateTreeEditorData : public UObject, public IStateTreeEditorPropertyBindingsOwner
+UCLASS(MinimalAPI, BlueprintType, EditInlineNew, CollapseCategories, Within = "StateTree")
+class UStateTreeEditorData : public UObject, public IStateTreeEditorPropertyBindingsOwner
 {
 	GENERATED_BODY()
 	
 public:
-	UStateTreeEditorData();
+	UE_API UStateTreeEditorData();
 
-	virtual void PostInitProperties() override;
-	virtual void Serialize(FArchive& Ar) override;
+	UE_API virtual void PostInitProperties() override;
+	UE_API virtual void Serialize(FArchive& Ar) override;
 	
 	//~ Begin IStateTreeEditorPropertyBindingsOwner interface
-	virtual void GetBindableStructs(const FGuid TargetStructID, TArray<TInstancedStruct<FPropertyBindingBindableStructDescriptor>>& OutStructDescs) const override;
-	virtual bool GetBindableStructByID(const FGuid StructID, TInstancedStruct<FPropertyBindingBindableStructDescriptor>& OutStructDesc) const override;
-	virtual bool GetBindingDataViewByID(const FGuid StructID, FPropertyBindingDataView& OutDataView) const override;
+	UE_API virtual void GetBindableStructs(const FGuid TargetStructID, TArray<TInstancedStruct<FPropertyBindingBindableStructDescriptor>>& OutStructDescs) const override;
+	UE_API virtual bool GetBindableStructByID(const FGuid StructID, TInstancedStruct<FPropertyBindingBindableStructDescriptor>& OutStructDesc) const override;
+	UE_API virtual bool GetBindingDataViewByID(const FGuid StructID, FPropertyBindingDataView& OutDataView) const override;
 	virtual const FStateTreeEditorPropertyBindings* GetPropertyEditorBindings() const override
 	{
 		return &EditorBindings;
@@ -88,12 +90,12 @@ public:
 		return &EditorBindings;
 	}
 
-	virtual FStateTreeBindableStructDesc FindContextData(const UStruct* ObjectType, const FString ObjectNameHint) const override;
+	UE_API virtual FStateTreeBindableStructDesc FindContextData(const UStruct* ObjectType, const FString ObjectNameHint) const override;
 
-	virtual bool CanCreateParameter(const FGuid InStructID) const override;
-	virtual void CreateParametersForStruct(const FGuid InStructID, TArrayView<UE::PropertyBinding::FPropertyCreationDescriptor> InOutCreationDescs) override;
-	virtual void OnPropertyBindingChanged(const FPropertyBindingPath& InSourcePath, const FPropertyBindingPath& InTargetPath) override;
-	virtual void AppendBindablePropertyFunctionStructs(TArray<TInstancedStruct<FPropertyBindingBindableStructDescriptor>>& InOutStructs) const override;
+	UE_API virtual bool CanCreateParameter(const FGuid InStructID) const override;
+	UE_API virtual void CreateParametersForStruct(const FGuid InStructID, TArrayView<UE::PropertyBinding::FPropertyCreationDescriptor> InOutCreationDescs) override;
+	UE_API virtual void OnPropertyBindingChanged(const FPropertyBindingPath& InSourcePath, const FPropertyBindingPath& InTargetPath) override;
+	UE_API virtual void AppendBindablePropertyFunctionStructs(TArray<TInstancedStruct<FPropertyBindingBindableStructDescriptor>>& InOutStructs) const override;
 	//~ End IStateTreeEditorPropertyBindingsOwner interface
 
 	/**
@@ -101,18 +103,18 @@ public:
 	 * Handles the name override logic, figures out required data for the GetDescription() call, and handles the fallbacks.
 	 * @return description for the node.
 	 */
-	FText GetNodeDescription(const FStateTreeEditorNode& Node, const EStateTreeNodeFormatting Formatting = EStateTreeNodeFormatting::Text) const;
+	UE_API FText GetNodeDescription(const FStateTreeEditorNode& Node, const EStateTreeNodeFormatting Formatting = EStateTreeNodeFormatting::Text) const;
 
 #if WITH_EDITOR
 	using FReplacementObjectMap = TMap<UObject*, UObject*>;
-	void OnObjectsReinstanced(const FReplacementObjectMap& ObjectMap);
-	void OnUserDefinedStructReinstanced(const UUserDefinedStruct& UserDefinedStruct);
-	void OnParametersChanged(const UStateTree& StateTree);
-	void OnStateParametersChanged(const UStateTree& StateTree, const FGuid StateID);
-	virtual void BeginDestroy() override;
-	virtual void PostLoad() override;
-	virtual void PostEditChangeChainProperty(FPropertyChangedChainEvent& PropertyChangedEvent) override;
-	virtual void PostDuplicate(EDuplicateMode::Type DuplicateMode) override;
+	UE_API void OnObjectsReinstanced(const FReplacementObjectMap& ObjectMap);
+	UE_API void OnUserDefinedStructReinstanced(const UUserDefinedStruct& UserDefinedStruct);
+	UE_API void OnParametersChanged(const UStateTree& StateTree);
+	UE_API void OnStateParametersChanged(const UStateTree& StateTree, const FGuid StateID);
+	UE_API virtual void BeginDestroy() override;
+	UE_API virtual void PostLoad() override;
+	UE_API virtual void PostEditChangeChainProperty(FPropertyChangedChainEvent& PropertyChangedEvent) override;
+	UE_API virtual void PostDuplicate(EDuplicateMode::Type DuplicateMode) override;
 #endif
 
 	/** @return the public parameters ID that could be used for bindings within the Tree. */
@@ -128,55 +130,55 @@ public:
 	}
 
 	/** @returns parent state of a struct, or nullptr if not found. */
-	const UStateTreeState* GetStateByStructID(const FGuid TargetStructID) const;
+	UE_API const UStateTreeState* GetStateByStructID(const FGuid TargetStructID) const;
 
 	/** @returns state based on its ID, or nullptr if not found. */
-	const UStateTreeState* GetStateByID(const FGuid StateID) const;
+	UE_API const UStateTreeState* GetStateByID(const FGuid StateID) const;
 
 	/** @returns mutable state based on its ID, or nullptr if not found. */
-	UStateTreeState* GetMutableStateByID(const FGuid StateID);
+	UE_API UStateTreeState* GetMutableStateByID(const FGuid StateID);
 
 	/** @returns the IDs and instance values of all bindable structs in the StateTree. */
-	void GetAllStructValues(TMap<FGuid, const FStateTreeDataView>& OutAllValues) const;
+	UE_API void GetAllStructValues(TMap<FGuid, const FStateTreeDataView>& OutAllValues) const;
 
 	/** @returns the IDs and instance values of all bindable structs in the StateTree. */
-	void GetAllStructValues(TMap<FGuid, const FPropertyBindingDataView>& OutAllValues) const;
+	UE_API void GetAllStructValues(TMap<FGuid, const FPropertyBindingDataView>& OutAllValues) const;
 
 	/**
 	* Iterates over all structs that are related to binding
 	* @param InFunc function called at each node, should return Continue if visiting is continued or Break to stop.
 	*/
-	EStateTreeVisitor VisitHierarchy(TFunctionRef<EStateTreeVisitor(UStateTreeState& State, UStateTreeState* ParentState)> InFunc) const;
+	UE_API EStateTreeVisitor VisitHierarchy(TFunctionRef<EStateTreeVisitor(UStateTreeState& State, UStateTreeState* ParentState)> InFunc) const;
 
 	/**
 	 * Iterates over all structs at the global level (context, tree parameters, evaluators, global tasks) that are related to binding.
 	 * @param InFunc function called at each node, should return Continue if visiting is continued or Break to stop.
 	 */
-	EStateTreeVisitor VisitGlobalNodes(TFunctionRef<EStateTreeVisitor(const UStateTreeState* State, const FStateTreeBindableStructDesc& Desc, const FStateTreeDataView Value)> InFunc) const;
+	UE_API EStateTreeVisitor VisitGlobalNodes(TFunctionRef<EStateTreeVisitor(const UStateTreeState* State, const FStateTreeBindableStructDesc& Desc, const FStateTreeDataView Value)> InFunc) const;
 
 	/**
 	 * Iterates over all structs in the state hierarchy that are related to binding.
 	 * @param InFunc function called at each node, should return Continue if visiting is continued or Break to stop.
 	 */
-	EStateTreeVisitor VisitHierarchyNodes(TFunctionRef<EStateTreeVisitor(const UStateTreeState* State, const FStateTreeBindableStructDesc& Desc, const FStateTreeDataView Value)> InFunc) const;
+	UE_API EStateTreeVisitor VisitHierarchyNodes(TFunctionRef<EStateTreeVisitor(const UStateTreeState* State, const FStateTreeBindableStructDesc& Desc, const FStateTreeDataView Value)> InFunc) const;
 
 	/**
 	 * Iterates over all structs that are related to binding.
 	 * @param InFunc function called at each node, should return Continue if visiting is continued or Break to stop.
 	 */
-	EStateTreeVisitor VisitAllNodes(TFunctionRef<EStateTreeVisitor(const UStateTreeState* State, const FStateTreeBindableStructDesc& Desc, const FStateTreeDataView Value)> InFunc) const;
+	UE_API EStateTreeVisitor VisitAllNodes(TFunctionRef<EStateTreeVisitor(const UStateTreeState* State, const FStateTreeBindableStructDesc& Desc, const FStateTreeDataView Value)> InFunc) const;
 
 	/**
 	 * Iterates over all nodes in a given state.
 	 * @param InFunc function called at each node, should return Continue if visiting is continued or Break to stop.
 	 */
-	EStateTreeVisitor VisitStateNodes(const UStateTreeState& State, TFunctionRef<EStateTreeVisitor(const UStateTreeState* State, const FStateTreeBindableStructDesc& Desc, const FStateTreeDataView Value)> InFunc) const;
+	UE_API EStateTreeVisitor VisitStateNodes(const UStateTreeState& State, TFunctionRef<EStateTreeVisitor(const UStateTreeState* State, const FStateTreeBindableStructDesc& Desc, const FStateTreeDataView Value)> InFunc) const;
 
 	/**
 	 * Iterates recursively over all property functions of the provided node. Also nested ones.
 	 * @param InFunc function called at each node, should return Continue if visiting is continued or Break to stop.
 	 */
-	EStateTreeVisitor VisitStructBoundPropertyFunctions(FGuid StructID, const FString& StatePath, TFunctionRef<EStateTreeVisitor(const FStateTreeEditorNode& EditorNode, const FStateTreeBindableStructDesc& Desc, const FStateTreeDataView Value)> InFunc) const;
+	UE_API EStateTreeVisitor VisitStructBoundPropertyFunctions(FGuid StructID, const FString& StatePath, TFunctionRef<EStateTreeVisitor(const FStateTreeEditorNode& EditorNode, const FStateTreeBindableStructDesc& Desc, const FStateTreeDataView Value)> InFunc) const;
 
 	/**
 	 * Returns array of nodes along the execution path, up to the TargetStruct.
@@ -184,7 +186,7 @@ public:
 	 * @param TargetStructID The ID of the node where to stop.
 	 * @param OutStructDescs Array of nodes accessible on the given path.
 	 */
-	void GetAccessibleStructsInExecutionPath(const TConstArrayView<const UStateTreeState*> Path, const FGuid TargetStructID, TArray<TInstancedStruct<FPropertyBindingBindableStructDescriptor>>& OutStructDescs) const;
+	UE_API void GetAccessibleStructsInExecutionPath(const TConstArrayView<const UStateTreeState*> Path, const FGuid TargetStructID, TArray<TInstancedStruct<FPropertyBindingBindableStructDescriptor>>& OutStructDescs) const;
 
 	UE_DEPRECATED(5.6, "Use GetAccessibleStructsInExecutionPath instead")
 	void GetAccessibleStruct(const TConstArrayView<const UStateTreeState*> Path, const FGuid TargetStructID, TArray<TInstancedStruct<FPropertyBindingBindableStructDescriptor>>& OutStructDescs) const
@@ -192,7 +194,7 @@ public:
 		GetAccessibleStructsInExecutionPath(Path, TargetStructID, OutStructDescs);
 	}
 
-	void ReparentStates();
+	UE_API void ReparentStates();
 	
 	// StateTree Builder API
 
@@ -302,11 +304,11 @@ PRAGMA_ENABLE_DEPRECATION_WARNINGS
 	}
 
 #if WITH_STATETREE_TRACE_DEBUGGER
-	bool HasAnyBreakpoint(FGuid ID) const;
-	bool HasBreakpoint(FGuid ID, EStateTreeBreakpointType BreakpointType) const;
-	const FStateTreeEditorBreakpoint* GetBreakpoint(FGuid ID, EStateTreeBreakpointType BreakpointType) const;
-	void AddBreakpoint(FGuid ID, EStateTreeBreakpointType BreakpointType);
-	bool RemoveBreakpoint(FGuid ID, EStateTreeBreakpointType BreakpointType);
+	UE_API bool HasAnyBreakpoint(FGuid ID) const;
+	UE_API bool HasBreakpoint(FGuid ID, EStateTreeBreakpointType BreakpointType) const;
+	UE_API const FStateTreeEditorBreakpoint* GetBreakpoint(FGuid ID, EStateTreeBreakpointType BreakpointType) const;
+	UE_API void AddBreakpoint(FGuid ID, EStateTreeBreakpointType BreakpointType);
+	UE_API bool RemoveBreakpoint(FGuid ID, EStateTreeBreakpointType BreakpointType);
 #endif // WITH_STATETREE_TRACE_DEBUGGER
 
 	// ~StateTree Builder API
@@ -322,14 +324,14 @@ PRAGMA_ENABLE_DEPRECATION_WARNINGS
 	virtual void CreateRootProperties(TArrayView<UE::PropertyBinding::FPropertyCreationDescriptor> InOutCreationDescs) { UE::PropertyBinding::CreateUniquelyNamedPropertiesInPropertyBag(InOutCreationDescs, RootParameterPropertyBag); }
 
 private:
-	virtual EStateTreeVisitor EnumerateBindablePropertyFunctionNodes(TFunctionRef<EStateTreeVisitor(const UScriptStruct* NodeStruct, const FStateTreeBindableStructDesc& Desc, const FStateTreeDataView Value)> InFunc) const override;
+	UE_API virtual EStateTreeVisitor EnumerateBindablePropertyFunctionNodes(TFunctionRef<EStateTreeVisitor(const UScriptStruct* NodeStruct, const FStateTreeBindableStructDesc& Desc, const FStateTreeDataView Value)> InFunc) const override;
 
-	void FixObjectInstance(TSet<UObject*>& SeenObjects, UObject& Outer, FStateTreeEditorNode& Node);
-	void FixObjectNodes();
-	void FixDuplicateIDs();
-	void DuplicateIDs();
-	void UpdateBindingsInstanceStructs();
-	void CallPostLoadOnNodes();
+	UE_API void FixObjectInstance(TSet<UObject*>& SeenObjects, UObject& Outer, FStateTreeEditorNode& Node);
+	UE_API void FixObjectNodes();
+	UE_API void FixDuplicateIDs();
+	UE_API void DuplicateIDs();
+	UE_API void UpdateBindingsInstanceStructs();
+	UE_API void CallPostLoadOnNodes();
 
 #if WITH_EDITORONLY_DATA
 	FDelegateHandle OnObjectsReinstancedHandle;
@@ -402,4 +404,6 @@ UCLASS()
 class UQAStateTreeEditorData : public UStateTreeEditorData
 {
 	GENERATED_BODY()
-};
\ No newline at end of file
+};
+
+#undef UE_API
diff --git a/Engine/Plugins/Runtime/StateTree/Source/StateTreeEditorModule/Public/StateTreeEditorMode.h b/Engine/Plugins/Runtime/StateTree/Source/StateTreeEditorModule/Public/StateTreeEditorMode.h
--- a/Engine/Plugins/Runtime/StateTree/Source/StateTreeEditorModule/Public/StateTreeEditorMode.h
+++ b/Engine/Plugins/Runtime/StateTree/Source/StateTreeEditorModule/Public/StateTreeEditorMode.h
@@ -8,50 +8,52 @@
 
 #include "StateTreeEditorMode.generated.h"
 
+#define UE_API STATETREEEDITORMODULE_API
+
 class FStateTreeBindingExtension;
 class FStateTreeBindingsChildrenCustomization;
 class IDetailsView;
 class IMessageLogListing;
 
-UCLASS(Transient)
-class STATETREEEDITORMODULE_API UStateTreeEditorMode : public UEdMode
+UCLASS(MinimalAPI, Transient)
+class UStateTreeEditorMode : public UEdMode
 {
 	GENERATED_BODY()
 public:
-	const static FEditorModeID EM_StateTree;
+	UE_API const static FEditorModeID EM_StateTree;
 	
-	UStateTreeEditorMode();
+	UE_API UStateTreeEditorMode();
 
-	virtual void Enter() override;
-	virtual void Exit() override;
-	virtual void CreateToolkit() override;
-	virtual void BindCommands() override;
+	UE_API virtual void Enter() override;
+	UE_API virtual void Exit() override;
+	UE_API virtual void CreateToolkit() override;
+	UE_API virtual void BindCommands() override;
 
 protected:
-	void OnStateTreeChanged();
-	void BindToolkitCommands(const TSharedRef<FUICommandList>& ToolkitCommands);
+	UE_API void OnStateTreeChanged();
+	UE_API void BindToolkitCommands(const TSharedRef<FUICommandList>& ToolkitCommands);
 
-	void OnPropertyBindingChanged(const FPropertyBindingPath& SourcePath, const FPropertyBindingPath& TargetPath);
-	void OnIdentifierChanged(const UStateTree& InStateTree);
-	void OnSchemaChanged(const UStateTree& InStateTree);
-	void ForceRefreshDetailsView() const;
-	void OnRefreshDetailsView(const UStateTree& InStateTree) const;
-	void OnStateParametersChanged(const UStateTree& InStateTree, const FGuid ChangedStateID) const;
+	UE_API void OnPropertyBindingChanged(const FPropertyBindingPath& SourcePath, const FPropertyBindingPath& TargetPath);
+	UE_API void OnIdentifierChanged(const UStateTree& InStateTree);
+	UE_API void OnSchemaChanged(const UStateTree& InStateTree);
+	UE_API void ForceRefreshDetailsView() const;
+	UE_API void OnRefreshDetailsView(const UStateTree& InStateTree) const;
+	UE_API void OnStateParametersChanged(const UStateTree& InStateTree, const FGuid ChangedStateID) const;
 
-	void HandleMessageTokenClicked(const TSharedRef<IMessageToken>& InMessageToken) const;
+	UE_API void HandleMessageTokenClicked(const TSharedRef<IMessageToken>& InMessageToken) const;
 
-	void OnAssetFinishedChangingProperties(const FPropertyChangedEvent& PropertyChangedEvent) const;
-	void OnSelectionFinishedChangingProperties(const FPropertyChangedEvent& PropertyChangedEvent);
+	UE_API void OnAssetFinishedChangingProperties(const FPropertyChangedEvent& PropertyChangedEvent) const;
+	UE_API void OnSelectionFinishedChangingProperties(const FPropertyChangedEvent& PropertyChangedEvent);
 
-	void Compile();
-	bool CanCompile() const;
-	bool IsCompileVisible() const;
+	UE_API void Compile();
+	UE_API bool CanCompile() const;
+	UE_API bool IsCompileVisible() const;
 
-	bool HasValidStateTree() const;
+	UE_API bool HasValidStateTree() const;
 	
-	void HandleModelAssetChanged();
-	void HandleModelSelectionChanged(const TArray<TWeakObjectPtr<UStateTreeState>>& SelectedStates) const;
-	void HandleModelBringNodeToFocus(const UStateTreeState* State, const FGuid NodeID) const;
+	UE_API void HandleModelAssetChanged();
+	UE_API void HandleModelSelectionChanged(const TArray<TWeakObjectPtr<UStateTreeState>>& SelectedStates) const;
+	UE_API void HandleModelBringNodeToFocus(const UStateTreeState* State, const FGuid NodeID) const;
 
 	void HandleStateAdded(UStateTreeState* , UStateTreeState*)
 	{
@@ -74,15 +76,15 @@ protected:
 	}
 
 	/** Resolve the internal editor data and fixup the StateTree nodes. */
-	void UpdateAsset();
+	UE_API void UpdateAsset();
 
-	TSharedPtr<IDetailsView> GetDetailsView() const;
-	TSharedPtr<IDetailsView> GetAssetDetailsView() const;
+	UE_API TSharedPtr<IDetailsView> GetDetailsView() const;
+	UE_API TSharedPtr<IDetailsView> GetAssetDetailsView() const;
 	
-	TSharedPtr<IMessageLogListing> GetMessageLogListing() const;	
-	void ShowCompilerTab() const;
+	UE_API TSharedPtr<IMessageLogListing> GetMessageLogListing() const;	
+	UE_API void ShowCompilerTab() const;
 
-	UStateTree* GetStateTree() const;
+	UE_API UStateTree* GetStateTree() const;
 	
 	friend class FStateTreeEditorModeToolkit;
 
@@ -100,3 +102,4 @@ protected:
 	TSharedPtr<FStateTreeBindingsChildrenCustomization> DetailsViewChildrenCustomizationHandler;
 };
 
+#undef UE_API
diff --git a/Engine/Plugins/Runtime/StateTree/Source/StateTreeEditorModule/Public/StateTreeEditorModeToolkit.h b/Engine/Plugins/Runtime/StateTree/Source/StateTreeEditorModule/Public/StateTreeEditorModeToolkit.h
--- a/Engine/Plugins/Runtime/StateTree/Source/StateTreeEditorModule/Public/StateTreeEditorModeToolkit.h
+++ b/Engine/Plugins/Runtime/StateTree/Source/StateTreeEditorModule/Public/StateTreeEditorModeToolkit.h
@@ -6,6 +6,8 @@
 #include "IStateTreeEditorHost.h"
 #include "Toolkits/BaseToolkit.h"
 
+#define UE_API STATETREEEDITORMODULE_API
+
 class UStateTreeEditorMode;
 struct FPropertyBindingBindingCollection;
 namespace UE::StateTreeEditor
@@ -13,45 +15,45 @@ namespace UE::StateTreeEditor
 	struct FSpawnedWorkspaceTab;
 }
 
-class STATETREEEDITORMODULE_API FStateTreeEditorModeToolkit : public FModeToolkit
+class FStateTreeEditorModeToolkit : public FModeToolkit
 {
 public:
 
-	FStateTreeEditorModeToolkit(UStateTreeEditorMode* InEditorMode);
+	UE_API FStateTreeEditorModeToolkit(UStateTreeEditorMode* InEditorMode);
 	
 	/** IToolkit interface */
-	virtual FName GetToolkitFName() const override;
-	virtual FText GetBaseToolkitName() const override;
+	UE_API virtual FName GetToolkitFName() const override;
+	UE_API virtual FText GetBaseToolkitName() const override;
 
-	virtual void Init(const TSharedPtr<IToolkitHost>& InitToolkitHost, TWeakObjectPtr<UEdMode> InOwningMode) override;
-	virtual void InvokeUI() override;
-	virtual void RequestModeUITabs() override;
+	UE_API virtual void Init(const TSharedPtr<IToolkitHost>& InitToolkitHost, TWeakObjectPtr<UEdMode> InOwningMode) override;
+	UE_API virtual void InvokeUI() override;
+	UE_API virtual void RequestModeUITabs() override;
 	
-	virtual void ExtendSecondaryModeToolbar(UToolMenu* InModeToolbarMenu) override;
+	UE_API virtual void ExtendSecondaryModeToolbar(UToolMenu* InModeToolbarMenu) override;
 		
-	void OnStateTreeChanged();
+	UE_API void OnStateTreeChanged();
 
 protected:
-	FSlateIcon GetCompileStatusImage() const;
+	UE_API FSlateIcon GetCompileStatusImage() const;
 
-	static FSlateIcon GetNewTaskButtonImage();
-	TSharedRef<SWidget> GenerateTaskBPBaseClassesMenu() const;
+	static UE_API FSlateIcon GetNewTaskButtonImage();
+	UE_API TSharedRef<SWidget> GenerateTaskBPBaseClassesMenu() const;
 
-	static FSlateIcon GetNewConditionButtonImage();
-	TSharedRef<SWidget> GenerateConditionBPBaseClassesMenu() const;
+	static UE_API FSlateIcon GetNewConditionButtonImage();
+	UE_API TSharedRef<SWidget> GenerateConditionBPBaseClassesMenu() const;
     
-	static FSlateIcon GetNewConsiderationButtonImage();
-	TSharedRef<SWidget> GenerateConsiderationBPBaseClassesMenu() const;
+	static UE_API FSlateIcon GetNewConsiderationButtonImage();
+	UE_API TSharedRef<SWidget> GenerateConsiderationBPBaseClassesMenu() const;
 
-	void OnNodeBPBaseClassPicked(UClass* NodeClass) const;
+	UE_API void OnNodeBPBaseClassPicked(UClass* NodeClass) const;
 	
-	FText GetStatisticsText() const;
-	const FPropertyBindingBindingCollection* GetBindingCollection() const;
+	UE_API FText GetStatisticsText() const;
+	UE_API const FPropertyBindingBindingCollection* GetBindingCollection() const;
 
-	void UpdateStateTreeOutliner();
+	UE_API void UpdateStateTreeOutliner();
 
-	void HandleTabSpawned(UE::StateTreeEditor::FSpawnedWorkspaceTab SpawnedTab);
-	void HandleTabClosed(UE::StateTreeEditor::FSpawnedWorkspaceTab SpawnedTab);
+	UE_API void HandleTabSpawned(UE::StateTreeEditor::FSpawnedWorkspaceTab SpawnedTab);
+	UE_API void HandleTabClosed(UE::StateTreeEditor::FSpawnedWorkspaceTab SpawnedTab);
 
 protected:
 	TWeakObjectPtr<UStateTreeEditorMode> WeakEditorMode;
@@ -62,10 +64,11 @@ protected:
 	TWeakPtr<SDockTab> WeakOutlinerTab = nullptr;
 
 #if WITH_STATETREE_TRACE_DEBUGGER
-	void UpdateDebuggerView();
+	UE_API void UpdateDebuggerView();
 	TWeakPtr<SDockTab> WeakDebuggerTab = nullptr;
 #endif // WITH_STATETREE_TRACE_DEBUGGER
 
 };
 
 
+#undef UE_API
diff --git a/Engine/Plugins/Runtime/StateTree/Source/StateTreeEditorModule/Public/StateTreeEditorModule.h b/Engine/Plugins/Runtime/StateTree/Source/StateTreeEditorModule/Public/StateTreeEditorModule.h
--- a/Engine/Plugins/Runtime/StateTree/Source/StateTreeEditorModule/Public/StateTreeEditorModule.h
+++ b/Engine/Plugins/Runtime/StateTree/Source/StateTreeEditorModule/Public/StateTreeEditorModule.h
@@ -8,6 +8,8 @@
 #include "Templates/PimplPtr.h"
 #include "Toolkits/AssetEditorToolkit.h"
 
+#define UE_API STATETREEEDITORMODULE_API
+
 class UStateTree;
 class UStateTreeEditorData;
 class UStateTreeSchema;
@@ -27,40 +29,40 @@ STATETREEEDITORMODULE_API DECLARE_LOG_CATEGORY_EXTERN(LogStateTreeEditor, Log, A
 /**
 * The public interface to this module
 */
-class STATETREEEDITORMODULE_API FStateTreeEditorModule : public IModuleInterface, public IHasMenuExtensibility, public IHasToolBarExtensibility
+class FStateTreeEditorModule : public IModuleInterface, public IHasMenuExtensibility, public IHasToolBarExtensibility
 {
 public:
 	//~Begin IModuleInterface
-	virtual void StartupModule() override;
-	virtual void ShutdownModule() override;
+	UE_API virtual void StartupModule() override;
+	UE_API virtual void ShutdownModule() override;
 	//~End IModuleInterface
 
 	/** Gets this module, will attempt to load and should always exist. */
-	static FStateTreeEditorModule& GetModule();
+	static UE_API FStateTreeEditorModule& GetModule();
 
 	/** Gets this module, will not attempt to load and may not exist. */
-	static FStateTreeEditorModule* GetModulePtr();
+	static UE_API FStateTreeEditorModule* GetModulePtr();
 
 	/** Creates an instance of StateTree editor. Only virtual so that it can be called across the DLL boundary. */
-	virtual TSharedRef<IStateTreeEditor> CreateStateTreeEditor(const EToolkitMode::Type Mode, const TSharedPtr< class IToolkitHost >& InitToolkitHost, UStateTree* StateTree);
+	UE_API virtual TSharedRef<IStateTreeEditor> CreateStateTreeEditor(const EToolkitMode::Type Mode, const TSharedPtr< class IToolkitHost >& InitToolkitHost, UStateTree* StateTree);
 
 	/** Sets the Details View with required State Tree Detail Property Handlers */
-	static void SetDetailPropertyHandlers(IDetailsView& DetailsView);
+	static UE_API void SetDetailPropertyHandlers(IDetailsView& DetailsView);
 
 	virtual TSharedPtr<FExtensibilityManager> GetMenuExtensibilityManager() override { return MenuExtensibilityManager; }
 	virtual TSharedPtr<FExtensibilityManager> GetToolBarExtensibilityManager() override { return ToolBarExtensibilityManager; }
 
-	TSharedPtr<FStateTreeNodeClassCache> GetNodeClassCache();
+	UE_API TSharedPtr<FStateTreeNodeClassCache> GetNodeClassCache();
 	
 	DECLARE_EVENT_OneParam(FStateTreeEditorModule, FOnRegisterLayoutExtensions, FLayoutExtender&);
 	FOnRegisterLayoutExtensions& OnRegisterLayoutExtensions() { return RegisterLayoutExtensions; }
 
 	/** Register the editor data type for a specific schema. */
-	void RegisterEditorDataClass(TNonNullSubclassOf<const UStateTreeSchema> Schema, TNonNullSubclassOf<const UStateTreeEditorData> EditorData);
+	UE_API void RegisterEditorDataClass(TNonNullSubclassOf<const UStateTreeSchema> Schema, TNonNullSubclassOf<const UStateTreeEditorData> EditorData);
 	/** Unregister the editor data type for a specific schema. */
-	void UnregisterEditorDataClass(TNonNullSubclassOf<const UStateTreeSchema> Schema);
+	UE_API void UnregisterEditorDataClass(TNonNullSubclassOf<const UStateTreeSchema> Schema);
 	/** Get the editor data type for a specific schema. */
-	TNonNullSubclassOf<UStateTreeEditorData> GetEditorDataClass(TNonNullSubclassOf<const UStateTreeSchema> Schema) const;
+	UE_API TNonNullSubclassOf<UStateTreeEditorData> GetEditorDataClass(TNonNullSubclassOf<const UStateTreeSchema> Schema) const;
 
 protected:
 	TSharedPtr<FExtensibilityManager> MenuExtensibilityManager;
@@ -83,3 +85,5 @@ protected:
 	FDelegateHandle OnUserDefinedStructReinstancedHandle;
 	FOnRegisterLayoutExtensions RegisterLayoutExtensions;
 };
+
+#undef UE_API
diff --git a/Engine/Plugins/Runtime/StateTree/Source/StateTreeEditorModule/Public/StateTreeEditorNode.h b/Engine/Plugins/Runtime/StateTree/Source/StateTreeEditorModule/Public/StateTreeEditorNode.h
--- a/Engine/Plugins/Runtime/StateTree/Source/StateTreeEditorModule/Public/StateTreeEditorNode.h
+++ b/Engine/Plugins/Runtime/StateTree/Source/StateTreeEditorModule/Public/StateTreeEditorNode.h
@@ -5,6 +5,8 @@
 #include "StateTreeNodeBase.h"
 #include "StateTreeEditorNode.generated.h"
 
+#define UE_API STATETREEEDITORMODULE_API
+
 UENUM()
 enum class EStateTreeNodeType : uint8
 {
@@ -20,7 +22,7 @@ enum class EStateTreeNodeType : uint8
  * Base for Evaluator, Task and Condition nodes.
  */
 USTRUCT()
-struct STATETREEEDITORMODULE_API FStateTreeEditorNode
+struct FStateTreeEditorNode
 {
 	GENERATED_BODY()
 
@@ -37,7 +39,7 @@ struct STATETREEEDITORMODULE_API FStateTreeEditorNode
 	 * If the node has a specified name, used that, or else of return the display name of the node.
 	 * @return name of the node.
 	 */
-	FName GetName() const;
+	UE_API FName GetName() const;
 
 	FStateTreeDataView GetInstance() const
 	{
@@ -92,3 +94,5 @@ struct TStateTreeEditorNode : public FStateTreeEditorNode
 	FORCEINLINE T& GetNode() { return Node.template GetMutable<T>(); }
 	FORCEINLINE typename T::FInstanceDataType& GetInstanceData() { return Instance.template GetMutable<typename T::FInstanceDataType>(); }
 };
+
+#undef UE_API
diff --git a/Engine/Plugins/Runtime/StateTree/Source/StateTreeEditorModule/Public/StateTreeEditorPropertyBindings.h b/Engine/Plugins/Runtime/StateTree/Source/StateTreeEditorModule/Public/StateTreeEditorPropertyBindings.h
--- a/Engine/Plugins/Runtime/StateTree/Source/StateTreeEditorModule/Public/StateTreeEditorPropertyBindings.h
+++ b/Engine/Plugins/Runtime/StateTree/Source/StateTreeEditorModule/Public/StateTreeEditorPropertyBindings.h
@@ -9,6 +9,8 @@
 #include "StateTreeEditorNode.h"
 #include "StateTreeEditorPropertyBindings.generated.h"
 
+#define UE_API STATETREEEDITORMODULE_API
+
 enum class EStateTreeNodeFormatting : uint8;
 class IStateTreeEditorPropertyBindingsOwner;
 enum class EStateTreeVisitor : uint8;
@@ -17,7 +19,7 @@ enum class EStateTreeVisitor : uint8;
  * Editor representation of all property bindings in a StateTree
  */
 USTRUCT()
-struct STATETREEEDITORMODULE_API FStateTreeEditorPropertyBindings : public FPropertyBindingBindingCollection
+struct FStateTreeEditorPropertyBindings : public FPropertyBindingBindingCollection
 {
 	GENERATED_BODY()
 
@@ -40,15 +42,15 @@ struct STATETREEEDITORMODULE_API FStateTreeEditorPropertyBindings : public FProp
 	}
 
 	//~ Begin FPropertyBindingBindingCollection overrides
-	virtual int32 GetNumBindableStructDescriptors() const override;
-	virtual const FPropertyBindingBindableStructDescriptor* GetBindableStructDescriptorFromHandle(FConstStructView InSourceHandleView) const override;
-
-	virtual int32 GetNumBindings() const override;
-	virtual void ForEachBinding(TFunctionRef<void(const FPropertyBindingBinding& Binding)> InFunction) const override;
-	virtual void ForEachBinding(FPropertyBindingIndex16 InBegin, FPropertyBindingIndex16 InEnd, TFunctionRef<void(const FPropertyBindingBinding& Binding, const int32 BindingIndex)> InFunction) const override;
-	virtual void ForEachMutableBinding(TFunctionRef<void(FPropertyBindingBinding& Binding)> InFunction) override;
-	virtual void VisitBindings(TFunctionRef<EVisitResult(const FPropertyBindingBinding& Binding)> InFunction) const override;
-	virtual void VisitMutableBindings(TFunctionRef<EVisitResult(FPropertyBindingBinding& Binding)> InFunction) override;
+	UE_API virtual int32 GetNumBindableStructDescriptors() const override;
+	UE_API virtual const FPropertyBindingBindableStructDescriptor* GetBindableStructDescriptorFromHandle(FConstStructView InSourceHandleView) const override;
+
+	UE_API virtual int32 GetNumBindings() const override;
+	UE_API virtual void ForEachBinding(TFunctionRef<void(const FPropertyBindingBinding& Binding)> InFunction) const override;
+	UE_API virtual void ForEachBinding(FPropertyBindingIndex16 InBegin, FPropertyBindingIndex16 InEnd, TFunctionRef<void(const FPropertyBindingBinding& Binding, const int32 BindingIndex)> InFunction) const override;
+	UE_API virtual void ForEachMutableBinding(TFunctionRef<void(FPropertyBindingBinding& Binding)> InFunction) override;
+	UE_API virtual void VisitBindings(TFunctionRef<EVisitResult(const FPropertyBindingBinding& Binding)> InFunction) const override;
+	UE_API virtual void VisitMutableBindings(TFunctionRef<EVisitResult(FPropertyBindingBinding& Binding)> InFunction) override;
 	//~ End FPropertyBindingBindingCollection overrides
 
 PRAGMA_DISABLE_DEPRECATION_WARNINGS
@@ -79,7 +81,7 @@ PRAGMA_DISABLE_DEPRECATION_WARNINGS
 	 * @param InTargetPath Binding target property path.
 	 * @return Constructed binding source property path.
 	 */
-	FPropertyBindingPath AddFunctionBinding(const UScriptStruct* InPropertyFunctionNodeStruct, TConstArrayView<FPropertyBindingPathSegment> InSourcePathSegments, const FPropertyBindingPath& InTargetPath);
+	UE_API FPropertyBindingPath AddFunctionBinding(const UScriptStruct* InPropertyFunctionNodeStruct, TConstArrayView<FPropertyBindingPathSegment> InSourcePathSegments, const FPropertyBindingPath& InTargetPath);
 	UE_DEPRECATED(5.6, "Use AddFunctionBinding taking FPropertyBindingPath instead")
 	FStateTreePropertyPath AddFunctionPropertyBinding(const UScriptStruct* InPropertyFunctionNodeStruct, TConstArrayView<FPropertyBindingPathSegment> InSourcePathSegments, const FStateTreePropertyPath& InTargetPath)
 	{
@@ -109,9 +111,9 @@ PRAGMA_DISABLE_DEPRECATION_WARNINGS
 	/**
 	 * @return binding to the target path.
 	 */
-	const FStateTreePropertyPathBinding* FindPropertyBinding(const FPropertyBindingPath& TargetPath, ESearchMode SearchMode = ESearchMode::Exact) const;
+	UE_API const FStateTreePropertyPathBinding* FindPropertyBinding(const FPropertyBindingPath& TargetPath, ESearchMode SearchMode = ESearchMode::Exact) const;
 	UE_DEPRECATED(5.6, "Use the version taking FPropertyBindingPath instead")
-	const FStateTreePropertyPathBinding* FindPropertyBinding(const FStateTreePropertyPath& TargetPath, ESearchMode SearchMode = ESearchMode::Exact) const;
+	UE_API const FStateTreePropertyPathBinding* FindPropertyBinding(const FStateTreePropertyPath& TargetPath, ESearchMode SearchMode = ESearchMode::Exact) const;
 
 	/**
 	 * @return Source path for given target path, or null if binding does not exists.
@@ -142,16 +144,16 @@ PRAGMA_DISABLE_DEPRECATION_WARNINGS
 	}
 
 	UE_DEPRECATED(5.5, "Use GetPropertyBindingsFor returning pointers instead.")
-	void GetPropertyBindingsFor(const FGuid StructID, TArray<FStateTreePropertyPathBinding>& OutBindings) const;
+	UE_API void GetPropertyBindingsFor(const FGuid StructID, TArray<FStateTreePropertyPathBinding>& OutBindings) const;
 PRAGMA_ENABLE_DEPRECATION_WARNINGS
 
 protected:
 	//~ Begin FPropertyBindingBindingCollection overrides
 	virtual FPropertyBindingBinding* AddBindingInternal(const FPropertyBindingPath& InSourcePath, const FPropertyBindingPath& InTargetPath) override;
-	virtual void CopyBindingsInternal(const FGuid InFromStructID, const FGuid InToStructID) override;
-	virtual void RemoveBindingsInternal(TFunctionRef<bool(FPropertyBindingBinding&)> InPredicate) override;
-	virtual bool HasBindingInternal(TFunctionRef<bool(const FPropertyBindingBinding&)> InPredicate) const override;
-	virtual const FPropertyBindingBinding* FindBindingInternal(TFunctionRef<bool(const FPropertyBindingBinding&)> InPredicate) const override;
+	UE_API virtual void CopyBindingsInternal(const FGuid InFromStructID, const FGuid InToStructID) override;
+	UE_API virtual void RemoveBindingsInternal(TFunctionRef<bool(FPropertyBindingBinding&)> InPredicate) override;
+	UE_API virtual bool HasBindingInternal(TFunctionRef<bool(const FPropertyBindingBinding&)> InPredicate) const override;
+	UE_API virtual const FPropertyBindingBinding* FindBindingInternal(TFunctionRef<bool(const FPropertyBindingBinding&)> InPredicate) const override;
 	//~ Begin FPropertyBindingBindingCollection overrides
 
 private:
@@ -179,7 +181,7 @@ struct UE_DEPRECATED(5.6, "Use FPropertyCreationDesc instead") FStateTreeEditorP
 	const void* SourceContainerAddress = nullptr;
 };
 
-class STATETREEEDITORMODULE_API IStateTreeEditorPropertyBindingsOwner : public IPropertyBindingBindingCollectionOwner
+class IStateTreeEditorPropertyBindingsOwner : public IPropertyBindingBindingCollectionOwner
 {
 	GENERATED_BODY()
 
@@ -251,15 +253,17 @@ PRAGMA_ENABLE_DEPRECATION_WARNINGS
 // Currently FStateTreeEditorPropertyBindings is meant to be used as a member for just to store things,
 // IStateTreeEditorPropertyBindingsOwner is meant return model specific stuff,
 // and IStateTreeBindingLookup is used in non-editor code and it cannot be in FStateTreeEditorPropertyBindings because bindings don't know about the owner.
-struct STATETREEEDITORMODULE_API FStateTreeBindingLookup : public IStateTreeBindingLookup
+struct FStateTreeBindingLookup : public IStateTreeBindingLookup
 {
-	FStateTreeBindingLookup(const IStateTreeEditorPropertyBindingsOwner* InBindingOwner);
+	UE_API FStateTreeBindingLookup(const IStateTreeEditorPropertyBindingsOwner* InBindingOwner);
 
 	const IStateTreeEditorPropertyBindingsOwner* BindingOwner = nullptr;
 
-	virtual const FPropertyBindingPath* GetPropertyBindingSource(const FPropertyBindingPath& InTargetPath) const override;
-	virtual FText GetPropertyPathDisplayName(const FPropertyBindingPath& InTargetPath, EStateTreeNodeFormatting Formatting) const override;
-	virtual FText GetBindingSourceDisplayName(const FPropertyBindingPath& InTargetPath, EStateTreeNodeFormatting Formatting) const override;
-	virtual const FProperty* GetPropertyPathLeafProperty(const FPropertyBindingPath& InPath) const override;
+	UE_API virtual const FPropertyBindingPath* GetPropertyBindingSource(const FPropertyBindingPath& InTargetPath) const override;
+	UE_API virtual FText GetPropertyPathDisplayName(const FPropertyBindingPath& InTargetPath, EStateTreeNodeFormatting Formatting) const override;
+	UE_API virtual FText GetBindingSourceDisplayName(const FPropertyBindingPath& InTargetPath, EStateTreeNodeFormatting Formatting) const override;
+	UE_API virtual const FProperty* GetPropertyPathLeafProperty(const FPropertyBindingPath& InPath) const override;
 
 };
+
+#undef UE_API
diff --git a/Engine/Plugins/Runtime/StateTree/Source/StateTreeEditorModule/Public/StateTreeEditorSettings.h b/Engine/Plugins/Runtime/StateTree/Source/StateTreeEditorModule/Public/StateTreeEditorSettings.h
--- a/Engine/Plugins/Runtime/StateTree/Source/StateTreeEditorModule/Public/StateTreeEditorSettings.h
+++ b/Engine/Plugins/Runtime/StateTree/Source/StateTreeEditorModule/Public/StateTreeEditorSettings.h
@@ -5,6 +5,8 @@
 #include "Engine/DeveloperSettings.h"
 #include "StateTreeEditorSettings.generated.h"
 
+#define UE_API STATETREEEDITORMODULE_API
+
 UENUM()
 enum class EStateTreeSaveOnCompile : uint8
 {
@@ -13,8 +15,8 @@ enum class EStateTreeSaveOnCompile : uint8
 	Always UMETA(DisplayName = "Always"),
 };
 
-UCLASS(config = EditorPerProjectUserSettings)
-class STATETREEEDITORMODULE_API UStateTreeEditorSettings : public UDeveloperSettings
+UCLASS(MinimalAPI, config = EditorPerProjectUserSettings)
+class UStateTreeEditorSettings : public UDeveloperSettings
 {
 	GENERATED_UCLASS_BODY()
 public:
@@ -41,9 +43,11 @@ public:
 
 protected:
 #if WITH_EDITOR
-	virtual FText GetSectionText() const override;
-	virtual FText GetSectionDescription() const override;
+	UE_API virtual FText GetSectionText() const override;
+	UE_API virtual FText GetSectionDescription() const override;
 #endif
 
-	virtual FName GetCategoryName() const override;
+	UE_API virtual FName GetCategoryName() const override;
 };
+
+#undef UE_API
diff --git a/Engine/Plugins/Runtime/StateTree/Source/StateTreeEditorModule/Public/StateTreeEditorStyle.h b/Engine/Plugins/Runtime/StateTree/Source/StateTreeEditorModule/Public/StateTreeEditorStyle.h
--- a/Engine/Plugins/Runtime/StateTree/Source/StateTreeEditorModule/Public/StateTreeEditorStyle.h
+++ b/Engine/Plugins/Runtime/StateTree/Source/StateTreeEditorModule/Public/StateTreeEditorStyle.h
@@ -4,25 +4,29 @@
 
 #include "Styling/SlateStyle.h"
 
+#define UE_API STATETREEEDITORMODULE_API
+
 enum class EStateTreeStateSelectionBehavior : uint8;
 enum class EStateTreeStateType : uint8;
 
 class ISlateStyle;
 
-class STATETREEEDITORMODULE_API FStateTreeEditorStyle
+class FStateTreeEditorStyle
 	: public FSlateStyleSet
 {
 public:
-	static FStateTreeEditorStyle& Get();
+	static UE_API FStateTreeEditorStyle& Get();
 	
-	static const FSlateBrush* GetBrushForSelectionBehaviorType(EStateTreeStateSelectionBehavior InBehaviour, bool bHasChildren, EStateTreeStateType StateType);	
+	static UE_API const FSlateBrush* GetBrushForSelectionBehaviorType(EStateTreeStateSelectionBehavior InBehaviour, bool bHasChildren, EStateTreeStateType StateType);	
 
 protected:
 	friend class FStateTreeEditorModule;
 
-	static void Register();
-	static void Unregister();
+	static UE_API void Register();
+	static UE_API void Unregister();
 
 private:
-	FStateTreeEditorStyle();
+	UE_API FStateTreeEditorStyle();
 };
+
+#undef UE_API
diff --git a/Engine/Plugins/Runtime/StateTree/Source/StateTreeEditorModule/Public/StateTreeEditorWorkspaceTabHost.h b/Engine/Plugins/Runtime/StateTree/Source/StateTreeEditorModule/Public/StateTreeEditorWorkspaceTabHost.h
--- a/Engine/Plugins/Runtime/StateTree/Source/StateTreeEditorModule/Public/StateTreeEditorWorkspaceTabHost.h
+++ b/Engine/Plugins/Runtime/StateTree/Source/StateTreeEditorModule/Public/StateTreeEditorWorkspaceTabHost.h
@@ -6,6 +6,8 @@
 #include "Textures/SlateIcon.h"
 #include "Framework/Docking/TabManager.h"
 
+#define UE_API STATETREEEDITORMODULE_API
+
 class SDockTab;
 
 namespace UE::StateTreeEditor
@@ -26,35 +28,37 @@ namespace UE::StateTreeEditor
 	};
 
 	// Interface required for re-using the same tab management across different AssetEditors
-	class STATETREEEDITORMODULE_API FWorkspaceTabHost : public TSharedFromThis<FWorkspaceTabHost>
+	class FWorkspaceTabHost : public TSharedFromThis<FWorkspaceTabHost>
 	{
 	public:
-		static const FLazyName BindingTabId;
-		static const FLazyName DebuggerTabId;
-		static const FLazyName OutlinerTabId;
-		static const FLazyName SearchTabId;
-		static const FLazyName StatisticsTabId;
+		static UE_API const FLazyName BindingTabId;
+		static UE_API const FLazyName DebuggerTabId;
+		static UE_API const FLazyName OutlinerTabId;
+		static UE_API const FLazyName SearchTabId;
+		static UE_API const FLazyName StatisticsTabId;
 
 	public:
 		virtual ~FWorkspaceTabHost() = default;
 
-		FOnSpawnTab CreateSpawnDelegate(FName TabID);
-		TConstArrayView<FMinorWorkspaceTabConfig> GetTabConfigs() const;
+		UE_API FOnSpawnTab CreateSpawnDelegate(FName TabID);
+		UE_API TConstArrayView<FMinorWorkspaceTabConfig> GetTabConfigs() const;
 
 		TConstArrayView<FSpawnedWorkspaceTab> GetSpawnedTabs() const
 		{
 			return SpawnedTabs;
 		}
 
-		TSharedRef<SDockTab> Spawn(FName TabID);
+		UE_API TSharedRef<SDockTab> Spawn(FName TabID);
 
 		DECLARE_MULTICAST_DELEGATE_OneParam(FOnWorkspaceSpawnTab, FSpawnedWorkspaceTab);
 		FOnWorkspaceSpawnTab OnTabSpawned;
 		FOnWorkspaceSpawnTab OnTabClosed;
 
 	private:
-		TSharedRef<SDockTab> HandleSpawnDelegate(const FSpawnTabArgs& Args, FName TabID);
-		void HandleTabClosed(TSharedRef<SDockTab>);
+		UE_API TSharedRef<SDockTab> HandleSpawnDelegate(const FSpawnTabArgs& Args, FName TabID);
+		UE_API void HandleTabClosed(TSharedRef<SDockTab>);
 		TArray<FSpawnedWorkspaceTab> SpawnedTabs;
 	};
 }
+
+#undef UE_API
diff --git a/Engine/Plugins/Runtime/StateTree/Source/StateTreeEditorModule/Public/StateTreeFactory.h b/Engine/Plugins/Runtime/StateTree/Source/StateTreeEditorModule/Public/StateTreeFactory.h
--- a/Engine/Plugins/Runtime/StateTree/Source/StateTreeEditorModule/Public/StateTreeFactory.h
+++ b/Engine/Plugins/Runtime/StateTree/Source/StateTreeEditorModule/Public/StateTreeFactory.h
@@ -5,26 +5,30 @@
 #include "Factories/Factory.h"
 #include "StateTreeFactory.generated.h"
 
+#define UE_API STATETREEEDITORMODULE_API
+
 class UStateTreeSchema;
 
 /**
  * Factory for UStateTree
  */
 
-UCLASS()
-class STATETREEEDITORMODULE_API UStateTreeFactory : public UFactory
+UCLASS(MinimalAPI)
+class UStateTreeFactory : public UFactory
 {
 	GENERATED_UCLASS_BODY()
 
 	//~ UFactory interface
-	virtual bool ConfigureProperties() override;
-	virtual UObject* FactoryCreateNew(UClass* Class, UObject* InParent, FName Name, EObjectFlags Flags, UObject* Context, FFeedbackContext* Warn) override;
+	UE_API virtual bool ConfigureProperties() override;
+	UE_API virtual UObject* FactoryCreateNew(UClass* Class, UObject* InParent, FName Name, EObjectFlags Flags, UObject* Context, FFeedbackContext* Warn) override;
 	//~ End of UFactory interface
 
-	void SetSchemaClass(const TObjectPtr<UClass>& InSchemaClass);  
+	UE_API void SetSchemaClass(const TObjectPtr<UClass>& InSchemaClass);  
 
 protected:
 	
 	UPROPERTY(Transient)
 	TObjectPtr<UClass> StateTreeSchemaClass;
 };
+
+#undef UE_API
diff --git a/Engine/Plugins/Runtime/StateTree/Source/StateTreeEditorModule/Public/StateTreeNodeClassCache.h b/Engine/Plugins/Runtime/StateTree/Source/StateTreeEditorModule/Public/StateTreeNodeClassCache.h
--- a/Engine/Plugins/Runtime/StateTree/Source/StateTreeEditorModule/Public/StateTreeNodeClassCache.h
+++ b/Engine/Plugins/Runtime/StateTree/Source/StateTreeEditorModule/Public/StateTreeNodeClassCache.h
@@ -4,6 +4,8 @@
 
 #include "StateTreeNodeClassCache.generated.h"
 
+#define UE_API STATETREEEDITORMODULE_API
+
 enum class EReloadCompleteReason;
 
 struct FAssetData;
@@ -14,17 +16,17 @@ struct FStateTreeNodeBase;
  * If the class or struct is from a package that is not yet loaded, the data will update on GetStruct/Class/Scripstruct()
  */
 USTRUCT()
-struct STATETREEEDITORMODULE_API FStateTreeNodeClassData
+struct FStateTreeNodeClassData
 {
 	GENERATED_BODY()
 
 	FStateTreeNodeClassData() {}
-	FStateTreeNodeClassData(UStruct* InStruct);
-	FStateTreeNodeClassData(const FString& InClassAssetName, const FString& InClassPackage, const FName InStructName, UStruct* InStruct);
+	UE_API FStateTreeNodeClassData(UStruct* InStruct);
+	UE_API FStateTreeNodeClassData(const FString& InClassAssetName, const FString& InClassPackage, const FName InStructName, UStruct* InStruct);
 
 	FName GetStructName() const { return StructName; }
 
-	UStruct* GetStruct(bool bSilent = false);
+	UE_API UStruct* GetStruct(bool bSilent = false);
 
 	UClass* GetClass(bool bSilent = false)
 	{
@@ -35,7 +37,7 @@ struct STATETREEEDITORMODULE_API FStateTreeNodeClassData
 		return Cast<UScriptStruct>(GetStruct(bSilent));
 	}
 
-	const UStruct* GetInstanceDataStruct(bool bSilent = false);
+	UE_API const UStruct* GetInstanceDataStruct(bool bSilent = false);
 
 private:
 
@@ -59,13 +61,13 @@ private:
  * Caches specified classes or structs and reacts to engine events to keep the lists always up to date.
  * All the derived classes or structs are kept in the cache.
  */
-struct STATETREEEDITORMODULE_API FStateTreeNodeClassCache
+struct FStateTreeNodeClassCache
 {
-	FStateTreeNodeClassCache();
-	~FStateTreeNodeClassCache();
+	UE_API FStateTreeNodeClassCache();
+	UE_API ~FStateTreeNodeClassCache();
 
 	/** Adds a Struct to keep track of */
-	void AddRootStruct(UStruct* RootStruct);
+	UE_API void AddRootStruct(UStruct* RootStruct);
 	
 	/** Adds a Class to keep track of */
 	void AddRootClass(UClass* RootClass)
@@ -80,7 +82,7 @@ struct STATETREEEDITORMODULE_API FStateTreeNodeClassCache
 	}
 
 	/** Returns know derived Structs based on provided base. If the base Struct is not added as root Struct, nothing is returned. */
-	void GetStructs(const UStruct* BaseStruct, TArray<TSharedPtr<FStateTreeNodeClassData>>& AvailableClasses);
+	UE_API void GetStructs(const UStruct* BaseStruct, TArray<TSharedPtr<FStateTreeNodeClassData>>& AvailableClasses);
 	
 	/** Returns know derived Classes based on provided base. If the base Class is not added as root Class, nothing is returned. */
 	void GetClasses(const UStruct* BaseClass, TArray<TSharedPtr<FStateTreeNodeClassData>>& AvailableClasses)
@@ -95,16 +97,16 @@ struct STATETREEEDITORMODULE_API FStateTreeNodeClassCache
 	}
 
 	/** Invalidates the cache, it will be rebuild on next access. */
-	void InvalidateCache();
+	UE_API void InvalidateCache();
 
 protected:
-	void OnAssetAdded(const FAssetData& AssetData);
-	void OnAssetRemoved(const FAssetData& AssetData);
-	void OnReloadComplete(EReloadCompleteReason Reason);
+	UE_API void OnAssetAdded(const FAssetData& AssetData);
+	UE_API void OnAssetRemoved(const FAssetData& AssetData);
+	UE_API void OnReloadComplete(EReloadCompleteReason Reason);
 	
 private:
-	void UpdateBlueprintClass(const FAssetData& AssetData);
-	void CacheClasses();
+	UE_API void UpdateBlueprintClass(const FAssetData& AssetData);
+	UE_API void CacheClasses();
 
 	struct FRootClassContainer
 	{
@@ -121,3 +123,5 @@ private:
 	TArray<FRootClassContainer> RootClasses;
 	TMap<FString, int32> ClassNameToRootIndex;
 };
+
+#undef UE_API
diff --git a/Engine/Plugins/Runtime/StateTree/Source/StateTreeEditorModule/Public/StateTreePropertyBindingCompiler.h b/Engine/Plugins/Runtime/StateTree/Source/StateTreeEditorModule/Public/StateTreePropertyBindingCompiler.h
--- a/Engine/Plugins/Runtime/StateTree/Source/StateTreeEditorModule/Public/StateTreePropertyBindingCompiler.h
+++ b/Engine/Plugins/Runtime/StateTree/Source/StateTreeEditorModule/Public/StateTreePropertyBindingCompiler.h
@@ -7,6 +7,8 @@
 #include "StateTreePropertyBindings.h"
 #include "StateTreePropertyBindingCompiler.generated.h"
 
+#define UE_API STATETREEEDITORMODULE_API
+
 enum class EPropertyAccessCompatibility;
 struct FStateTreeCompilerLog;
 struct FStateTreePropertyPathBinding;
@@ -16,7 +18,7 @@ struct FStateTreePropertyPathBinding;
  * TODO: Better error reporting, something that can be shown in the UI.
  */
 USTRUCT()
-struct STATETREEEDITORMODULE_API FStateTreePropertyBindingCompiler
+struct FStateTreePropertyBindingCompiler
 {
 	GENERATED_BODY()
 
@@ -25,7 +27,7 @@ struct STATETREEEDITORMODULE_API FStateTreePropertyBindingCompiler
 	  * @param PropertyBindings - Reference to the Property Bindings where all the batches will be stored.
 	  * @return true on success.
 	  */
-	[[nodiscard]] bool Init(FStateTreePropertyBindings& InPropertyBindings, FStateTreeCompilerLog& InLog);
+	[[nodiscard]] UE_API bool Init(FStateTreePropertyBindings& InPropertyBindings, FStateTreeCompilerLog& InLog);
 
 	/**
 	  * Compiles a batch of property copies.
@@ -36,7 +38,7 @@ struct STATETREEEDITORMODULE_API FStateTreePropertyBindingCompiler
 	  * @param OutBatchIndex - Resulting batch index, if index is INDEX_NONE, no bindings were found and no batch was generated.
 	  * @return True on success, false on failure.
 	 */
-	[[nodiscard]] bool CompileBatch(const FStateTreeBindableStructDesc& TargetStruct, TConstArrayView<FStateTreePropertyPathBinding> PropertyBindings, FStateTreeIndex16 PropertyFuncsBegin, FStateTreeIndex16 PropertyFuncsEnd, int32& OutBatchIndex);
+	[[nodiscard]] UE_API bool CompileBatch(const FStateTreeBindableStructDesc& TargetStruct, TConstArrayView<FStateTreePropertyPathBinding> PropertyBindings, FStateTreeIndex16 PropertyFuncsBegin, FStateTreeIndex16 PropertyFuncsEnd, int32& OutBatchIndex);
 
 	 /**
 	  * Compiles delegates dispatcher for selected struct.
@@ -46,7 +48,7 @@ struct STATETREEEDITORMODULE_API FStateTreePropertyBindingCompiler
 	  * @param InstanceDataView - View to the instance data
 	  * @return True on success, false on failure.
 	 */
-	[[nodiscard]] bool CompileDelegateDispatchers(const FStateTreeBindableStructDesc& SourceStruct, TConstArrayView<FStateTreeEditorDelegateDispatcherCompiledBinding> PreviousDistpatchers, TConstArrayView<FStateTreePropertyPathBinding> DelegateSourceBindings, FStateTreeDataView InstanceDataView);
+	[[nodiscard]] UE_API bool CompileDelegateDispatchers(const FStateTreeBindableStructDesc& SourceStruct, TConstArrayView<FStateTreeEditorDelegateDispatcherCompiledBinding> PreviousDistpatchers, TConstArrayView<FStateTreePropertyPathBinding> DelegateSourceBindings, FStateTreeDataView InstanceDataView);
 
 	 /**
 	  * Compiles delegates listener for selected struct.
@@ -55,7 +57,7 @@ struct STATETREEEDITORMODULE_API FStateTreePropertyBindingCompiler
 	  * @param InstanceDataView - View to the instance data
 	  * @return True on success, false on failure.
 	 */
-	[[nodiscard]] bool CompileDelegateListeners(const FStateTreeBindableStructDesc& TargetStruct, TConstArrayView<FStateTreePropertyPathBinding> DelegateSourceBindings, FStateTreeDataView InstanceDataView);
+	[[nodiscard]] UE_API bool CompileDelegateListeners(const FStateTreeBindableStructDesc& TargetStruct, TConstArrayView<FStateTreePropertyPathBinding> DelegateSourceBindings, FStateTreeDataView InstanceDataView);
 
 	/**
 	  * Compiles references for selected struct
@@ -64,10 +66,10 @@ struct STATETREEEDITORMODULE_API FStateTreePropertyBindingCompiler
 	  * @param InstanceDataView - view to the instance data
 	  * @return True on success, false on failure.
 	 */
-	[[nodiscard]] bool CompileReferences(const FStateTreeBindableStructDesc& TargetStruct, TConstArrayView<FStateTreePropertyPathBinding> PropertyReferenceBindings, FStateTreeDataView InstanceDataView, const TMap<FGuid, const FStateTreeDataView>& IDToStructValue);
+	[[nodiscard]] UE_API bool CompileReferences(const FStateTreeBindableStructDesc& TargetStruct, TConstArrayView<FStateTreePropertyPathBinding> PropertyReferenceBindings, FStateTreeDataView InstanceDataView, const TMap<FGuid, const FStateTreeDataView>& IDToStructValue);
 
 	/** Finalizes compilation, should be called once all batches are compiled. */
-	void Finalize();
+	UE_API void Finalize();
 
 	/**
 	  * Adds source struct. When compiling a batch, the bindings can be between any of the defined source structs, and the target struct.
@@ -75,13 +77,13 @@ struct STATETREEEDITORMODULE_API FStateTreePropertyBindingCompiler
 	  * @param SourceStruct - Description of the source struct to add.
 	  * @return Source struct index.
 	  */
-	int32 AddSourceStruct(const FStateTreeBindableStructDesc& SourceStruct);
+	UE_API int32 AddSourceStruct(const FStateTreeBindableStructDesc& SourceStruct);
 
 	/** @return delegate dispatcher residing behind given path. */
-	FStateTreeDelegateDispatcher GetDispatcherFromPath(const FPropertyBindingPath& PathToDispatcher) const;
+	UE_API FStateTreeDelegateDispatcher GetDispatcherFromPath(const FPropertyBindingPath& PathToDispatcher) const;
 
 	/** @return the list of compiled delegate dispatcher. */
-	TArray<FStateTreeEditorDelegateDispatcherCompiledBinding> GetCompiledDelegateDispatchers() const;
+	UE_API TArray<FStateTreeEditorDelegateDispatcherCompiledBinding> GetCompiledDelegateDispatchers() const;
 
 	const FStateTreeBindableStructDesc* GetSourceStructDescByID(const FGuid& ID) const
 	{
@@ -90,7 +92,7 @@ struct STATETREEEDITORMODULE_API FStateTreePropertyBindingCompiler
 PRAGMA_DISABLE_DEPRECATION_WARNINGS
 
 	UE_DEPRECATED(5.6, "Use the version taking FPropertyBindingPath instead")
-	FStateTreeIndex16 GetDispatcherIDFromPath(const FStateTreePropertyPath& PathToDispatcher) const;
+	UE_API FStateTreeIndex16 GetDispatcherIDFromPath(const FStateTreePropertyPath& PathToDispatcher) const;
 
 	UE_DEPRECATED(5.5, "Use CompileBatch with PropertyFuncsBegin and PropertyFuncsEnd instead.")
 	[[nodiscard]] bool CompileBatch(const FStateTreeBindableStructDesc& InTargetStruct, TConstArrayView<FStateTreePropertyPathBinding> InPropertyBindings, int32& OutBatchIndex) 
@@ -123,3 +125,5 @@ protected:
 
 	FStateTreeCompilerLog* Log = nullptr;
 };
+
+#undef UE_API
diff --git a/Engine/Plugins/Runtime/StateTree/Source/StateTreeEditorModule/Public/StateTreeState.h b/Engine/Plugins/Runtime/StateTree/Source/StateTreeEditorModule/Public/StateTreeState.h
--- a/Engine/Plugins/Runtime/StateTree/Source/StateTreeEditorModule/Public/StateTreeState.h
+++ b/Engine/Plugins/Runtime/StateTree/Source/StateTreeEditorModule/Public/StateTreeState.h
@@ -9,6 +9,8 @@
 #include "StateTreeEvents.h"
 #include "StateTreeState.generated.h"
 
+#define UE_API STATETREEEDITORMODULE_API
+
 class UStateTreeState;
 class UStateTree;
 
@@ -75,15 +77,15 @@ struct FStateTreeTransitionDelegateListener
  * Editor representation of a transition in StateTree
  */
 USTRUCT()
-struct STATETREEEDITORMODULE_API FStateTreeTransition
+struct FStateTreeTransition
 {
 	GENERATED_BODY()
 
 	// Macro needed to avoid deprecation errors with members being copied or created.
 PRAGMA_DISABLE_DEPRECATION_WARNINGS
 	FStateTreeTransition() = default;
-	FStateTreeTransition(const EStateTreeTransitionTrigger InTrigger, const EStateTreeTransitionType InType, const UStateTreeState* InState = nullptr);
-	FStateTreeTransition(const EStateTreeTransitionTrigger InTrigger, const FGameplayTag InEventTag, const EStateTreeTransitionType InType, const UStateTreeState* InState = nullptr);
+	UE_API FStateTreeTransition(const EStateTreeTransitionTrigger InTrigger, const EStateTreeTransitionType InType, const UStateTreeState* InState = nullptr);
+	UE_API FStateTreeTransition(const EStateTreeTransitionTrigger InTrigger, const FGameplayTag InEventTag, const EStateTreeTransitionType InType, const UStateTreeState* InState = nullptr);
 	FStateTreeTransition(const FStateTreeTransition&) = default;
 	FStateTreeTransition(FStateTreeTransition&&) = default;
 	FStateTreeTransition& operator=(const FStateTreeTransition&) = default;
@@ -109,7 +111,7 @@ PRAGMA_ENABLE_DEPRECATION_WARNINGS
 		return FGuid::Combine(ID, FGuid::NewDeterministicGuid(TEXT("Event")));
 	}
 
-	void PostSerialize(const FArchive& Ar);
+	UE_API void PostSerialize(const FArchive& Ar);
 
 	/** When to try trigger the transition. */
 	UPROPERTY(EditDefaultsOnly, Category = "Transition")
@@ -176,7 +178,7 @@ struct TStructOpsTypeTraits<FStateTreeTransition> : public TStructOpsTypeTraitsB
 
 
 USTRUCT()
-struct STATETREEEDITORMODULE_API FStateTreeStateParameters
+struct FStateTreeStateParameters
 {
 	GENERATED_BODY()
 
@@ -188,7 +190,7 @@ struct STATETREEEDITORMODULE_API FStateTreeStateParameters
 	}
 
 	/** Removes overrides that do appear in Parameters. */
-	void RemoveUnusedOverrides();
+	UE_API void RemoveUnusedOverrides();
 	
 	UPROPERTY(EditDefaultsOnly, Category = Parameters)
 	FInstancedPropertyBag Parameters;
@@ -207,28 +209,28 @@ struct STATETREEEDITORMODULE_API FStateTreeStateParameters
 /**
  * Editor representation of a state in StateTree
  */
-UCLASS(BlueprintType, EditInlineNew, CollapseCategories)
-class STATETREEEDITORMODULE_API UStateTreeState : public UObject, public IStateTreeSchemaProvider
+UCLASS(MinimalAPI, BlueprintType, EditInlineNew, CollapseCategories)
+class UStateTreeState : public UObject, public IStateTreeSchemaProvider
 {
 	GENERATED_BODY()
 
 public:
-	UStateTreeState(const FObjectInitializer& ObjectInitializer);
-	virtual ~UStateTreeState() override;
+	UE_API UStateTreeState(const FObjectInitializer& ObjectInitializer);
+	UE_API virtual ~UStateTreeState() override;
 
-	virtual void PostInitProperties() override;
-	virtual void PreEditChange(FEditPropertyChain& PropertyAboutToChange) override;
-	virtual void PostEditChangeChainProperty(FPropertyChangedChainEvent& PropertyChangedEvent) override;
-	virtual void PostLoad() override;;
-	void UpdateParametersFromLinkedSubtree();
-	void OnTreeCompiled(const UStateTree& StateTree);
+	UE_API virtual void PostInitProperties() override;
+	UE_API virtual void PreEditChange(FEditPropertyChain& PropertyAboutToChange) override;
+	UE_API virtual void PostEditChangeChainProperty(FPropertyChangedChainEvent& PropertyChangedEvent) override;
+	UE_API virtual void PostLoad() override;;
+	UE_API void UpdateParametersFromLinkedSubtree();
+	UE_API void OnTreeCompiled(const UStateTree& StateTree);
 
-	const UStateTreeState* GetRootState() const;
-	const UStateTreeState* GetNextSiblingState() const;
-	const UStateTreeState* GetNextSelectableSiblingState() const;
+	UE_API const UStateTreeState* GetRootState() const;
+	UE_API const UStateTreeState* GetNextSiblingState() const;
+	UE_API const UStateTreeState* GetNextSelectableSiblingState() const;
 
 	/** @return the path of the state as string. */
-	FString GetPath() const;
+	UE_API FString GetPath() const;
 	
 	/** @return true if the property of specified ID is overridden. */
 	bool IsParametersPropertyOverridden(const FGuid PropertyID) const
@@ -237,14 +239,14 @@ public:
 	}
 
 	/** Sets the override status of specified property by ID. */
-	void SetParametersPropertyOverridden(const FGuid PropertyID, const bool bIsOverridden);
+	UE_API void SetParametersPropertyOverridden(const FGuid PropertyID, const bool bIsOverridden);
 
 	/** @returns Default parameters from linked state or asset). */
-	const FInstancedPropertyBag* GetDefaultParameters() const;
+	UE_API const FInstancedPropertyBag* GetDefaultParameters() const;
 	
 	//~ StateTree Builder API
 	/** @return state link to this state. */
-	FStateTreeStateLink GetLinkToState() const;
+	UE_API FStateTreeStateLink GetLinkToState() const;
 	
 	/** Adds child state with specified name. */
 	UStateTreeState& AddChildState(const FName ChildName, const EStateTreeStateType StateType = EStateTreeStateType::State)
@@ -295,10 +297,10 @@ public:
 	}
 
 	/** Sets linked state and updates parameters to match the linked state. */
-	void SetLinkedState(FStateTreeStateLink InStateLink);
+	UE_API void SetLinkedState(FStateTreeStateLink InStateLink);
 
 	/** Sets linked asset and updates parameters to match the linked asset. */
-	void SetLinkedStateAsset(UStateTree* InLinkedAsset);
+	UE_API void SetLinkedStateAsset(UStateTree* InLinkedAsset);
 	
 	/**
 	 * Adds Transition.
@@ -325,7 +327,7 @@ public:
 
 	//~ IStateTreeSchemaProvider API
 	/**  @return Class of schema used by the state tree containing this state. */
-	virtual TSubclassOf<UStateTreeSchema> GetSchema() const override;
+	UE_API virtual TSubclassOf<UStateTreeSchema> GetSchema() const override;
 
 	//~IStateTreeSchemaProvider API
 
@@ -433,3 +435,5 @@ public:
 	UPROPERTY(meta = (ExcludeFromHash))
 	TObjectPtr<UStateTreeState> Parent = nullptr;
 };
+
+#undef UE_API
diff --git a/Engine/Plugins/Runtime/StateTree/Source/StateTreeEditorModule/Public/StateTreeViewModel.h b/Engine/Plugins/Runtime/StateTree/Source/StateTreeEditorModule/Public/StateTreeViewModel.h
--- a/Engine/Plugins/Runtime/StateTree/Source/StateTreeEditorModule/Public/StateTreeViewModel.h
+++ b/Engine/Plugins/Runtime/StateTree/Source/StateTreeEditorModule/Public/StateTreeViewModel.h
@@ -6,6 +6,8 @@
 #include "EditorUndoClient.h"
 #include "StateTreeViewModel.generated.h"
 
+#define UE_API STATETREEEDITORMODULE_API
+
 class FMenuBuilder;
 class UStateTree;
 class UStateTreeEditorData;
@@ -37,7 +39,7 @@ enum class UE_DEPRECATED(5.6, "Use the enum with the E prefix") FStateTreeViewMo
 /**
  * ModelView for editing StateTreeEditorData.
  */
-class STATETREEEDITORMODULE_API FStateTreeViewModel : public FEditorUndoClient, public TSharedFromThis<FStateTreeViewModel>
+class FStateTreeViewModel : public FEditorUndoClient, public TSharedFromThis<FStateTreeViewModel>
 {
 public:
 
@@ -50,90 +52,90 @@ public:
 	DECLARE_MULTICAST_DELEGATE_OneParam(FOnSelectionChanged, const TArray<TWeakObjectPtr<UStateTreeState>>& /*SelectedStates*/);
 	DECLARE_MULTICAST_DELEGATE_TwoParams(FOnBringNodeToFocus, const UStateTreeState* /*State*/, const FGuid /*NodeID*/);
 
-	FStateTreeViewModel();
-	virtual ~FStateTreeViewModel() override;
+	UE_API FStateTreeViewModel();
+	UE_API virtual ~FStateTreeViewModel() override;
 
-	void Init(UStateTreeEditorData* InTreeData);
+	UE_API void Init(UStateTreeEditorData* InTreeData);
 
 	//~ FEditorUndoClient
-	virtual void PostUndo(bool bSuccess) override;
-	virtual void PostRedo(bool bSuccess) override;
+	UE_API virtual void PostUndo(bool bSuccess) override;
+	UE_API virtual void PostRedo(bool bSuccess) override;
 
 	// Selection handling.
-	void ClearSelection();
-	void SetSelection(UStateTreeState* Selected);
-	void SetSelection(const TArray<TWeakObjectPtr<UStateTreeState>>& InSelection);
-	bool IsSelected(const UStateTreeState* State) const;
-	bool IsChildOfSelection(const UStateTreeState* State) const;
-	void GetSelectedStates(TArray<UStateTreeState*>& OutSelectedStates) const;
-	void GetSelectedStates(TArray<TWeakObjectPtr<UStateTreeState>>& OutSelectedStates) const;
-	bool HasSelection() const;
-
-	void BringNodeToFocus(UStateTreeState* State, const FGuid NodeID);
+	UE_API void ClearSelection();
+	UE_API void SetSelection(UStateTreeState* Selected);
+	UE_API void SetSelection(const TArray<TWeakObjectPtr<UStateTreeState>>& InSelection);
+	UE_API bool IsSelected(const UStateTreeState* State) const;
+	UE_API bool IsChildOfSelection(const UStateTreeState* State) const;
+	UE_API void GetSelectedStates(TArray<UStateTreeState*>& OutSelectedStates) const;
+	UE_API void GetSelectedStates(TArray<TWeakObjectPtr<UStateTreeState>>& OutSelectedStates) const;
+	UE_API bool HasSelection() const;
+
+	UE_API void BringNodeToFocus(UStateTreeState* State, const FGuid NodeID);
 	
 	// Returns associated state tree asset.
-	const UStateTree* GetStateTree() const;
+	UE_API const UStateTree* GetStateTree() const;
 
-	const UStateTreeEditorData* GetStateTreeEditorData() const;
+	UE_API const UStateTreeEditorData* GetStateTreeEditorData() const;
 
-	const UStateTreeState* GetStateByID(const FGuid StateID) const;
-	UStateTreeState* GetMutableStateByID(const FGuid StateID) const;
+	UE_API const UStateTreeState* GetStateByID(const FGuid StateID) const;
+	UE_API UStateTreeState* GetMutableStateByID(const FGuid StateID) const;
 	
 	// Returns array of subtrees to edit.
-	TArray<TObjectPtr<UStateTreeState>>* GetSubTrees() const;
-	int32 GetSubTreeCount() const;
-	void GetSubTrees(TArray<TWeakObjectPtr<UStateTreeState>>& OutSubtrees) const;
+	UE_API TArray<TObjectPtr<UStateTreeState>>* GetSubTrees() const;
+	UE_API int32 GetSubTreeCount() const;
+	UE_API void GetSubTrees(TArray<TWeakObjectPtr<UStateTreeState>>& OutSubtrees) const;
 
 	/** Find the states that are linked to the provided StateID. */
-	void GetLinkStates(FGuid StateID, TArray<FGuid>& LinkingIn, TArray<FGuid>& LinkedOut) const;
+	UE_API void GetLinkStates(FGuid StateID, TArray<FGuid>& LinkingIn, TArray<FGuid>& LinkedOut) const;
 
 	// Gets and sets StateTree view expansion state store in the asset.
-	void SetPersistentExpandedStates(TSet<TWeakObjectPtr<UStateTreeState>>& InExpandedStates);
-	void GetPersistentExpandedStates(TSet<TWeakObjectPtr<UStateTreeState>>& OutExpandedStates);
+	UE_API void SetPersistentExpandedStates(TSet<TWeakObjectPtr<UStateTreeState>>& InExpandedStates);
+	UE_API void GetPersistentExpandedStates(TSet<TWeakObjectPtr<UStateTreeState>>& OutExpandedStates);
 
 	// State manipulation
-	void AddState(UStateTreeState* AfterState);
-	void AddChildState(UStateTreeState* ParentState);
-	void RenameState(UStateTreeState* State, FName NewName);
-	void RemoveSelectedStates();
-	void CopySelectedStates();
-	bool CanPasteStatesFromClipboard() const;
-	void PasteStatesFromClipboard(UStateTreeState* AfterState);
-	void PasteStatesAsChildrenFromClipboard(UStateTreeState* ParentState);
-	void DuplicateSelectedStates();
-	void MoveSelectedStatesBefore(UStateTreeState* TargetState);
-	void MoveSelectedStatesAfter(UStateTreeState* TargetState);
-	void MoveSelectedStatesInto(UStateTreeState* TargetState);
-	bool CanEnableStates() const;
-	bool CanDisableStates() const;
-	void SetSelectedStatesEnabled(bool bEnable);
+	UE_API void AddState(UStateTreeState* AfterState);
+	UE_API void AddChildState(UStateTreeState* ParentState);
+	UE_API void RenameState(UStateTreeState* State, FName NewName);
+	UE_API void RemoveSelectedStates();
+	UE_API void CopySelectedStates();
+	UE_API bool CanPasteStatesFromClipboard() const;
+	UE_API void PasteStatesFromClipboard(UStateTreeState* AfterState);
+	UE_API void PasteStatesAsChildrenFromClipboard(UStateTreeState* ParentState);
+	UE_API void DuplicateSelectedStates();
+	UE_API void MoveSelectedStatesBefore(UStateTreeState* TargetState);
+	UE_API void MoveSelectedStatesAfter(UStateTreeState* TargetState);
+	UE_API void MoveSelectedStatesInto(UStateTreeState* TargetState);
+	UE_API bool CanEnableStates() const;
+	UE_API bool CanDisableStates() const;
+	UE_API void SetSelectedStatesEnabled(bool bEnable);
 
 	// EditorNode and Transition manipulation
 	// @todo: support Add, ReplaceWith and Rename
-	void DeleteNode(TWeakObjectPtr<UStateTreeState> State, const FGuid& ID);
-	void DeleteAllNodes(TWeakObjectPtr<UStateTreeState> State, const FGuid& ID);
-	void CopyNode(TWeakObjectPtr<UStateTreeState> State, const FGuid& ID);
-	void PasteNode(TWeakObjectPtr<UStateTreeState> State, const FGuid& ID);
-	void DuplicateNode(TWeakObjectPtr<UStateTreeState> State, const FGuid& ID);
+	UE_API void DeleteNode(TWeakObjectPtr<UStateTreeState> State, const FGuid& ID);
+	UE_API void DeleteAllNodes(TWeakObjectPtr<UStateTreeState> State, const FGuid& ID);
+	UE_API void CopyNode(TWeakObjectPtr<UStateTreeState> State, const FGuid& ID);
+	UE_API void PasteNode(TWeakObjectPtr<UStateTreeState> State, const FGuid& ID);
+	UE_API void DuplicateNode(TWeakObjectPtr<UStateTreeState> State, const FGuid& ID);
 
 	// Force to update the view externally.
-	void NotifyAssetChangedExternally() const;
-	void NotifyStatesChangedExternally(const TSet<UStateTreeState*>& ChangedStates, const FPropertyChangedEvent& PropertyChangedEvent) const;
+	UE_API void NotifyAssetChangedExternally() const;
+	UE_API void NotifyStatesChangedExternally(const TSet<UStateTreeState*>& ChangedStates, const FPropertyChangedEvent& PropertyChangedEvent) const;
 
 	// Debugging
 #if WITH_STATETREE_TRACE_DEBUGGER
-	bool CanAddStateBreakpoint(EStateTreeBreakpointType Type) const;
-	bool CanRemoveStateBreakpoint(EStateTreeBreakpointType Type) const;
-	ECheckBoxState GetStateBreakpointCheckState(EStateTreeBreakpointType Type) const;
-	void HandleEnableStateBreakpoint(EStateTreeBreakpointType Type);
+	UE_API bool CanAddStateBreakpoint(EStateTreeBreakpointType Type) const;
+	UE_API bool CanRemoveStateBreakpoint(EStateTreeBreakpointType Type) const;
+	UE_API ECheckBoxState GetStateBreakpointCheckState(EStateTreeBreakpointType Type) const;
+	UE_API void HandleEnableStateBreakpoint(EStateTreeBreakpointType Type);
 
-	UStateTreeState* FindStateAssociatedToBreakpoint(FStateTreeDebuggerBreakpoint Breakpoint) const;
+	UE_API UStateTreeState* FindStateAssociatedToBreakpoint(FStateTreeDebuggerBreakpoint Breakpoint) const;
 
 	TSharedRef<FStateTreeDebugger> GetDebugger() const { return Debugger; }
-	void RefreshDebuggerBreakpoints();
+	UE_API void RefreshDebuggerBreakpoints();
 #endif // WITH_STATETREE_TRACE_DEBUGGER
 
-	bool IsStateActiveInDebugger(const UStateTreeState& State) const;
+	UE_API bool IsStateActiveInDebugger(const UStateTreeState& State) const;
 
 	// Called when the whole asset is updated (i.e. undo/redo).
 	FOnAssetChanged& GetOnAssetChanged()
@@ -183,22 +185,22 @@ public:
 	}
 	 
 protected:
-	void GetExpandedStatesRecursive(UStateTreeState* State, TSet<TWeakObjectPtr<UStateTreeState>>& ExpandedStates);
+	UE_API void GetExpandedStatesRecursive(UStateTreeState* State, TSet<TWeakObjectPtr<UStateTreeState>>& ExpandedStates);
 
-	void MoveSelectedStates(UStateTreeState* TargetState, const EStateTreeViewModelInsert RelativeLocation);
+	UE_API void MoveSelectedStates(UStateTreeState* TargetState, const EStateTreeViewModelInsert RelativeLocation);
 
-	void PasteStatesAsChildrenFromText(const FString& TextToImport, UStateTreeState* ParentState, const int32 IndexToInsertAt);
+	UE_API void PasteStatesAsChildrenFromText(const FString& TextToImport, UStateTreeState* ParentState, const int32 IndexToInsertAt);
 
-	void HandleIdentifierChanged(const UStateTree& StateTree) const;
+	UE_API void HandleIdentifierChanged(const UStateTree& StateTree) const;
 	
-	void BindToDebuggerDelegates();
+	UE_API void BindToDebuggerDelegates();
 	
 	TWeakObjectPtr<UStateTreeEditorData> TreeDataWeak;
 	TSet<TWeakObjectPtr<UStateTreeState>> SelectedStates;
 
 #if WITH_STATETREE_TRACE_DEBUGGER
-	void HandleBreakpointsChanged(const UStateTree& StateTree);
-	void HandlePostCompile(const UStateTree& StateTree);
+	UE_API void HandleBreakpointsChanged(const UStateTree& StateTree);
+	UE_API void HandlePostCompile(const UStateTree& StateTree);
 
 	TSharedRef<FStateTreeDebugger> Debugger;
 	TArray<FGuid> ActiveStates;
@@ -215,11 +217,13 @@ protected:
 };
 
 /** Helper class to allow to copy bindings into clipboard. */
-UCLASS(Hidden)
-class STATETREEEDITORMODULE_API UStateTreeClipboardBindings : public UObject
+UCLASS(MinimalAPI, Hidden)
+class UStateTreeClipboardBindings : public UObject
 {
 	GENERATED_BODY()
 public:
 	UPROPERTY()
 	TArray<FStateTreePropertyPathBinding> Bindings;
-};
\ No newline at end of file
+};
+
+#undef UE_API
diff --git a/Engine/Plugins/Runtime/USDCore/Source/USDClasses/Public/USDAssetCache.h b/Engine/Plugins/Runtime/USDCore/Source/USDClasses/Public/USDAssetCache.h
--- a/Engine/Plugins/Runtime/USDCore/Source/USDClasses/Public/USDAssetCache.h
+++ b/Engine/Plugins/Runtime/USDCore/Source/USDClasses/Public/USDAssetCache.h
@@ -6,25 +6,27 @@
 
 #include "USDAssetCache.generated.h"
 
+#define UE_API USDCLASSES_API
+
 class UE_DEPRECATED(5.2, "Use the UUsdAssetCache2 class instead, which is an actual asset you can create on the Content Browser") UUsdAssetCache;
 
 /** Owns the assets generated and reused by the USD stage, allowing thread-safe retrieval/storage */
 UCLASS(meta = (ScriptName = "UsdAssetCache_Deprecated"))
 
-class USDCLASSES_API UUsdAssetCache : public UObject
+class UUsdAssetCache : public UObject
 {
 	GENERATED_BODY()
 
 public:
-	UUsdAssetCache();
+	UE_API UUsdAssetCache();
 
 #if WITH_EDITOR
-	virtual void PostEditChangeProperty(struct FPropertyChangedEvent& PropertyChangedEvent) override;
+	UE_API virtual void PostEditChangeProperty(struct FPropertyChangedEvent& PropertyChangedEvent) override;
 #endif
 
-	void CacheAsset(const FString& Hash, UObject* Asset, const FString& PrimPath = FString());
-	void DiscardAsset(const FString& Hash);
-	UObject* GetCachedAsset(const FString& Hash) const;
+	UE_API void CacheAsset(const FString& Hash, UObject* Asset, const FString& PrimPath = FString());
+	UE_API void DiscardAsset(const FString& Hash);
+	UE_API UObject* GetCachedAsset(const FString& Hash) const;
 
 	TMap<FString, TObjectPtr<UObject>> GetCachedAssets() const
 	{
@@ -34,18 +36,18 @@ public:
 		return CachedAssets;
 	}
 
-	void LinkAssetToPrim(const FString& PrimPath, UObject* Asset);
-	void RemoveAssetPrimLink(const FString& PrimPath);
-	UObject* GetAssetForPrim(const FString& PrimPath) const;
+	UE_API void LinkAssetToPrim(const FString& PrimPath, UObject* Asset);
+	UE_API void RemoveAssetPrimLink(const FString& PrimPath);
+	UE_API UObject* GetAssetForPrim(const FString& PrimPath) const;
 
 	/** Returns the first prim path associated with an asset (see TMap::FindKey) */
-	FString GetPrimForAsset(UObject* Asset) const;
+	UE_API FString GetPrimForAsset(UObject* Asset) const;
 
 	/**
 	 * Returns the hash used as a key for an asset, in case we own it.
 	 * Returns the empty string otherwise.
 	 */
-	FString GetHashForAsset(UObject* Asset) const;
+	UE_API FString GetHashForAsset(UObject* Asset) const;
 
 	TMap<FString, TWeakObjectPtr<UObject>> GetAssetPrimLinks() const
 	{
@@ -62,7 +64,7 @@ public:
 		return TransientStorage.Num() + PersistentStorage.Num();
 	}
 
-	void Reset();
+	UE_API void Reset();
 
 	/**
 	 * Every time an asset is retrieved/inserted we add it to ActiveAssets. With these functions that container can be reset/returned.
@@ -72,13 +74,13 @@ public:
 	 * The USDStageImporter will then use these functions before/after translating the scene to know which items from the
 	 * cache that are actually used for the current scene.
 	 */
-	void MarkAssetsAsStale();
+	UE_API void MarkAssetsAsStale();
 
 	/** Returns assets that aren't marked as stale */
-	TSet<UObject*> GetActiveAssets() const;
+	UE_API TSet<UObject*> GetActiveAssets() const;
 
 	// We need to be serializable so that AUsdStageActor can duplicate us for PIE
-	virtual void Serialize(FArchive& Ar) override;
+	UE_API virtual void Serialize(FArchive& Ar) override;
 
 private:
 	// For now everything is NonPIEDuplicateTransient as this is mostly a subobject of AUsdStageActor. When the actor is duplicated
@@ -107,3 +109,5 @@ private:
 	// Used to ensure thread safety
 	mutable FCriticalSection CriticalSection;
 };
+
+#undef UE_API
diff --git a/Engine/Plugins/Runtime/USDCore/Source/USDClasses/Public/USDAssetCache2.h b/Engine/Plugins/Runtime/USDCore/Source/USDClasses/Public/USDAssetCache2.h
--- a/Engine/Plugins/Runtime/USDCore/Source/USDClasses/Public/USDAssetCache2.h
+++ b/Engine/Plugins/Runtime/USDCore/Source/USDClasses/Public/USDAssetCache2.h
@@ -9,11 +9,13 @@
 
 #include "USDAssetCache2.generated.h"
 
+#define UE_API USDCLASSES_API
+
 /**
  * Owns the assets generated and reused by USD Stages, allowing thread-safe retrieval/storage.
  */
-UCLASS(BlueprintType, Blueprintable)
-class USDCLASSES_API UUsdAssetCache2 : public UObject
+UCLASS(MinimalAPI, BlueprintType, Blueprintable)
+class UUsdAssetCache2 : public UObject
 {
 	GENERATED_BODY()
 
@@ -42,14 +44,14 @@ public:
 public:
 	/** Adds an asset to the cache attached to a particular hash, and optionally registering a referencer */
 	UFUNCTION(BlueprintCallable, Category = "Caching", meta = (CallInEditor = "true"))
-	void CacheAsset(const FString& Hash, UObject* Asset, const UObject* Referencer = nullptr);
+	UE_API void CacheAsset(const FString& Hash, UObject* Asset, const UObject* Referencer = nullptr);
 
 	/**
 	 * Returns true if the asset with the given hash can be removed from the cache. It will return false in case the
 	 * asset is still being used, either by another consumer asset or directly by some referencer.
 	 */
 	UFUNCTION(BlueprintCallable, Category = "Caching", meta = (CallInEditor = "true"))
-	bool CanRemoveAsset(const FString& Hash);
+	UE_API bool CanRemoveAsset(const FString& Hash);
 
 	/**
 	 * If an asset is associated with `Hash`, it will be returned and the asset cache will stop tracking this asset
@@ -59,7 +61,7 @@ public:
 	 * placing the asset at a new Outer object.
 	 */
 	UFUNCTION(BlueprintCallable, Category = "Caching", meta = (CallInEditor = "true"))
-	UObject* RemoveAsset(const FString& Hash);
+	UE_API UObject* RemoveAsset(const FString& Hash);
 
 	/**
 	 * Returns an asset associated with a particular `Hash`.
@@ -67,14 +69,14 @@ public:
 	 * asset to be read from the asset cache's file on disk.
 	 */
 	UFUNCTION(BlueprintCallable, Category = "Caching", meta = (CallInEditor = "true"))
-	UObject* GetCachedAsset(const FString& Hash);
+	UE_API UObject* GetCachedAsset(const FString& Hash);
 
 	/**
 	 * Marks the provided asset as being used at this point, optionally adding a specific referencer.
 	 * This is useful because the asset cache will always prioritize retaining the most recently used assets
 	 */
 	UFUNCTION(BlueprintCallable, Category = "Caching", meta = (CallInEditor = "true"))
-	bool TouchAsset(const UObject* Asset, const UObject* Referencer = nullptr);
+	UE_API bool TouchAsset(const UObject* Asset, const UObject* Referencer = nullptr);
 
 	/**
 	 * Adds a new UObject referencer to a particular asset, returning true if the operation succeeded.
@@ -84,101 +86,101 @@ public:
 	 * pointers to them.
 	 */
 	UFUNCTION(BlueprintCallable, Category = "Caching", meta = (CallInEditor = "true"))
-	bool AddAssetReference(const UObject* Asset, const UObject* Referencer);
+	UE_API bool AddAssetReference(const UObject* Asset, const UObject* Referencer);
 
 	/**
 	 * Removes an UObject referencer from a particular asset, returning true if the operation succeeded.
 	 * If no specific Referencer is provided, all referencers to Asset will be removed.
 	 */
 	UFUNCTION(BlueprintCallable, Category = "Caching", meta = (CallInEditor = "true"))
-	bool RemoveAssetReference(const UObject* Asset, const UObject* Referencer = nullptr);
+	UE_API bool RemoveAssetReference(const UObject* Asset, const UObject* Referencer = nullptr);
 
 	/**
 	 * Removes the particular referencer to all assets tracked by the cache, if it was a referencer to any of them.
 	 * Returns true if at least one asset referencer count was altered by this.
 	 */
 	UFUNCTION(BlueprintCallable, Category = "Caching", meta = (CallInEditor = "true"))
-	bool RemoveAllAssetReferences(const UObject* Referencer);
+	UE_API bool RemoveAllAssetReferences(const UObject* Referencer);
 
 	/**
 	 * Returns the hash associated with an asset, in case we own it. Returns the empty string otherwise.
 	 * Note: This has O(1) time complexity.
 	 */
 	UFUNCTION(BlueprintCallable, Category = "USD", meta = (CallInEditor = "true"))
-	FString GetHashForAsset(const UObject* Asset) const;
+	UE_API FString GetHashForAsset(const UObject* Asset) const;
 
 	/**
 	 * Returns true in case the asset at `AssetPath` is tracked by the cache in any way (persistent asset,
 	 * unreferenced or referenced). An example AssetPath would be "/Game/MyTextures/RedBrick.RedBrick".
 	 */
 	UFUNCTION(BlueprintCallable, Category = "USD", meta = (CallInEditor = "true"))
-	bool IsAssetOwnedByCache(const FString& AssetPath) const;
+	UE_API bool IsAssetOwnedByCache(const FString& AssetPath) const;
 
 	/**
 	 * Returns how many assets are tracked by the asset cache in total (summing up persistent, referenced and
 	 * unreferenced storage)
 	 */
 	UFUNCTION(BlueprintCallable, Category = "USD", meta = (CallInEditor = "true"))
-	int32 GetNumAssets() const;
+	UE_API int32 GetNumAssets() const;
 
 	/**
 	 * Returns all asset hashes tracked by the asset cache, for all storage types. This includes unloaded
 	 * persistent assets
 	 */
 	UFUNCTION(BlueprintCallable, Category = "USD", meta = (CallInEditor = "true"))
-	TArray<FString> GetAllAssetHashes() const;
+	UE_API TArray<FString> GetAllAssetHashes() const;
 
 	/**
 	 * Returns all assets that are currently loaded in the asset cache.
 	 * This will not include persistent assets that haven't been loaded yet.
 	 */
 	UFUNCTION(BlueprintCallable, Category = "USD", meta = (CallInEditor = "true"))
-	TArray<UObject*> GetAllLoadedAssets() const;
+	UE_API TArray<UObject*> GetAllLoadedAssets() const;
 
 	/**
 	 * Returns all asset paths tracked by the asset cache, for all storage types. (e.g.
 	 * "/Game/MyTextures/RedBrick.RedBrick"). This includes unloaded persistent assets
 	 */
 	UFUNCTION(BlueprintCallable, Category = "USD", meta = (CallInEditor = "true"))
-	TArray<FString> GetAllCachedAssetPaths() const;
+	UE_API TArray<FString> GetAllCachedAssetPaths() const;
 
 	/** Discards all tracked assets across all storage types */
 	UFUNCTION(BlueprintCallable, Category = "Settings", meta = (CallInEditor = "true"))
-	void Reset();
+	UE_API void Reset();
 
 	/**
 	 * Updates which assets belong to each storage type. You must call this in case you perform direct operations on
 	 * the asset cache, after those operations are fully complete.
 	 */
 	UFUNCTION(BlueprintCallable, Category = "Settings", meta = (CallInEditor = "true"))
-	void RefreshStorage();
+	UE_API void RefreshStorage();
 
 public:
-	UUsdAssetCache2();
+	UE_API UUsdAssetCache2();
 
 	/**
 	 * Every time an asset is retrieved/inserted we add it to ActiveAssets.
 	 * With these functions that container can be reset/returned.
 	 */
-	void MarkAssetsAsStale();
+	UE_API void MarkAssetsAsStale();
 
 	/** Returns assets that aren't marked as stale */
-	TSet<UObject*> GetActiveAssets() const;
+	UE_API TSet<UObject*> GetActiveAssets() const;
 
 	// Begin UObject interfacte
 #if WITH_EDITOR
-	virtual void PostEditChangeProperty(struct FPropertyChangedEvent& PropertyChangedEvent) override;
+	UE_API virtual void PostEditChangeProperty(struct FPropertyChangedEvent& PropertyChangedEvent) override;
 #endif
-	virtual void Serialize(FArchive& Ar) override;
-	virtual void PostLoad() override;
+	UE_API virtual void Serialize(FArchive& Ar) override;
+	UE_API virtual void PostLoad() override;
 	// End UObject interfacte
 
 private:
 	// Tries unloading unreferenced, persistent assets back to bulkdata on disk if we can
-	bool TryUnloadAsset(FCachedAssetInfo& InOutInfo);
-	bool IsAssetOwnedByCacheInternal(const FString& AssetPath) const;
-	bool AddAssetReferenceInternal(const UObject* Asset, const UObject* Referencer);
-	void TouchAssetInternal(const UObject* Asset, const UObject* Referencer = nullptr);
+	UE_API bool TryUnloadAsset(FCachedAssetInfo& InOutInfo);
+	UE_API bool IsAssetOwnedByCacheInternal(const FString& AssetPath) const;
+	UE_API bool AddAssetReferenceInternal(const UObject* Asset, const UObject* Referencer);
+	UE_API void TouchAssetInternal(const UObject* Asset, const UObject* Referencer = nullptr);
 
 private:
 	friend struct FUsdScopedAssetCacheReferencer;
@@ -261,11 +263,11 @@ private:
  *     FUsdScopedAssetCacheReferencer ScopedReferencer{ StageActor->UsdAssetCache, StageActor };
  *     StageActor->SetRootLayer(MyRootLayer);  // Adds assets to the cache
  */
-struct USDCLASSES_API FUsdScopedAssetCacheReferencer
+struct FUsdScopedAssetCacheReferencer
 {
 public:
-	explicit FUsdScopedAssetCacheReferencer(UUsdAssetCache2* InAssetCache, const UObject* Referencer);
-	~FUsdScopedAssetCacheReferencer();
+	UE_API explicit FUsdScopedAssetCacheReferencer(UUsdAssetCache2* InAssetCache, const UObject* Referencer);
+	UE_API ~FUsdScopedAssetCacheReferencer();
 
 	FUsdScopedAssetCacheReferencer() = delete;
 	FUsdScopedAssetCacheReferencer(const FUsdScopedAssetCacheReferencer&) = delete;
@@ -276,3 +278,5 @@ public:
 private:
 	TWeakObjectPtr<UUsdAssetCache2> AssetCache;
 };
+
+#undef UE_API
diff --git a/Engine/Plugins/Runtime/USDCore/Source/USDClasses/Public/USDAssetCache3.h b/Engine/Plugins/Runtime/USDCore/Source/USDClasses/Public/USDAssetCache3.h
--- a/Engine/Plugins/Runtime/USDCore/Source/USDClasses/Public/USDAssetCache3.h
+++ b/Engine/Plugins/Runtime/USDCore/Source/USDClasses/Public/USDAssetCache3.h
@@ -8,6 +8,8 @@
 
 #include "USDAssetCache3.generated.h"
 
+#define UE_API USDCLASSES_API
+
 /**
  * This class is an asset that can be created via the Content Browser and assigned to AUsdStageActors.
  *
@@ -20,8 +22,8 @@
  * A "Default Asset Cache" can be set on the project settings, and will be automatically used for any AUsdStageActor that hasn't
  * had an asset cache manually set beforehand.
  */
-UCLASS(BlueprintType, Blueprintable, meta = (DisplayName = "USD Asset Cache", ScriptName = "UsdAssetCache"))
-class USDCLASSES_API UUsdAssetCache3 : public UObject
+UCLASS(MinimalAPI, BlueprintType, Blueprintable, meta = (DisplayName = "USD Asset Cache", ScriptName = "UsdAssetCache"))
+class UUsdAssetCache3 : public UObject
 {
 	GENERATED_BODY()
 
@@ -43,7 +45,7 @@ public:
 	 * @return The asset that was returned or created
 	 */
 	UFUNCTION(BlueprintCallable, Category = "USD", meta = (CallInEditor = "true"))
-	UObject* GetOrCreateCachedAsset(
+	UE_API UObject* GetOrCreateCachedAsset(
 		const FString& Hash,
 		UClass* Class,
 		const FString& DesiredName,
@@ -100,7 +102,7 @@ public:
 	 * @param Referencer - The asset will not be deleted or untracked until this referencer is removed (via any of the RemoveAssetReferencer member functions)
 	 * @return The asset that was returned or created
 	 */
-	UObject* GetOrCreateCustomCachedAsset(
+	UE_API UObject* GetOrCreateCustomCachedAsset(
 		const FString& Hash,
 		UClass* Class,
 		const FString& DesiredName,
@@ -140,7 +142,7 @@ public:
 
 	/** Adds an existing asset to the cache attached to a particular hash, and optionally registering a referencer */
 	UFUNCTION(BlueprintCallable, Category = "USD", meta = (CallInEditor = "true"))
-	void CacheAsset(const FString& Hash, const FSoftObjectPath& AssetPath, const UObject* Referencer = nullptr);
+	UE_API void CacheAsset(const FString& Hash, const FSoftObjectPath& AssetPath, const UObject* Referencer = nullptr);
 
 	/**
 	 * Removes all info about the asset associated with Hash from this cache, if there is any.
@@ -148,7 +150,7 @@ public:
 	 * and only when manually created by it or if flagged with SetAssetDeletable
 	 */
 	UFUNCTION(BlueprintCallable, Category = "USD", meta = (CallInEditor = "true"))
-	FSoftObjectPath StopTrackingAsset(const FString& Hash);
+	UE_API FSoftObjectPath StopTrackingAsset(const FString& Hash);
 
 	/**
 	 * Returns the asset associated with a particular Hash, if any. Returns nullptr if there isn't any
@@ -157,7 +159,7 @@ public:
 	 * WARNING: As this may try loading a package from disk, this can only be called from the game thread!
 	 */
 	UFUNCTION(BlueprintCallable, Category = "USD", meta = (CallInEditor = "true"))
-	UObject* GetCachedAsset(const FString& Hash) const;
+	UE_API UObject* GetCachedAsset(const FString& Hash) const;
 
 	/**
 	 * Templated version of GetCachedAsset for convenience
@@ -175,32 +177,32 @@ public:
 	 * If there is no asset associated with Hash, will return an invalid (empty) FSoftObjectPath.
 	 */
 	UFUNCTION(BlueprintCallable, Category = "USD", meta = (CallInEditor = "true"))
-	FSoftObjectPath GetCachedAssetPath(const FString& Hash) const;
+	UE_API FSoftObjectPath GetCachedAssetPath(const FString& Hash) const;
 
 	/**
 	 * Returns the hash associated with a particular asset, or the empty string if there isn't any.
 	 * Note: The asset cache keeps internal reverse maps, so this should be O(1)
 	 */
 	UFUNCTION(BlueprintCallable, Category = "USD", meta = (CallInEditor = "true"))
-	FString GetHashForAsset(const FSoftObjectPath& AssetPath) const;
+	UE_API FString GetHashForAsset(const FSoftObjectPath& AssetPath) const;
 
 	/**
 	 * Returns true if this asset is currently tracked by the asset cache's main hash to asset maps
 	 */
 	UFUNCTION(BlueprintCallable, Category = "USD", meta = (CallInEditor = "true"))
-	bool IsAssetTrackedByCache(const FSoftObjectPath& AssetPath) const;
+	UE_API bool IsAssetTrackedByCache(const FSoftObjectPath& AssetPath) const;
 
 	/**
 	 * Returns the total number of cached asset paths, whether these resolve to assets or not
 	 */
 	UFUNCTION(BlueprintCallable, Category = "USD", meta = (CallInEditor = "true"))
-	int32 GetNumAssets() const;
+	UE_API int32 GetNumAssets() const;
 
 	/**
 	 * Returns a copy of the internal mapping between hashes and asset paths
 	 */
 	UFUNCTION(BlueprintCallable, Category = "USD", meta = (CallInEditor = "true"))
-	TMap<FString, FSoftObjectPath> GetAllTrackedAssets() const;
+	UE_API TMap<FString, FSoftObjectPath> GetAllTrackedAssets() const;
 
 	/**
 	 * The same as GetAllTrackedAssets, except that it will automatically try loading all the asset paths before
@@ -209,41 +211,41 @@ public:
 	 * WARNING: As this may try loading a package from disk, this can only be called from the game thread!
 	 */
 	UFUNCTION(BlueprintCallable, Category = "USD", meta = (CallInEditor = "true"))
-	TMap<FString, UObject*> LoadAndGetAllTrackedAssets() const;
+	UE_API TMap<FString, UObject*> LoadAndGetAllTrackedAssets() const;
 
 	/**
 	 * Adds a new UObject referencer to a particular asset, returning true if the operation succeeded.
 	 * Assets will not be deleted or untracked by the asset cache while the referencer is registered.
 	 */
 	UFUNCTION(BlueprintCallable, Category = "Caching", meta = (CallInEditor = "true"))
-	bool AddAssetReferencer(const UObject* Asset, const UObject* Referencer);
+	UE_API bool AddAssetReferencer(const UObject* Asset, const UObject* Referencer);
 
 	/**
 	 * Removes an UObject referencer from a particular asset, returning true if anything was removed.
 	 * Will do nothing in case Asset or Referencer are invalid.
 	 */
 	UFUNCTION(BlueprintCallable, Category = "USD", meta = (CallInEditor = "true"))
-	bool RemoveAssetReferencer(const UObject* Asset, const UObject* Referencer);
+	UE_API bool RemoveAssetReferencer(const UObject* Asset, const UObject* Referencer);
 
 	/**
 	 * Removes all UObject referencers from a particular asset, returning true if anything was removed.
 	 * Will do nothing in case Asset is invalid.
 	 */
 	UFUNCTION(BlueprintCallable, Category = "USD", meta = (CallInEditor = "true"))
-	bool RemoveAllReferencersForAsset(const UObject* Asset);
+	UE_API bool RemoveAllReferencersForAsset(const UObject* Asset);
 
 	/**
 	 * Removes a particular UObject referencer from all tracked assets, returning true if anything was removed.
 	 * Will do nothing in case Referencer is invalid.
 	 */
 	UFUNCTION(BlueprintCallable, Category = "USD", meta = (CallInEditor = "true"))
-	bool RemoveAllReferencerAssets(const UObject* Referencer);
+	UE_API bool RemoveAllReferencerAssets(const UObject* Referencer);
 
 	/**
 	 * Removes all UObject referencer from all tracked assets, returning true if anything was removed.
 	 */
 	UFUNCTION(BlueprintCallable, Category = "USD", meta = (CallInEditor = "true"))
-	bool RemoveAllAssetReferencers();
+	UE_API bool RemoveAllAssetReferencers();
 
 	/**
 	 * Sets a particular asset as deletable or not.
@@ -252,11 +254,11 @@ public:
 	 * set as deletable.
 	 */
 	UFUNCTION(BlueprintCallable, Category = "USD", meta = (CallInEditor = "true"))
-	void SetAssetDeletable(const UObject* Asset, bool bIsDeletable);
+	UE_API void SetAssetDeletable(const UObject* Asset, bool bIsDeletable);
 
 	/** Returns whether a particular asset is currently marked as deletable or not */
 	UFUNCTION(BlueprintCallable, Category = "USD", meta = (CallInEditor = "true"))
-	bool IsAssetDeletable(const UObject* Asset) const;
+	UE_API bool IsAssetDeletable(const UObject* Asset) const;
 
 	/**
 	 * Deletes all assets that:
@@ -272,7 +274,7 @@ public:
 	 * WARNING: This will clear the undo buffer (i.e. transaction history) and run garbage collection after deleting.
 	 */
 	UFUNCTION(BlueprintCallable, Category = "Cleanup", meta = (CallInEditor = "true"))
-	void DeleteUnreferencedAssets(const bool bShowConfirmation = false);
+	UE_API void DeleteUnreferencedAssets(const bool bShowConfirmation = false);
 
 	/**
 	 * This is the same as calling DeleteUnreferencedAssets and providing true for bShowConfirmation.
@@ -282,7 +284,7 @@ public:
 	 * WARNING: This will clear the undo buffer (i.e. transaction history) and run garbage collection after deleting.
 	 */
 	UFUNCTION(BlueprintCallable, Category = "Cleanup", meta = (CallInEditor = "true", DisplayName = "Delete Unreferenced Assets"))
-	void DeleteUnreferencedAssetsWithConfirmation();
+	UE_API void DeleteUnreferencedAssetsWithConfirmation();
 
 	/**
 	 * Checks the current AssetDirectory for any new assets that were generated from USD, and automatically caches
@@ -292,7 +294,7 @@ public:
 	 * is associated with a hash that is already in use, it will be ignored)
 	 */
 	UFUNCTION(BlueprintCallable, Category = "Storage", meta = (CallInEditor = "true"))
-	void RescanAssetDirectory();
+	UE_API void RescanAssetDirectory();
 
 public:
 	/**
@@ -306,11 +308,11 @@ public:
 	 *     FUsdScopedReferencer ScopedReferencer{ StageActor->UsdAssetCache, StageActor };
 	 *     StageActor->SetRootLayer(MyRootLayer);  // Adds assets to the cache
 	 */
-	struct USDCLASSES_API FUsdScopedReferencer
+	struct FUsdScopedReferencer
 	{
 	public:
-		explicit FUsdScopedReferencer(UUsdAssetCache3* InAssetCache, const UObject* Referencer);
-		~FUsdScopedReferencer();
+		UE_API explicit FUsdScopedReferencer(UUsdAssetCache3* InAssetCache, const UObject* Referencer);
+		UE_API ~FUsdScopedReferencer();
 
 		FUsdScopedReferencer() = delete;
 		FUsdScopedReferencer(const FUsdScopedReferencer&) = delete;
@@ -324,23 +326,23 @@ public:
 	};
 
 public:
-	UUsdAssetCache3();
+	UE_API UUsdAssetCache3();
 
-	virtual void PostLoad() override;
-	virtual void BeginDestroy() override;
-	virtual void Serialize(FArchive& Ar) override;
+	UE_API virtual void PostLoad() override;
+	UE_API virtual void BeginDestroy() override;
+	UE_API virtual void Serialize(FArchive& Ar) override;
 #if WITH_EDITOR
-	virtual void PostEditChangeProperty(FPropertyChangedEvent& PropertyChangedEvent) override;
+	UE_API virtual void PostEditChangeProperty(FPropertyChangedEvent& PropertyChangedEvent) override;
 #endif	  // WITH_EDITOR
 
-	void OnRegistryAssetRenamed(const FAssetData& NewAssetData, const FString& OldName);
+	UE_API void OnRegistryAssetRenamed(const FAssetData& NewAssetData, const FString& OldName);
 
-	void RequestDelayedAssetAutoCleanup();
-	void TouchAsset(const FString& Hash, const UObject* Referencer = nullptr);
-	void TouchAssetPath(const FSoftObjectPath& AssetPath, const UObject* Referencer = nullptr);
-	void MarkAssetsAsStale();
-	TSet<FSoftObjectPath> GetActiveAssets() const;
-	const UObject* SetCurrentScopedReferencer(const UObject* NewReferencer);
+	UE_API void RequestDelayedAssetAutoCleanup();
+	UE_API void TouchAsset(const FString& Hash, const UObject* Referencer = nullptr);
+	UE_API void TouchAssetPath(const FSoftObjectPath& AssetPath, const UObject* Referencer = nullptr);
+	UE_API void MarkAssetsAsStale();
+	UE_API TSet<FSoftObjectPath> GetActiveAssets() const;
+	UE_API const UObject* SetCurrentScopedReferencer(const UObject* NewReferencer);
 
 public:
 	/**
@@ -390,13 +392,13 @@ public:
 	bool bCleanUpUnreferencedAssets = false;
 
 protected:
-	void AddReferenceInternal(const FString& Hash, const UObject* Referencer);
-	FSoftObjectPath StopTrackingAssetInternal(const FString& Hash);
-	bool RemoveAllAssetReferencersInternal(const FString& Hash);
-	void TouchAssetInternal(const FSoftObjectPath& AssetPath, const UObject* Referencer = nullptr);
-	void TryCachingAssetFromAssetUserData(const FAssetData& NewAssetData);
-	bool IsTransientCache();
-	void ForceValidAssetDirectoryInternal(bool bEmitWarning = true);
+	UE_API void AddReferenceInternal(const FString& Hash, const UObject* Referencer);
+	UE_API FSoftObjectPath StopTrackingAssetInternal(const FString& Hash);
+	UE_API bool RemoveAllAssetReferencersInternal(const FString& Hash);
+	UE_API void TouchAssetInternal(const FSoftObjectPath& AssetPath, const UObject* Referencer = nullptr);
+	UE_API void TryCachingAssetFromAssetUserData(const FAssetData& NewAssetData);
+	UE_API bool IsTransientCache();
+	UE_API void ForceValidAssetDirectoryInternal(bool bEmitWarning = true);
 
 private:
 	// Reverse map, to speed up queries like GetHashForAsset and IsAssetTrackedByCache. Should always match HashToAssetPaths
@@ -425,3 +427,5 @@ private:
 	// by the asset cache are actually used by the stage it is about to import, and then publish exactly those.
 	mutable TSet<FSoftObjectPath> ActiveAssets;
 };
+
+#undef UE_API
diff --git a/Engine/Plugins/Runtime/USDCore/Source/USDClasses/Public/USDAssetImportData.h b/Engine/Plugins/Runtime/USDCore/Source/USDClasses/Public/USDAssetImportData.h
--- a/Engine/Plugins/Runtime/USDCore/Source/USDClasses/Public/USDAssetImportData.h
+++ b/Engine/Plugins/Runtime/USDCore/Source/USDClasses/Public/USDAssetImportData.h
@@ -8,8 +8,8 @@
 
 #include "USDAssetImportData.generated.h"
 
-UCLASS()
-class USDCLASSES_API UUsdAssetImportData : public UAssetImportData
+UCLASS(MinimalAPI)
+class UUsdAssetImportData : public UAssetImportData
 {
 	GENERATED_BODY()
 
