[StateTree] Added OnEventConsumed debug event.

diff --git a/Engine/Plugins/Runtime/StateTree/Source/StateTreeModule/Private/Debugger/StateTreeDebug.cpp b/Engine/Plugins/Runtime/StateTree/Source/StateTreeModule/Private/Debugger/StateTreeDebug.cpp
--- a/Engine/Plugins/Runtime/StateTree/Source/StateTreeModule/Private/Debugger/StateTreeDebug.cpp
+++ b/Engine/Plugins/Runtime/StateTree/Source/StateTreeModule/Private/Debugger/StateTreeDebug.cpp
@@ -26,6 +26,7 @@ FNodeDelegate OnTaskEnterState_AnyThread;
 FNodeDelegate OnTickTask_AnyThread;
 FNodeDelegate OnTaskExitState_AnyThread;
 FEventSentDelegate OnEventSent_AnyThread;
+FEventConsumedDelegate OnEventConsumed_AnyThread;
 
 FNodeReference::FNodeReference(TNotNull<const UStateTree*> InStateTree, const FStateTreeIndex16 InNodeIndex)
 	: StateTree(InStateTree)
@@ -155,6 +156,11 @@ void EventSent(const FStateTreeMinimalExecutionContext& ExecutionContext, TNotNu
 	OnEventSent_AnyThread.Broadcast(ExecutionContext, FEventSentDelegateArgs{ .StateTree = StateTree, .Tag = Tag, .Payload = Payload, .Origin = Origin });
 }
 
+void EventConsumed(const FStateTreeExecutionContext& ExecutionContext, const FStateTreeSharedEvent& Event)
+{
+	OnEventConsumed_AnyThread.Broadcast(ExecutionContext, Event);
+}
+
 }//namespace UE::StateTree::Debug
 
 #endif // WITH_STATETREE_DEBUG
diff --git a/Engine/Plugins/Runtime/StateTree/Source/StateTreeModule/Private/StateTreeEvents.cpp b/Engine/Plugins/Runtime/StateTree/Source/StateTreeModule/Private/StateTreeEvents.cpp
--- a/Engine/Plugins/Runtime/StateTree/Source/StateTreeModule/Private/StateTreeEvents.cpp
+++ b/Engine/Plugins/Runtime/StateTree/Source/StateTreeModule/Private/StateTreeEvents.cpp
@@ -40,10 +40,11 @@ void FStateTreeEventQueue::SendEvent(const UObject* Owner, const FGameplayTag& T
 	SharedEvents.Emplace(Tag, Payload, Origin);
 }
 
-void FStateTreeEventQueue::ConsumeEvent(const FStateTreeSharedEvent& Event)
+bool FStateTreeEventQueue::ConsumeEvent(const FStateTreeSharedEvent& Event)
 {
-	SharedEvents.RemoveAll([&EventToRemove = Event](const FStateTreeSharedEvent& Event)
+	const int32 Count = SharedEvents.RemoveAllSwap([&EventToRemove = Event](const FStateTreeSharedEvent& Event)
 	{
 		return Event == EventToRemove;
 	});
+	return Count > 0;
 }
diff --git a/Engine/Plugins/Runtime/StateTree/Source/StateTreeModule/Private/StateTreeExecutionContext.cpp b/Engine/Plugins/Runtime/StateTree/Source/StateTreeModule/Private/StateTreeExecutionContext.cpp
--- a/Engine/Plugins/Runtime/StateTree/Source/StateTreeModule/Private/StateTreeExecutionContext.cpp
+++ b/Engine/Plugins/Runtime/StateTree/Source/StateTreeModule/Private/StateTreeExecutionContext.cpp
@@ -1686,6 +1686,14 @@ void FStateTreeExecutionContext::UnbindDelegate(const FStateTreeDelegateListener
 	GetExecState().DelegateActiveListeners.Remove(Listener);
 }
 
+void FStateTreeExecutionContext::ConsumeEvent(const FStateTreeSharedEvent& Event)
+{
+	if (EventQueue && EventQueue->ConsumeEvent(Event))
+	{
+		UE_STATETREE_DEBUG_EVENT_CONSUMED(this, Event);
+	}
+}
+
 void FStateTreeExecutionContext::RequestTransition(const FStateTreeTransitionRequest& Request)
 {
 	CSV_SCOPED_TIMING_STAT_EXCLUSIVE(StateTree_RequestTransition);
diff --git a/Engine/Plugins/Runtime/StateTree/Source/StateTreeModule/Public/Debugger/StateTreeDebug.h b/Engine/Plugins/Runtime/StateTree/Source/StateTreeModule/Public/Debugger/StateTreeDebug.h
--- a/Engine/Plugins/Runtime/StateTree/Source/StateTreeModule/Public/Debugger/StateTreeDebug.h
+++ b/Engine/Plugins/Runtime/StateTree/Source/StateTreeModule/Public/Debugger/StateTreeDebug.h
@@ -7,6 +7,7 @@
 class UStateTree;
 struct FStateTreeExecutionContext;
 struct FStateTreeMinimalExecutionContext;
+struct FStateTreeSharedEvent;
 struct FStateTreeTransitionSource;
 enum class EStateTreeUpdatePhase : uint8;
 enum class EStateTreeTraceEventType : uint8;
@@ -26,6 +27,7 @@ enum class EStateTreeTraceEventType : uint8;
 #define UE_STATETREE_DEBUG_TASK_EXIT_STATE(InCtx, InStateTree, InNodeIndex) ::UE::StateTree::Debug::TaskExitState(*(InCtx), ::UE::StateTree::Debug::FNodeReference((InStateTree), FStateTreeIndex16(InNodeIndex)))
 
 #define UE_STATETREE_DEBUG_SEND_EVENT(InCtx, InStateTree, InTag, InPayload, InOrigin) ::UE::StateTree::Debug::EventSent(*(InCtx), InStateTree, InTag, InPayload, InOrigin)
+#define UE_STATETREE_DEBUG_EVENT_CONSUMED(InExecutionContextPtr, Event) ::UE::StateTree::Debug::EventConsumed(*(InExecutionContextPtr), Event)
 
 #define UE_STATETREE_DEBUG_ENTER_PHASE(InCtx, InPhase) \
 		::UE::StateTree::Debug::EnterPhase(*(InCtx), (InPhase), FStateTreeStateHandle::Invalid); \
@@ -108,6 +110,7 @@ namespace UE::StateTree::Debug
 	};
 
 	DECLARE_TS_MULTICAST_DELEGATE_TwoParams(FEventSentDelegate, const FStateTreeMinimalExecutionContext& ExecutionContext, const FEventSentDelegateArgs& EventSentArgs);
+	DECLARE_TS_MULTICAST_DELEGATE_TwoParams(FEventConsumedDelegate, const FStateTreeExecutionContext& ExecutionContext, const FStateTreeSharedEvent& Event);
 	DECLARE_TS_MULTICAST_DELEGATE_TwoParams(FNodeDelegate, const FStateTreeExecutionContext&, FNodeDelegateArgs);
 	DECLARE_TS_MULTICAST_DELEGATE_ThreeParams(FPhaseDelegate, const FStateTreeExecutionContext&, EStateTreeUpdatePhase, FStateTreeStateHandle);
 	DECLARE_TS_MULTICAST_DELEGATE_ThreeParams(FStateDelegate, const FStateTreeExecutionContext&, FStateTreeStateHandle, EStateTreeTraceEventType);
@@ -126,6 +129,7 @@ namespace UE::StateTree::Debug
 	void TaskExitState(const FStateTreeExecutionContext& ExecutionContext, FNodeReference Node);
 
 	void EventSent(const FStateTreeMinimalExecutionContext& ExecutionContext, TNotNull<const UStateTree*> StateTree, FGameplayTag Tag, FConstStructView Payload, FName Origin);
+	void EventConsumed(const FStateTreeExecutionContext& ExecutionContext, const FStateTreeSharedEvent& Event);
 
 	void EnterPhase(const FStateTreeExecutionContext& ExecutionContext, EStateTreeUpdatePhase Phase, FStateTreeStateHandle StateHandle);
 	void ExitPhase(const FStateTreeExecutionContext& ExecutionContext, EStateTreeUpdatePhase Phase, FStateTreeStateHandle StateHandle);
@@ -231,6 +235,13 @@ namespace UE::StateTree::Debug
 	 * @note The StateTree can execute on any thread.
 	 */
 	STATETREEMODULE_API extern FEventSentDelegate OnEventSent_AnyThread;
+
+	/**
+	 * Debugging callback for when an event is consumed.
+	 * @note The callback executes inside the StateTree logic.
+	 * @note The StateTree can execute on any thread.
+	 */
+	STATETREEMODULE_API extern FEventConsumedDelegate OnEventConsumed_AnyThread;
 } //UE::StateTree::Debug
 
 #else
@@ -248,6 +259,7 @@ namespace UE::StateTree::Debug
 #define UE_STATETREE_DEBUG_TASK_EXIT_STATE(...)
 
 #define UE_STATETREE_DEBUG_SEND_EVENT(...)
+#define UE_STATETREE_DEBUG_EVENT_CONSUMED(...)
 
 #define UE_STATETREE_DEBUG_ENTER_PHASE(...)
 #define UE_STATETREE_DEBUG_EXIT_PHASE(...)
diff --git a/Engine/Plugins/Runtime/StateTree/Source/StateTreeModule/Public/StateTreeEvents.h b/Engine/Plugins/Runtime/StateTree/Source/StateTreeModule/Public/StateTreeEvents.h
--- a/Engine/Plugins/Runtime/StateTree/Source/StateTreeModule/Public/StateTreeEvents.h
+++ b/Engine/Plugins/Runtime/StateTree/Source/StateTreeModule/Public/StateTreeEvents.h
@@ -187,8 +187,11 @@ struct FStateTreeEventQueue
 	 */
 	UE_API void SendEvent(const UObject* Owner, const FGameplayTag& Tag, const FConstStructView Payload = FConstStructView(), const FName Origin = FName());
 
-	/** Consumes and removes the specified event from the event queue. */
-	UE_API void ConsumeEvent(const FStateTreeSharedEvent& Event);
+	/**
+	 * Consumes and removes the specified event from the event queue.
+	 * @return true if successfully found and removed the event from the queue.
+	 */
+	UE_API bool ConsumeEvent(const FStateTreeSharedEvent& Event);
 
 	/**
 	 * Iterates over all events.
diff --git a/Engine/Plugins/Runtime/StateTree/Source/StateTreeModule/Public/StateTreeExecutionContext.h b/Engine/Plugins/Runtime/StateTree/Source/StateTreeModule/Public/StateTreeExecutionContext.h
--- a/Engine/Plugins/Runtime/StateTree/Source/StateTreeModule/Public/StateTreeExecutionContext.h
+++ b/Engine/Plugins/Runtime/StateTree/Source/StateTreeModule/Public/StateTreeExecutionContext.h
@@ -524,13 +524,7 @@ PRAGMA_ENABLE_DEPRECATION_WARNINGS
 	}
 
 	/** Consumes and removes the specified event from the event queue. */
-	void ConsumeEvent(const FStateTreeSharedEvent& Event)
-	{
-		if (EventQueue)
-		{
-			EventQueue->ConsumeEvent(Event);
-		}
-	}
+	UE_API void ConsumeEvent(const FStateTreeSharedEvent& Event);
 
 	UE_DEPRECATED(5.5, "Use GetEventsToProcessView() instead.")
 	TConstArrayView<FStateTreeEvent> GetEventsToProcess() const { return {}; }
