[StateTree] Add a default icon for StateTree asset.

#rnx

diff --git a/Engine/Plugins/Runtime/StateTree/Source/StateTreeEditorModule/Private/StateTreeEditorStyle.cpp b/Engine/Plugins/Runtime/StateTree/Source/StateTreeEditorModule/Private/StateTreeEditorStyle.cpp
--- a/Engine/Plugins/Runtime/StateTree/Source/StateTreeEditorModule/Private/StateTreeEditorStyle.cpp
+++ b/Engine/Plugins/Runtime/StateTree/Source/StateTreeEditorModule/Private/StateTreeEditorStyle.cpp
@@ -291,7 +291,6 @@ FStateTreeEditorStyle::FStateTreeEditorStyle()
 			.SetEditableTextBoxStyle(EditableTextBlockStyle));
 	}
 
-	
 	// Command icons
 	{
 		// From generic Engine
@@ -372,6 +371,9 @@ FStateTreeEditorStyle::FStateTreeEditorStyle()
 
 	{
 		// From plugin
+		Set("ClassThumbnail.StateTree", new IMAGE_BRUSH_SVG("Icons/StateTree_64", CoreStyleConstants::Icon64x64));
+		Set("ClassIcon.StateTree", new IMAGE_BRUSH_SVG("Icons/StateTree", CoreStyleConstants::Icon16x16));
+
 		Set("StateTreeEditor.AddSiblingState", new IMAGE_BRUSH_SVG("Icons/Sibling_State", CoreStyleConstants::Icon16x16));
 		Set("StateTreeditor.AddChildState", new IMAGE_BRUSH_SVG("Icons/Child_State", CoreStyleConstants::Icon16x16));
 
