[StateTree] Add a "go to source" for tasks and conditions.
#rnx

diff --git a/Engine/Plugins/Runtime/StateTree/Source/StateTreeEditorModule/Private/Customizations/StateTreeEditorNodeDetails.cpp b/Engine/Plugins/Runtime/StateTree/Source/StateTreeEditorModule/Private/Customizations/StateTreeEditorNodeDetails.cpp
--- a/Engine/Plugins/Runtime/StateTree/Source/StateTreeEditorModule/Private/Customizations/StateTreeEditorNodeDetails.cpp
+++ b/Engine/Plugins/Runtime/StateTree/Source/StateTreeEditorModule/Private/Customizations/StateTreeEditorNodeDetails.cpp
@@ -40,6 +40,7 @@
 #include "Styling/SlateTypes.h"
 #include "TextStyleDecorator.h"
 #include "Framework/Application/SlateApplication.h"
+#include "SourceCodeNavigation.h"
 #include "StateTreeDelegate.h"
 #include "StateTreeScopedEditorDataFixer.h"
 
@@ -607,6 +608,28 @@ void FStateTreeEditorNodeDetails::CustomizeHeader(TSharedRef<class IPropertyHand
 						UE::StateTreeEditor::DebuggerExtensions::CreateEditorNodeWidget(StructPropertyHandle, EditorData.Get())
 					]
 
+					// Browse To source Button
+					+ SHorizontalBox::Slot()
+					.AutoWidth()
+					.VAlign(VAlign_Center)
+					[
+						SNew(SBox)
+						.Visibility(this, &FStateTreeEditorNodeDetails::IsBrowseToSourceVisible)
+						.Padding(FMargin(4.f, 0.f, 0.f, 0.f))
+						[
+							SNew(SButton)
+							.ButtonStyle(FAppStyle::Get(), "SimpleButton")
+							.OnClicked(this, &FStateTreeEditorNodeDetails::OnBrowseToSource)
+							.ToolTipText(FText::Format(LOCTEXT("GoToCode_ToolTip", "Click to open the node source file in {0}"), FSourceCodeNavigation::GetSelectedSourceCodeIDE()))
+							.ContentPadding(0.0f)
+							[
+								SNew(SImage)
+								.Image(FAppStyle::GetBrush("Icons.OpenSourceLocation"))
+								.ColorAndOpacity(FSlateColor::UseForeground())
+							]
+						]
+					]
+
 					// Browse To BP Button
 					+ SHorizontalBox::Slot()
 					.AutoWidth()
@@ -1754,81 +1777,99 @@ FText FStateTreeEditorNodeDetails::GetNodePickerTooltip() const
 	return TextBuilder.ToText();
 }
 
-FReply FStateTreeEditorNodeDetails::OnBrowseToNodeBlueprint() const
+namespace UE::StateTree::Editor::Private
 {
-	const UObject* InstanceObject = nullptr;
-	const FPropertyAccess::Result AccessResult = InstanceObjectProperty->GetValue(InstanceObject);
-	if (AccessResult == FPropertyAccess::Success)
+	const UScriptStruct* GetNodeStruct(const TSharedPtr<IPropertyHandle>& NodeProperty)
 	{
-		check(InstanceObject);
-		if (const UBlueprintGeneratedClass* BlueprintGeneratedClass = Cast<UBlueprintGeneratedClass>(InstanceObject->GetClass()))
+		void* Address = nullptr;
+		const FPropertyAccess::Result AccessResult = NodeProperty->GetValueData(Address);
+		if (AccessResult == FPropertyAccess::Success && Address)
 		{
-			//If the blueprint asset has been cooked, UBlueprint Object will be set to null and we need to browse to its BlueprintGeneratedClass
-			GEditor->SyncBrowserToObject(BlueprintGeneratedClass->ClassGeneratedBy ? BlueprintGeneratedClass->ClassGeneratedBy.Get() : BlueprintGeneratedClass);
+			FInstancedStruct* InstancedStruct = static_cast<FInstancedStruct*>(Address);
+			return InstancedStruct->GetScriptStruct();
 		}
+		return nullptr;
+	}
+
+	const UClass* GetInstanceObjectClass(const TSharedPtr<IPropertyHandle>& InstanceObjectProperty)
+	{
+		const UObject* InstanceObject = nullptr;
+		const FPropertyAccess::Result AccessResult = InstanceObjectProperty->GetValue(InstanceObject);
+		if (AccessResult == FPropertyAccess::Success && InstanceObject)
+		{
+			return InstanceObject->GetClass();
+		}
+		return nullptr;
+	}
+}
+
+FReply FStateTreeEditorNodeDetails::OnBrowseToSource() const
+{
+	using namespace UE::StateTree::Editor::Private;
+	if (const UScriptStruct* Node = GetNodeStruct(NodeProperty))
+	{
+		FSourceCodeNavigation::NavigateToStruct(Node);
 	}
 
 	return FReply::Handled();
 }
 
-FReply FStateTreeEditorNodeDetails::OnEditNodeBlueprint() const
+FReply FStateTreeEditorNodeDetails::OnBrowseToNodeBlueprint() const
 {
-	const UObject* InstanceObject = nullptr;
-	const FPropertyAccess::Result AccessResult = InstanceObjectProperty->GetValue(InstanceObject);
-	if (AccessResult == FPropertyAccess::Success)
+	using namespace UE::StateTree::Editor::Private;
+	if (const UBlueprintGeneratedClass* BlueprintGeneratedClass = Cast<const UBlueprintGeneratedClass>(GetInstanceObjectClass(InstanceObjectProperty)))
 	{
-		check(InstanceObject);
-		const UBlueprintGeneratedClass* BlueprintGeneratedClass = Cast<UBlueprintGeneratedClass>(InstanceObject->GetClass());
-		if (BlueprintGeneratedClass && BlueprintGeneratedClass->ClassGeneratedBy)
-		{
-			//Cooked blueprint asset is not editable
-			GEditor->GetEditorSubsystem<UAssetEditorSubsystem>()->OpenEditorForAsset(BlueprintGeneratedClass->ClassGeneratedBy);
-		}
+		//If the blueprint asset has been cooked, UBlueprint Object will be set to null and we need to browse to its BlueprintGeneratedClass
+		GEditor->SyncBrowserToObject(BlueprintGeneratedClass->ClassGeneratedBy ? BlueprintGeneratedClass->ClassGeneratedBy.Get() : BlueprintGeneratedClass);
 	}
 
 	return FReply::Handled();
 }
 
-EVisibility FStateTreeEditorNodeDetails::IsBrowseToNodeBlueprintVisible() const
+FReply FStateTreeEditorNodeDetails::OnEditNodeBlueprint() const
 {
-	const UObject* InstanceObject = nullptr;
-	const FPropertyAccess::Result AccessResult = InstanceObjectProperty->GetValue(InstanceObject);
-	if (AccessResult == FPropertyAccess::Success)
+	//Cooked blueprint asset is not editable.
+	using namespace UE::StateTree::Editor::Private;
+	const UBlueprintGeneratedClass* BlueprintGeneratedClass = Cast<const UBlueprintGeneratedClass>(GetInstanceObjectClass(InstanceObjectProperty));
+	if (BlueprintGeneratedClass && BlueprintGeneratedClass->ClassGeneratedBy)
 	{
-		//The read could be null with an Success AccessResult in updating visibility 
-		if (InstanceObject)
-		{
-			if (const UBlueprintGeneratedClass* BlueprintGeneratedClass = Cast<UBlueprintGeneratedClass>(InstanceObject->GetClass()))
-			{
-				return EVisibility::Visible;
-			}
-		}
+		GEditor->GetEditorSubsystem<UAssetEditorSubsystem>()->OpenEditorForAsset(BlueprintGeneratedClass->ClassGeneratedBy);
 	}
 
-	return EVisibility::Collapsed;
+	return FReply::Handled();
 }
 
-EVisibility FStateTreeEditorNodeDetails::IsEditNodeBlueprintVisible() const
+EVisibility FStateTreeEditorNodeDetails::IsBrowseToSourceVisible() const
 {
-	const UObject* InstanceObject = nullptr;
-	const FPropertyAccess::Result AccessResult = InstanceObjectProperty->GetValue(InstanceObject);
-	if (AccessResult == FPropertyAccess::Success)
+	using namespace UE::StateTree::Editor::Private;
+	if (!Cast<const UBlueprintGeneratedClass>(GetInstanceObjectClass(InstanceObjectProperty)))
 	{
-		//The read could be null with an Success AccessResult in updating visibility 
-		if (InstanceObject)
+		if (const UScriptStruct* Node = GetNodeStruct(NodeProperty))
 		{
-			const UBlueprintGeneratedClass* BlueprintGeneratedClass = Cast<UBlueprintGeneratedClass>(InstanceObject->GetClass());
-			if (BlueprintGeneratedClass && BlueprintGeneratedClass->ClassGeneratedBy)
+			if (FSourceCodeNavigation::CanNavigateToStruct(Node))
 			{
-				//Cooked blueprint asset is not editable
 				return EVisibility::Visible;
 			}
 		}
 	}
-
 	return EVisibility::Collapsed;
 }
 
+EVisibility FStateTreeEditorNodeDetails::IsBrowseToNodeBlueprintVisible() const
+{
+	using namespace UE::StateTree::Editor::Private;
+	const UBlueprintGeneratedClass* BlueprintGeneratedClass = Cast<const UBlueprintGeneratedClass>(GetInstanceObjectClass(InstanceObjectProperty));
+	return BlueprintGeneratedClass != nullptr ? EVisibility::Visible : EVisibility::Collapsed;
+}
+
+EVisibility FStateTreeEditorNodeDetails::IsEditNodeBlueprintVisible() const
+{
+	//Cooked blueprint asset is not editable
+	using namespace UE::StateTree::Editor::Private;
+	const UBlueprintGeneratedClass* BlueprintGeneratedClass = Cast<const UBlueprintGeneratedClass>(GetInstanceObjectClass(InstanceObjectProperty));
+	return BlueprintGeneratedClass != nullptr && BlueprintGeneratedClass->ClassGeneratedBy ? EVisibility::Visible : EVisibility::Collapsed;
+}
+
 void FStateTreeEditorNodeDetails::GeneratePickerMenu(class FMenuBuilder& InMenuBuilder)
 {
 	// Expand and select currently selected item.
diff --git a/Engine/Plugins/Runtime/StateTree/Source/StateTreeEditorModule/Private/Customizations/StateTreeEditorNodeDetails.h b/Engine/Plugins/Runtime/StateTree/Source/StateTreeEditorModule/Private/Customizations/StateTreeEditorNodeDetails.h
--- a/Engine/Plugins/Runtime/StateTree/Source/StateTreeEditorModule/Private/Customizations/StateTreeEditorNodeDetails.h
+++ b/Engine/Plugins/Runtime/StateTree/Source/StateTreeEditorModule/Private/Customizations/StateTreeEditorNodeDetails.h
@@ -115,8 +115,10 @@ private:
 	void OnBindingChanged(const FPropertyBindingPath& SourcePath, const FPropertyBindingPath& TargetPath);
 	void FindOuterObjects();
 
+	FReply OnBrowseToSource() const;
 	FReply OnBrowseToNodeBlueprint() const;
 	FReply OnEditNodeBlueprint() const;
+	EVisibility IsBrowseToSourceVisible() const;
 	EVisibility IsBrowseToNodeBlueprintVisible() const;
 	EVisibility IsEditNodeBlueprintVisible() const;
 
