[StateTree] Fix memory leak introduced with CL 42822631. Use Snprintf for safer string manipulation.

#rnx

diff --git a/Engine/Plugins/Runtime/StateTree/Source/StateTreeModule/Private/CrashReporter/StateTreeCrashReporterHandler.cpp b/Engine/Plugins/Runtime/StateTree/Source/StateTreeModule/Private/CrashReporter/StateTreeCrashReporterHandler.cpp
--- a/Engine/Plugins/Runtime/StateTree/Source/StateTreeModule/Private/CrashReporter/StateTreeCrashReporterHandler.cpp
+++ b/Engine/Plugins/Runtime/StateTree/Source/StateTreeModule/Private/CrashReporter/StateTreeCrashReporterHandler.cpp
@@ -105,16 +105,11 @@ public:
 			Builder << TEXT('\n');
 			Builder << Context.ThreadID;
 
-			if (Builder.Len())
-			{
-				// Avoid allocating a new name for each entry.
-				TCHAR Identifier[] = TEXT("StateTree_000");
-				constexpr int32 NumberIndex = 10;
-				Identifier[NumberIndex + 0] = (TCHAR)((Index / 100) % 10 + int32(TEXT('0')));
-				Identifier[NumberIndex + 1] = (TCHAR)((Index / 10) % 10 + int32(TEXT('0')));
-				Identifier[NumberIndex + 2] = (TCHAR)((Index / 1) % 10 + int32(TEXT('0')));
-				Writer.AddString(Identifier, Builder.ToString());
-			}
+			// Avoid allocating a new name for each entry.
+			constexpr int32 BufferSize = 16;
+			TCHAR Identifier[BufferSize];
+			FCString::Snprintf(Identifier, UE_ARRAY_COUNT(Identifier), TEXT("StateTree%d"), Index);
+			Writer.AddString(Identifier, Builder.ToString());
 		}
 	}
 
@@ -125,7 +120,8 @@ private:
 };
 
 
-static FCrashReporterHandlerImpl* Instance = nullptr;
+//~ Use a shared ptr instead of manually deleting it in case we crash while unregistering the StateTree module.
+static TSharedPtr<FCrashReporterHandlerImpl> Instance;
 static FDelegateHandle AdditionalCrashContextDelegateHandle;
 
 } // namespace UE::StateTree::Private
@@ -135,14 +131,14 @@ namespace UE::StateTree
 
 void FCrashReporterHandler::Register()
 {
-	Private::Instance = new Private::FCrashReporterHandlerImpl();
-	Private::AdditionalCrashContextDelegateHandle = FGenericCrashContext::OnAdditionalCrashContextDelegate().AddRaw(Private::Instance, &Private::FCrashReporterHandlerImpl::DumpInfo);
+	Private::Instance = MakeShared<Private::FCrashReporterHandlerImpl>();
+	Private::AdditionalCrashContextDelegateHandle = FGenericCrashContext::OnAdditionalCrashContextDelegate().AddSP(Private::Instance.ToSharedRef(), &Private::FCrashReporterHandlerImpl::DumpInfo);
 }
 
 void FCrashReporterHandler::Unregister()
 {
 	FGenericCrashContext::OnAdditionalCrashContextDelegate().Remove(Private::AdditionalCrashContextDelegateHandle);
-	Private::Instance = nullptr;
+	// Keep the AdditionalCrashContextDelegateHandle valid in case it currently crashing. The pointer will be automatically removed (shared ptr) when the module unloads.
 }
 
 FCrashReporterScope::FCrashReporterScope(TNotNull<const UObject*> Owner, TNotNull<const UStateTree*> StateTree, FName Context)
