[StateTree] Create the node cache on startup. Helps with the multithread BP compilation.

diff --git a/Engine/Plugins/Runtime/StateTree/Source/StateTreeEditorModule/Private/StateTreeEditorModule.cpp b/Engine/Plugins/Runtime/StateTree/Source/StateTreeEditorModule/Private/StateTreeEditorModule.cpp
--- a/Engine/Plugins/Runtime/StateTree/Source/StateTreeEditorModule/Private/StateTreeEditorModule.cpp
+++ b/Engine/Plugins/Runtime/StateTree/Source/StateTreeEditorModule/Private/StateTreeEditorModule.cpp
@@ -51,6 +51,22 @@
 		return UStateTreeEditingSubsystem::CompileStateTree(&StateTree, Log);
 	}
 
+	static TSharedRef<FStateTreeNodeClassCache> InitNodeClassCache()
+	{
+		TSharedRef<FStateTreeNodeClassCache> NodeClassCache = MakeShareable(new FStateTreeNodeClassCache());
+		NodeClassCache->AddRootScriptStruct(FStateTreeEvaluatorBase::StaticStruct());
+		NodeClassCache->AddRootScriptStruct(FStateTreeTaskBase::StaticStruct());
+		NodeClassCache->AddRootScriptStruct(FStateTreeConditionBase::StaticStruct());
+		NodeClassCache->AddRootScriptStruct(FStateTreeConsiderationBase::StaticStruct());
+		NodeClassCache->AddRootScriptStruct(FStateTreePropertyFunctionBase::StaticStruct());
+		NodeClassCache->AddRootClass(UStateTreeEvaluatorBlueprintBase::StaticClass());
+		NodeClassCache->AddRootClass(UStateTreeTaskBlueprintBase::StaticClass());
+		NodeClassCache->AddRootClass(UStateTreeConditionBlueprintBase::StaticClass());
+		NodeClassCache->AddRootClass(UStateTreeConsiderationBlueprintBase::StaticClass());
+		NodeClassCache->AddRootClass(UStateTreeSchema::StaticClass());
+		return NodeClassCache;
+	}
+
 }; // UE::StateTree::Editor
 
 FStateTreeEditorModule& FStateTreeEditorModule::GetModule()
@@ -68,6 +84,11 @@
 	UE::StateTree::Delegates::OnRequestCompile.BindStatic(&UE::StateTree::Editor::CompileStateTree);
 	UE::StateTree::Delegates::OnRequestEditorHash.BindLambda([](const UStateTree& InStateTree) -> uint32 { return UStateTreeEditingSubsystem::CalculateStateTreeHash(&InStateTree); });
 
+	OnPostEngineInitHandle = FCoreDelegates::OnPostEngineInit.AddLambda([this]()
+		{
+			NodeClassCache = UE::StateTree::Editor::InitNodeClassCache();
+		});
+
 #if WITH_STATETREE_TRACE_DEBUGGER
 	FStateTreeDebuggerCommands::Register();
 
@@ -109,21 +130,6 @@
 
 void FStateTreeEditorModule::ShutdownModule()
 {
-	UE::StateTree::Delegates::OnRequestCompile.Unbind();
-
-#if WITH_STATETREE_TRACE_DEBUGGER
-	IModularFeatures::Get().UnregisterModularFeature(IRewindDebuggerRuntimeExtension::ModularFeatureName, RewindDebuggerRuntimeExtension.Get());
-	IModularFeatures::Get().UnregisterModularFeature(IRewindDebuggerExtension::ModularFeatureName, RewindDebuggerExtension.Get());
-	IModularFeatures::Get().UnregisterModularFeature(RewindDebugger::IRewindDebuggerTrackCreator::ModularFeatureName, RewindDebuggerTrackCreator.Get());
-	FStateTreeDebuggerCommands::Unregister();
-#endif // WITH_STATETREE_TRACE_DEBUGGER
-	
-	MenuExtensibilityManager.Reset();
-	ToolBarExtensibilityManager.Reset();
-
-	FStateTreeEditorStyle::Unregister();
-	FStateTreeEditorCommands::Unregister();
-
 	// Unregister the details customization
 	if (FModuleManager::Get().IsModuleLoaded("PropertyEditor"))
 	{
@@ -143,6 +149,25 @@
 		PropertyModule.UnregisterCustomClassLayout("StateTreeEditorData");
 		PropertyModule.NotifyCustomizationModuleChanged();
 	}
+
+	FStateTreeEditorStyle::Unregister();
+	FStateTreeEditorCommands::Unregister();
+
+	MenuExtensibilityManager.Reset();
+	ToolBarExtensibilityManager.Reset();
+
+#if WITH_STATETREE_TRACE_DEBUGGER
+	IModularFeatures::Get().UnregisterModularFeature(RewindDebugger::IRewindDebuggerTrackCreator::ModularFeatureName, RewindDebuggerTrackCreator.Get());
+	IModularFeatures::Get().UnregisterModularFeature(IRewindDebuggerRuntimeExtension::ModularFeatureName, RewindDebuggerRuntimeExtension.Get());
+	IModularFeatures::Get().UnregisterModularFeature(IRewindDebuggerExtension::ModularFeatureName, RewindDebuggerExtension.Get());
+	FStateTreeDebuggerCommands::Unregister();
+#endif // WITH_STATETREE_TRACE_DEBUGGER
+
+	FCoreDelegates::OnPostEngineInit.Remove(OnPostEngineInitHandle);
+	NodeClassCache.Reset();
+
+	UE::StateTree::Delegates::OnRequestEditorHash.Unbind();
+	UE::StateTree::Delegates::OnRequestCompile.Unbind();
 }
 
 TSharedRef<IStateTreeEditor> FStateTreeEditorModule::CreateStateTreeEditor(const EToolkitMode::Type Mode, const TSharedPtr< class IToolkitHost >& InitToolkitHost, UStateTree* StateTree)
@@ -160,21 +185,7 @@
 
 TSharedPtr<FStateTreeNodeClassCache> FStateTreeEditorModule::GetNodeClassCache()
 {
-	if (!NodeClassCache.IsValid())
-	{
-		NodeClassCache = MakeShareable(new FStateTreeNodeClassCache());
-		NodeClassCache->AddRootScriptStruct(FStateTreeEvaluatorBase::StaticStruct());
-		NodeClassCache->AddRootScriptStruct(FStateTreeTaskBase::StaticStruct());
-		NodeClassCache->AddRootScriptStruct(FStateTreeConditionBase::StaticStruct());
-		NodeClassCache->AddRootScriptStruct(FStateTreeConsiderationBase::StaticStruct());
-		NodeClassCache->AddRootScriptStruct(FStateTreePropertyFunctionBase::StaticStruct());
-		NodeClassCache->AddRootClass(UStateTreeEvaluatorBlueprintBase::StaticClass());
-		NodeClassCache->AddRootClass(UStateTreeTaskBlueprintBase::StaticClass());
-		NodeClassCache->AddRootClass(UStateTreeConditionBlueprintBase::StaticClass());
-		NodeClassCache->AddRootClass(UStateTreeConsiderationBlueprintBase::StaticClass());
-		NodeClassCache->AddRootClass(UStateTreeSchema::StaticClass());
-	}
-
+	check(NodeClassCache.IsValid());
 	return NodeClassCache;
 }
 
diff --git a/Engine/Plugins/Runtime/StateTree/Source/StateTreeEditorModule/Private/StateTreeNodeClassCache.cpp b/Engine/Plugins/Runtime/StateTree/Source/StateTreeEditorModule/Private/StateTreeNodeClassCache.cpp
--- a/Engine/Plugins/Runtime/StateTree/Source/StateTreeEditorModule/Private/StateTreeNodeClassCache.cpp
+++ b/Engine/Plugins/Runtime/StateTree/Source/StateTreeEditorModule/Private/StateTreeNodeClassCache.cpp
@@ -112,9 +112,12 @@ FStateTreeNodeClassCache::FStateTreeNodeClassCache()
 	FCoreUObjectDelegates::ReloadCompleteDelegate.AddRaw(this, &FStateTreeNodeClassCache::OnReloadComplete);
 
 	// Register to have Populate called when a Blueprint is compiled.
-	check(GEditor);
-	GEditor->OnBlueprintCompiled().AddRaw(this, &FStateTreeNodeClassCache::InvalidateCache);
-	GEditor->OnClassPackageLoadedOrUnloaded().AddRaw(this, &FStateTreeNodeClassCache::InvalidateCache);
+	if (GIsEditor)
+	{
+		check(GEditor);
+		GEditor->OnBlueprintCompiled().AddRaw(this, &FStateTreeNodeClassCache::InvalidateCache);
+		GEditor->OnClassPackageLoadedOrUnloaded().AddRaw(this, &FStateTreeNodeClassCache::InvalidateCache);
+	}
 }
 
 FStateTreeNodeClassCache::~FStateTreeNodeClassCache()
diff --git a/Engine/Plugins/Runtime/StateTree/Source/StateTreeEditorModule/Public/StateTreeEditorModule.h b/Engine/Plugins/Runtime/StateTree/Source/StateTreeEditorModule/Public/StateTreeEditorModule.h
--- a/Engine/Plugins/Runtime/StateTree/Source/StateTreeEditorModule/Public/StateTreeEditorModule.h
+++ b/Engine/Plugins/Runtime/StateTree/Source/StateTreeEditorModule/Public/StateTreeEditorModule.h
@@ -82,6 +82,8 @@ protected:
 	TPimplPtr<UE::StateTreeDebugger::FRewindDebuggerTrackCreator> RewindDebuggerTrackCreator;
 #endif  // WITH_STATETREE_TRACE_DEBUGGER
 
+	FDelegateHandle OnPostEngineInitHandle;
+	UE_DEPRECATED(5.7, "OnUserDefinedStructReinstancedHandle is not used.")
 	FDelegateHandle OnUserDefinedStructReinstancedHandle;
 	FOnRegisterLayoutExtensions RegisterLayoutExtensions;
 };
