[StateTree] Add accessor for global evaluators and global tasks.
#rnx

diff --git a/Engine/Plugins/Runtime/StateTree/Source/StateTreeModule/Public/StateTree.h b/Engine/Plugins/Runtime/StateTree/Source/StateTreeModule/Public/StateTree.h
--- a/Engine/Plugins/Runtime/StateTree/Source/StateTreeModule/Public/StateTree.h
+++ b/Engine/Plugins/Runtime/StateTree/Source/StateTreeModule/Public/StateTree.h
@@ -153,7 +153,7 @@ public:
 		return Schema;
 	}
 
-	/** @return Pointer to a frame or null if frame not found */
+	/** @return Pointer to a frame or null if frame not found. */
 	UE_API const FCompactStateTreeFrame* GetFrameFromHandle(const FStateTreeStateHandle StateHandle) const;
 
 	/** @return Pointer to a state or null if state not found */ 
@@ -162,14 +162,41 @@ public:
 	/** @return State handle matching a given Id; invalid handle if state not found. */
 	UE_API FStateTreeStateHandle GetStateHandleFromId(const FGuid Id) const;
 
-	/** @return Id of the state matching a given state handle; invalid Id if state not found. */
+	/** @return ID of the state matching a given state handle; invalid Id if state not found. */
 	UE_API FGuid GetStateIdFromHandle(const FStateTreeStateHandle Handle) const;
 
 	/** @return Struct view of the node matching a given node index; invalid view if state not found. */
 	UE_API FConstStructView GetNode(const int32 NodeIndex) const;
 
 	/** @return Struct views of all nodes */
-	const FInstancedStructContainer& GetNodes() const { return Nodes; }
+	const FInstancedStructContainer& GetNodes() const
+	{
+		return Nodes;
+	}
+
+	/** @return index to first global evaluator in GetNodes */
+	const uint16 GetGlobalEvaluatorsBegin() const
+	{
+		return EvaluatorsBegin;
+	}
+
+	/** @return number of global evaluators. */
+	const uint16 GetGlobalEvaluatorsNum() const
+	{
+		return EvaluatorsNum;
+	}
+
+	/** @return index to first global tasks in GetNodes */
+	const uint16 GetGlobalTasksBegin() const
+	{
+		return GlobalTasksBegin;
+	}
+	
+	/** @return number of global tasks. */
+	const uint16 GetGlobalTasksNum() const
+	{
+		return GlobalTasksNum;
+	}
 
 	/** @return Node index matching a given Id; invalid index if node not found. */
 	UE_API FStateTreeIndex16 GetNodeIndexFromId(const FGuid Id) const;
@@ -178,7 +205,10 @@ public:
 	UE_API FGuid GetNodeIdFromIndex(const FStateTreeIndex16 NodeIndex) const;
 
 	/** @return View of all states. */
-	TConstArrayView<FCompactStateTreeState> GetStates() const { return States; }
+	TConstArrayView<FCompactStateTreeState> GetStates() const
+	{
+		return States;
+	}
 
 	/** @return Pointer to the transition at a given index; null if not found. */ 
 	UE_API const FCompactStateTransition* GetTransitionFromIndex(const FStateTreeIndex16 TransitionIndex) const;
@@ -190,7 +220,10 @@ public:
 	UE_API FGuid GetTransitionIdFromIndex(const FStateTreeIndex16 Index) const;
 
 	/** @return Property bindings */
-	const FStateTreePropertyBindings& GetPropertyBindings() const { return PropertyBindings; }
+	const FStateTreePropertyBindings& GetPropertyBindings() const
+	{
+		return PropertyBindings;
+	}
 
 	/** @return True if there is any global tasks need ticking. */
 	bool DoesRequestTickGlobalTasks(bool bHasEvents) const
