[StateTree] Fixed broadcasting OnEventSent delegate event if it got rejected by the queue.
#rnx

diff --git a/Engine/Plugins/Runtime/StateTree/Source/StateTreeModule/Private/StateTreeEvents.cpp b/Engine/Plugins/Runtime/StateTree/Source/StateTreeModule/Private/StateTreeEvents.cpp
--- a/Engine/Plugins/Runtime/StateTree/Source/StateTreeModule/Private/StateTreeEvents.cpp
+++ b/Engine/Plugins/Runtime/StateTree/Source/StateTreeModule/Private/StateTreeEvents.cpp
@@ -23,21 +23,22 @@ void FStateTreeSharedEvent::AddStructReferencedObjects(FReferenceCollector& Coll
 // FStateTreeEventQueue
 //----------------------------------------------------------------//
 
-void FStateTreeEventQueue::SendEvent(const UObject* Owner, const FGameplayTag& Tag, const FConstStructView Payload, const FName Origin)
+bool FStateTreeEventQueue::SendEvent(const UObject* Owner, const FGameplayTag& Tag, const FConstStructView Payload, const FName Origin)
 {
 	if (!Tag.IsValid() && !Payload.IsValid())
 	{
 		UE_VLOG_UELOG(Owner, LogStateTree, Error, TEXT("%s: An event with an invalid tag and payload has been sent to '%s'. This is not allowed."), ANSI_TO_TCHAR(__FUNCTION__), *GetNameSafe(Owner));
-		return;
+		return false;
 	}
 
 	if (SharedEvents.Num() >= MaxActiveEvents)
 	{
 		UE_VLOG_UELOG(Owner, LogStateTree, Error, TEXT("%s: Too many events send on '%s'. Dropping event %s"), ANSI_TO_TCHAR(__FUNCTION__), *GetNameSafe(Owner), *Tag.ToString());
-		return;
+		return false;
 	}
 
 	SharedEvents.Emplace(Tag, Payload, Origin);
+	return true;
 }
 
 bool FStateTreeEventQueue::ConsumeEvent(const FStateTreeSharedEvent& Event)
diff --git a/Engine/Plugins/Runtime/StateTree/Source/StateTreeModule/Private/StateTreeExecutionContext.cpp b/Engine/Plugins/Runtime/StateTree/Source/StateTreeModule/Private/StateTreeExecutionContext.cpp
--- a/Engine/Plugins/Runtime/StateTree/Source/StateTreeModule/Private/StateTreeExecutionContext.cpp
+++ b/Engine/Plugins/Runtime/StateTree/Source/StateTreeModule/Private/StateTreeExecutionContext.cpp
@@ -693,10 +693,11 @@ void FStateTreeMinimalExecutionContext::SendEvent(const FGameplayTag Tag, const
 	UE_STATETREE_DEBUG_LOG_EVENT(this, Log, TEXT("Send Event '%s'"), *Tag.ToString());
 
 	FStateTreeEventQueue& LocalEventQueue = Storage.GetMutableEventQueue();
-	LocalEventQueue.SendEvent(&Owner, Tag, Payload, Origin);
-	ScheduleNextTick();
-
-	UE_STATETREE_DEBUG_SEND_EVENT(this, &RootStateTree, Tag, Payload, Origin);
+	if (LocalEventQueue.SendEvent(&Owner, Tag, Payload, Origin))
+	{
+		ScheduleNextTick();
+		UE_STATETREE_DEBUG_SEND_EVENT(this, &RootStateTree, Tag, Payload, Origin);
+	}
 }
 
 void FStateTreeMinimalExecutionContext::ScheduleNextTick()
diff --git a/Engine/Plugins/Runtime/StateTree/Source/StateTreeModule/Public/StateTreeEvents.h b/Engine/Plugins/Runtime/StateTree/Source/StateTreeModule/Public/StateTreeEvents.h
--- a/Engine/Plugins/Runtime/StateTree/Source/StateTreeModule/Public/StateTreeEvents.h
+++ b/Engine/Plugins/Runtime/StateTree/Source/StateTreeModule/Public/StateTreeEvents.h
@@ -184,8 +184,9 @@ struct FStateTreeEventQueue
 	 * @param Tag tag identifying the event.
 	 * @param Payload Optional reference to the payload struct.
 	 * @param Origin Optional name identifying the origin of the event.
+	 * @return true if successfully added the event to the events queue.
 	 */
-	UE_API void SendEvent(const UObject* Owner, const FGameplayTag& Tag, const FConstStructView Payload = FConstStructView(), const FName Origin = FName());
+	UE_API bool SendEvent(const UObject* Owner, const FGameplayTag& Tag, const FConstStructView Payload = FConstStructView(), const FName Origin = FName());
 
 	/**
 	 * Consumes and removes the specified event from the event queue.
