[StateTreeDebugger] ignore selected frame if no events were produced instead of displaying events from the last non empty frame

diff --git a/Engine/Plugins/Runtime/StateTree/Source/StateTreeEditorModule/Private/Debugger/SStateTreeDebuggerView.cpp b/Engine/Plugins/Runtime/StateTree/Source/StateTreeEditorModule/Private/Debugger/SStateTreeDebuggerView.cpp
--- a/Engine/Plugins/Runtime/StateTree/Source/StateTreeEditorModule/Private/Debugger/SStateTreeDebuggerView.cpp
+++ b/Engine/Plugins/Runtime/StateTree/Source/StateTreeEditorModule/Private/Debugger/SStateTreeDebuggerView.cpp
@@ -921,8 +921,15 @@ void SStateTreeDebuggerView::OnDebuggerScrubStateChanged(const UE::StateTreeDebu
 		return;
 	}
 
-	const int32 FirstEventIdx = Spans[SpanIdx].EventIdx;
-	const TraceServices::FFrame Frame = Spans[SpanIdx].Frame;
+	const UE::StateTreeDebugger::FFrameSpan& Span = Spans[SpanIdx];
+	if (ScrubState.GetScrubTime() < Span.GetWorldTimeStart()
+		|| ScrubState.GetScrubTime() > Span.GetWorldTimeEnd())
+	{
+		return;
+	}
+
+	const int32 FirstEventIdx = Span.EventIdx;
+	const TraceServices::FFrame Frame = Span.Frame;
 	const int32 MaxEventIdx = Spans.IsValidIndex(SpanIdx+1) ? Spans[SpanIdx+1].EventIdx : Events.Num();
 
 	const UStateTree* const RootTree = StateTree.Get();
