[State Tree] make GetInstanceDataPtr callable only on lvalue object, no discard on MakeStrongContext

diff --git a/Engine/Plugins/Runtime/StateTree/Source/StateTreeModule/Public/StateTreeAsyncExecutionContext.h b/Engine/Plugins/Runtime/StateTree/Source/StateTreeModule/Public/StateTreeAsyncExecutionContext.h
--- a/Engine/Plugins/Runtime/StateTree/Source/StateTreeModule/Public/StateTreeAsyncExecutionContext.h
+++ b/Engine/Plugins/Runtime/StateTree/Source/StateTreeModule/Public/StateTreeAsyncExecutionContext.h
@@ -146,11 +146,11 @@ struct TStateTreeStrongExecutionContext
 	bool UpdateScheduledTickRequest(UE::StateTree::FScheduledTickHandle Handle, FStateTreeScheduledTick ScheduledTick) const;
 
 	/**
-	 * Get the Instance Data of recorded node.
+	 * Get the Instance Data of recorded node. Only callable on lvalue because it would have lost the access track.
 	 * @return nullptr if the frame or state containing the node is no longer active.
 	 */
 	template<typename T>
-	std::conditional_t<bWithWriteAccess, T*, const T*> GetInstanceDataPtr() const
+	std::conditional_t<bWithWriteAccess, T*, const T*> GetInstanceDataPtr() const&
 	{
 		FStateTreeDataView DataView = GetInstanceDataPtrInternal();
 		if (DataView.IsValid() && ensure(DataView.GetStruct() == T::StaticStruct()))
@@ -161,6 +161,9 @@ struct TStateTreeStrongExecutionContext
 		return nullptr;
 	}
 
+	template<typename T>
+	std::conditional_t<bWithWriteAccess, T*, const T*> GetInstanceDataPtr() const&& = delete;
+
 	/**
 	 * Get the info of active frame, state and node this Context is based on
 	 * @return an invalid FActivePathInfo if the frame or state containing the node is no longer active.
@@ -241,12 +244,12 @@ public:
 		return StateTree.Pin();
 	}
 
-	FStateTreeStrongReadOnlyExecutionContext MakeStrongReadOnlyExecutionContext() const
+	[[nodiscard]] FStateTreeStrongReadOnlyExecutionContext MakeStrongReadOnlyExecutionContext() const
 	{
 		return FStateTreeStrongReadOnlyExecutionContext(*this);
 	}
 
-	FStateTreeStrongExecutionContext MakeStrongExecutionContext() const
+	[[nodiscard]] FStateTreeStrongExecutionContext MakeStrongExecutionContext() const
 	{
 		return FStateTreeStrongExecutionContext(*this);
 	}
