[StateTreeDebugger] trace state events OnExiting and OnExited only when ExitState is called (state not sustained). This is now consistent with the calls to OnEntering and OnEntered

diff --git a/Engine/Plugins/Runtime/StateTree/Source/StateTreeModule/Private/StateTreeExecutionContext.cpp b/Engine/Plugins/Runtime/StateTree/Source/StateTreeModule/Private/StateTreeExecutionContext.cpp
--- a/Engine/Plugins/Runtime/StateTree/Source/StateTreeModule/Private/StateTreeExecutionContext.cpp
+++ b/Engine/Plugins/Runtime/StateTree/Source/StateTreeModule/Private/StateTreeExecutionContext.cpp
@@ -3070,11 +3070,10 @@ void FStateTreeExecutionContext::ExitState(const FStateTreeTransitionResult& Tra
 				, *GetSafeStateName(CurrentFrame, CurrentHandle)
 				, *UEnum::GetDisplayValueAsText(CurrentTransition.ChangeType).ToString());
 
-			UE_STATETREE_DEBUG_STATE_EVENT(this, CurrentHandle, EStateTreeTraceEventType::OnExiting);
-
 			if (ExitCall.bShouldCall)
 			{
 				FCurrentlyProcessedStateScope StateScope(*this, CurrentHandle);
+				UE_STATETREE_DEBUG_STATE_EVENT(this, CurrentHandle, EStateTreeTraceEventType::OnExiting);
 
 				// Remove any delayed transitions that belong to this state.
 				Exec.DelayedTransitions.RemoveAllSwap(
@@ -3163,9 +3162,9 @@ void FStateTreeExecutionContext::ExitState(const FStateTreeTransitionResult& Tra
 
 				// Delegate Listeners Cleanup
 				GetExecState().DelegateActiveListeners.RemoveAll(CurrentFrame.ActiveStates.StateIDs[StateIndex]);
-			}
 
-			UE_STATETREE_DEBUG_STATE_EVENT(this, CurrentHandle, EStateTreeTraceEventType::OnExited);
+				UE_STATETREE_DEBUG_STATE_EVENT(this, CurrentHandle, EStateTreeTraceEventType::OnExited);
+			}
 		}
 
 		// Frame exit call
