[StateTree] Transition with ID. Use the state and frame ID to identify the state during a transition request. Add new rules for transitions for backward compatibility with the previous expected behavior. The new behavior creates new states when the state is completed. It calls Enter/Exit of the global task and the global evaluators. The execution frame is not used to keep the active state count. That makes it easier to know if the state is valid during Enter/Exit and fix async get data. The root state in Start is threaded as a temporary frame. Similar to linked state. This allows for converting more test cases and reusing the same code path.

diff --git a/Engine/Plugins/Runtime/StateTree/Source/StateTreeModule/Private/StateTree.cpp b/Engine/Plugins/Runtime/StateTree/Source/StateTreeModule/Private/StateTree.cpp
--- a/Engine/Plugins/Runtime/StateTree/Source/StateTreeModule/Private/StateTree.cpp
+++ b/Engine/Plugins/Runtime/StateTree/Source/StateTreeModule/Private/StateTree.cpp
@@ -796,6 +796,7 @@ void UStateTree::UpdateRuntimeFlags()
 	}
 
 	bScheduledTickAllowed = Schema ? Schema->IsScheduledTickAllowed() : false;
+	StateSelectionRules = Schema ? Schema->GetStateSelectionRules() : EStateTreeStateSelectionRules::Default;
 }
 
 bool UStateTree::PatchBindings()
