[StateTree] Transition with ID. Use the state and frame ID to identify the state during a transition request. Add new rules for transitions for backward compatibility with the previous expected behavi

diff --git a/Engine/Plugins/Runtime/StateTree/Source/StateTreeModule/Private/StateTreeExecutionContext.cpp b/Engine/Plugins/Runtime/StateTree/Source/StateTreeModule/Private/StateTreeExecutionContext.cpp
--- a/Engine/Plugins/Runtime/StateTree/Source/StateTreeModule/Private/StateTreeExecutionContext.cpp
+++ b/Engine/Plugins/Runtime/StateTree/Source/StateTreeModule/Private/StateTreeExecutionContext.cpp
@@ -30,9 +30,7 @@ namespace UE::StateTree::Debug
 	static std::atomic<uint32> InstanceSerialNumber = 0;
 }; // namespace
 
-namespace UE::StateTree::ExecutionContext
-{
-namespace Private
+namespace UE::StateTree::ExecutionContext::Private
 {
 	bool bCopyBoundPropertiesOnNonTickedTask = false;
 	FAutoConsoleVariableRef CVarCopyBoundPropertiesOnNonTickedTask(
@@ -78,6 +76,7 @@ namespace Private
 #undef UE_STATETREE_ESTATETREERUNSTATUS_PRIVATE
 			return Number;
 	}
+
 	constexpr uint32 NumEStateTreeFinishTaskType()
 	{
 #ifdef UE_STATETREE_ESTATETREEFINISHTASKTYPE_PRIVATE
@@ -92,8 +91,139 @@ namespace Private
 #undef UE_STATETREE_ESTATETREEFINISHTASKTYPE_PRIVATE
 			return Number;
 	}
+
+	template<typename FrameType>
+	FrameType* FindExecutionFrame(FActiveFrameID FrameID, const TArrayView<FrameType> Frames, const TArrayView<FrameType> TemporaryFrames)
+	{
+		FrameType* Frame = Frames.FindByPredicate(
+			[FrameID](const FStateTreeExecutionFrame& Frame)
+			{
+				return Frame.FrameID == FrameID;
+			});
+		if (Frame)
+		{
+			return Frame;
 }
 
+		return TemporaryFrames.FindByPredicate(
+			[FrameID](const FStateTreeExecutionFrame& Frame)
+			{
+				return Frame.FrameID == FrameID;
+			});
+	}
+
+	const FCompactStateTreeFrame* FindStateTreeFrame(const FExecutionFrameKey& LinkStateFrameKey)
+	{
+		const FCompactStateTreeFrame* FrameInfo = LinkStateFrameKey.StateTree->GetFrameFromHandle(LinkStateFrameKey.RootState);
+		ensureAlwaysMsgf(FrameInfo, TEXT("The compiled data is invalid. It should contains the information for the new root frame."));
+		return FrameInfo;
+	}
+
+	// From metric, paths are usually of length < 8. We put a large enough number to not have to worry about complex assets.
+	static constexpr int32 ExpectedActiveStatePathLength = 24;
+	void GetActiveStatePath(TArrayView<const FStateTreeExecutionFrame> Frames, TArray<UE::StateTree::FActiveState, TInlineAllocator<ExpectedActiveStatePathLength>>& OutResult)
+	{
+		if (Frames.Num() == 0 || Frames[0].StateTree == nullptr)
+		{
+			return;
+		}
+
+		for (const FStateTreeExecutionFrame& Frame : Frames)
+		{
+			for (int32 StateIndex = 0; StateIndex < Frame.ActiveStates.Num(); ++StateIndex)
+			{
+				OutResult.Emplace(Frame.FrameID, Frame.ActiveStates.StateIDs[StateIndex], Frame.ActiveStates.States[StateIndex]);
+			}
+		}
+	}
+
+	bool GetStatesListToState(FStateHandleContext State, TArray<FStateHandleContext, TInlineAllocator<FStateTreeActiveStates::MaxStates>>& OutStates)
+	{
+		// Walk towards the root from the state.
+		FStateHandleContext CurrentState = State;
+		while (CurrentState.StateTree && CurrentState.StateHandle.IsValid())
+		{
+			if (OutStates.Num() == FStateTreeActiveStates::MaxStates)
+			{
+				return false;
+			}
+			// Store the states that are in between the 'NextState' and the common ancestor.
+			OutStates.Add(CurrentState);
+			CurrentState = FStateHandleContext(CurrentState.StateTree, CurrentState.StateTree->GetStates()[CurrentState.StateHandle.Index].Parent);
+		}
+		Algo::Reverse(OutStates);
+		return true;
+	}
+
+	using FSharedEventInlineArray = TArray<FStateTreeSharedEvent, TInlineAllocator<8>>;
+	void GetTriggerTransitionEvent(const FCompactStateTransition& Transition, FStateTreeInstanceStorage& Storage, const FStateTreeSharedEvent& TransitionEvent, const TArrayView<const FStateTreeSharedEvent> EventsToProcess, FSharedEventInlineArray& OutTransitionEvents)
+	{
+		if (Transition.Trigger == EStateTreeTransitionTrigger::OnEvent)
+		{
+			check(Transition.RequiredEvent.IsValid());
+
+			if (TransitionEvent.IsValid())
+			{
+				OutTransitionEvents.Add(TransitionEvent);
+			}
+			else
+			{
+				for (const FStateTreeSharedEvent& Event : EventsToProcess)
+				{
+					check(Event.IsValid());
+					if (Transition.RequiredEvent.DoesEventMatchDesc(*Event))
+					{
+						OutTransitionEvents.Emplace(Event);
+					}
+				}
+			}
+		}
+		else if (EnumHasAnyFlags(Transition.Trigger, EStateTreeTransitionTrigger::OnTick))
+		{
+			// Dummy event to make sure we iterate the transition loop once.
+			OutTransitionEvents.Emplace(FStateTreeSharedEvent());
+		}
+		else if (EnumHasAnyFlags(Transition.Trigger, EStateTreeTransitionTrigger::OnDelegate))
+		{
+			if (Storage.IsDelegateBroadcasted(Transition.RequiredDelegateDispatcher))
+			{
+				// Dummy event to make sure we iterate the transition loop once.
+				OutTransitionEvents.Emplace(FStateTreeSharedEvent());
+			}
+		}
+		else
+		{
+			ensureMsgf(false, TEXT("The trigger type is not supported."));
+		}
+	}
+
+	FString GetStatePathAsString(TNotNull<const UStateTree*> StateTree, const TArrayView<const UE::StateTree::FActiveState>& Path)
+	{
+		TValueOrError<FString, void> Result = UE::StateTree::FActiveStatePath::Describe(StateTree, Path);
+		if (Result.HasValue())
+		{
+			return Result.StealValue();
+		}
+		return FString();
+	}
+
+	void CleanFrame(FStateTreeExecutionState& Exec, FActiveFrameID FrameID)
+	{
+		Exec.DelegateActiveListeners.RemoveAll(FrameID);
+	}
+
+	void CleanState(FStateTreeExecutionState& Exec, FActiveStateID StateID)
+	{
+		Exec.DelegateActiveListeners.RemoveAll(StateID);
+		Exec.DelayedTransitions.RemoveAll([StateID](const FStateTreeTransitionDelayedState& DelayedState)
+			{
+				return DelayedState.StateID == StateID;
+			});
+	}
+} // namespace UE::StateTree::ExecutionContext::Private
+
+namespace UE::StateTree::ExecutionContext
+{
 	bool MarkDelegateAsBroadcasted(FStateTreeDelegateDispatcher Dispatcher, const FStateTreeExecutionFrame& CurrentFrame, FStateTreeInstanceStorage& Storage)
 	{
 		const UStateTree* StateTree = CurrentFrame.StateTree;
@@ -175,7 +305,7 @@ namespace Private
 
 		return (EStateTreeRunStatus)InStatus;
 	}
-}
+} // namespace UE::StateTree::ExecutionContext
 
 /**
  * FStateTreeReadOnlyExecutionContext implementation
@@ -709,6 +839,44 @@ void FStateTreeMinimalExecutionContext::ScheduleNextTick()
 	}
 }
 
+/**
+ * FStateTreeExecutionContext::FSelectStateResult implementation
+ */
+FStateTreeExecutionFrame& FStateTreeExecutionContext::FSelectStateResult::MakeAndAddTemporaryFrame(
+	const UE::StateTree::FActiveFrameID FrameID,
+	const UE::StateTree::FExecutionFrameKey& FrameKey,
+	const bool bIsGlobalFrame)
+{
+	FStateTreeExecutionFrame& ExecFrame = TemporaryFrames.AddDefaulted_GetRef();
+	ExecFrame.FrameID = FrameID;
+	ExecFrame.StateTree = FrameKey.StateTree;
+	ExecFrame.RootState = FrameKey.RootState;
+	ExecFrame.bIsGlobalFrame = bIsGlobalFrame;
+
+	const FCompactStateTreeFrame* TargetFrame = UE::StateTree::ExecutionContext::Private::FindStateTreeFrame(FrameKey);
+	ensureMsgf(TargetFrame, TEXT("The compiled data is invalid. It should contains the information for the new root frame."));
+	ExecFrame.ActiveTasksStatus = TargetFrame ? FStateTreeTasksCompletionStatus(*TargetFrame) : FStateTreeTasksCompletionStatus();
+
+	return ExecFrame;
+}
+
+FStateTreeExecutionFrame& FStateTreeExecutionContext::FSelectStateResult::MakeAndAddTemporaryFrameWithNewRoot(
+	const UE::StateTree::FActiveFrameID FrameID,
+	const UE::StateTree::FExecutionFrameKey& FrameKey,
+	FStateTreeExecutionFrame& OtherFrame)
+{
+	constexpr bool bIsGlobalFrame = true;
+	FStateTreeExecutionFrame& ExecFrame = MakeAndAddTemporaryFrame(FrameID, FrameKey, bIsGlobalFrame);
+	ExecFrame.ExternalDataBaseIndex = OtherFrame.ExternalDataBaseIndex;
+	ExecFrame.GlobalParameterDataHandle = OtherFrame.GlobalParameterDataHandle;
+	ExecFrame.GlobalInstanceIndexBase = OtherFrame.GlobalInstanceIndexBase;
+	// Don't stop globals.
+	ExecFrame.bHaveEntered = OtherFrame.bHaveEntered;
+	OtherFrame.bHaveEntered = false;
+
+	return ExecFrame;
+}
+
 /**
  * FStateTreeExecutionContext::FCurrentlyProcessedFrameScope implementation
  */
@@ -1104,6 +1272,10 @@ EStateTreeRunStatus FStateTreeExecutionContext::Start(FStartParameters Parameter
 	CSV_SCOPED_TIMING_STAT_EXCLUSIVE(StateTree_Start);
 	UE_STATETREE_CRASH_REPORTER_SCOPE(&Owner, &RootStateTree, UE::StateTree::ExecutionContext::Private::Name_Start.Resolve());
 
+	using namespace UE::StateTree;
+	using namespace UE::StateTree::ExecutionContext;
+	using namespace UE::StateTree::ExecutionContext::Private;
+
 	if (!IsValid())
 	{
 		STATETREE_LOG(Warning, TEXT("%hs: StateTree context is not initialized properly ('%s' using StateTree '%s')"),
@@ -1144,25 +1316,20 @@ EStateTreeRunStatus FStateTreeExecutionContext::Start(FStartParameters Parameter
 		SetGlobalParameters(RootStateTree.GetDefaultParameters());
 	}
 
+	Exec.RandomStream.Initialize(Parameters.RandomSeed.IsSet() ? Parameters.RandomSeed.GetValue() : FPlatformTime::Cycles());
+
 	TGuardValue<bool> ScheduledNextTickScope(bAllowedToScheduleNextTick, false);
+	ensure(Exec.ActiveFrames.Num() == 0);
 
 	// Initialize for the init frame.
-	ensure(Exec.ActiveFrames.Num() == 0);
-	FStateTreeExecutionFrame& InitFrame = Exec.ActiveFrames.AddDefaulted_GetRef();
-	InitFrame.FrameID = UE::StateTree::FActiveFrameID(Storage.GenerateUniqueId());
-	InitFrame.StateTree = &RootStateTree;
-	InitFrame.RootState = FStateTreeStateHandle::Root;
-	InitFrame.ActiveStates = {};
-	InitFrame.bIsGlobalFrame = true;
-
-	const FCompactStateTreeFrame* FrameInfo = RootStateTree.GetFrameFromHandle(FStateTreeStateHandle::Root);
-	ensureAlwaysMsgf(FrameInfo, TEXT("The compiled data is invalid. It should contains the information for the root frame."));
-	InitFrame.ActiveTasksStatus = FrameInfo ? FStateTreeTasksCompletionStatus(*FrameInfo) : FStateTreeTasksCompletionStatus();
+	FSelectStateResult SelectStateResult;
+	const FActiveFrameID InitFrameID = FActiveFrameID(Storage.GenerateUniqueId());
+	const FExecutionFrameKey InitFrameKey = FExecutionFrameKey(&RootStateTree, FStateTreeStateHandle::Root);
+	constexpr bool bIsGlobalFrame = true;
+	FStateTreeExecutionFrame& InitFrame = SelectStateResult.MakeAndAddTemporaryFrame(InitFrameID, InitFrameKey, bIsGlobalFrame);
+	InitFrame.GlobalParameterDataHandle = FStateTreeDataHandle(EStateTreeDataSourceType::GlobalParameterData);
 	
-	UpdateInstanceData({}, Exec.ActiveFrames);
-	Exec.RandomStream.Initialize(Parameters.RandomSeed.IsSet() ? Parameters.RandomSeed.GetValue() : FPlatformTime::Cycles());
-
-	if (!CollectActiveExternalData())
+	if (!CollectActiveExternalData(SelectStateResult.TemporaryFrames))
 	{
 		STATETREE_LOG(Warning, TEXT("%hs: Failed to collect external data ('%s' using StateTree '%s')"),
 			__FUNCTION__, *GetNameSafe(&Owner), *GetFullNameSafe(&RootStateTree));
@@ -1179,40 +1346,44 @@ EStateTreeRunStatus FStateTreeExecutionContext::Start(FStartParameters Parameter
 	SetUpdatePhaseInExecutionState(Exec, EStateTreeUpdatePhase::StartTree);
 
 	// Start evaluators and global tasks. Fail the execution if any global task fails.
-	const EStateTreeRunStatus GlobalTasksRunStatus = StartEvaluatorsAndGlobalTasks();
+	constexpr FStateTreeExecutionFrame* ParentInitFrame = nullptr;
+	EStateTreeRunStatus GlobalTasksRunStatus = StartTemporaryEvaluatorsAndGlobalTasks(ParentInitFrame, InitFrame);
 	if (GlobalTasksRunStatus == EStateTreeRunStatus::Running)
 	{
 		// Exception with Start. Tick the evaluators.
 		constexpr float DeltaTime = 0.0f;
-		constexpr FStateTreeExecutionFrame* ParentInitFrame = nullptr;
-		TickGlobalEvaluatorsForFrameOnActiveInstances(DeltaTime, ParentInitFrame, InitFrame);
+		TickGlobalEvaluatorsForFrameWithValidation(DeltaTime, ParentInitFrame, InitFrame);
 
 		// Initialize to unset running state.
 		Exec.TreeRunStatus = EStateTreeRunStatus::Running;
 		Exec.LastTickStatus = EStateTreeRunStatus::Unset;
 
-		static const FStateTreeStateHandle RootState = FStateTreeStateHandle::Root;
-
-		FStateSelectionResult StateSelectionResult;
-		if (SelectState(InitFrame, RootState, StateSelectionResult))
+		FSelectStateArguments SelectStateArgs;
+		SelectStateArgs.SourceState = FActiveState(InitFrameID, FActiveStateID(), InitFrameKey.RootState);
+		SelectStateArgs.TargetState = FStateHandleContext(InitFrameKey.StateTree, InitFrameKey.RootState);
+		SelectStateArgs.Behavior = ESelectStateBehavior::StateTransition;
+		if (SelectState(SelectStateArgs, SelectStateResult))
 		{
-			check(StateSelectionResult.ContainsFrames());
-			if (StateSelectionResult.GetSelectedFrames().Last().ActiveStates.Last().IsCompletionState())
+			check(SelectStateResult.SelectedStates.Num() > 0);
+			const FStateTreeStateHandle LastSelectedStateHandle = SelectStateResult.SelectedStates.Last().GetStateHandle();
+			if (LastSelectedStateHandle.IsCompletionState())
 			{
 				// Transition to a terminal state (succeeded/failed).
 				STATETREE_LOG(Warning, TEXT("%hs: Tree %s at StateTree start on '%s' using StateTree '%s'."),
-					__FUNCTION__, StateSelectionResult.GetSelectedFrames().Last().ActiveStates.Last() == FStateTreeStateHandle::Succeeded ? TEXT("succeeded") : TEXT("failed"), *GetNameSafe(&Owner), *GetFullNameSafe(&RootStateTree));
-				Exec.TreeRunStatus = StateSelectionResult.GetSelectedFrames().Last().ActiveStates.Last().ToCompletionStatus();
+					__FUNCTION__,
+					LastSelectedStateHandle == FStateTreeStateHandle::Succeeded ? TEXT("succeeded") : TEXT("failed"),
+					*GetNameSafe(&Owner),
+					*GetFullNameSafe(&RootStateTree)
+				);
+				Exec.TreeRunStatus = LastSelectedStateHandle.ToCompletionStatus();
 			}
 			else
 			{
 				// Enter state tasks can fail/succeed, treat it same as tick.
 				FStateTreeTransitionResult Transition;
-				Transition.TargetState = RootState;
+				Transition.TargetState = InitFrameKey.RootState;
 				Transition.CurrentRunStatus = Exec.LastTickStatus;
-				Transition.NextActiveFrames = StateSelectionResult.GetSelectedFrames(); // Enter state will update Exec.ActiveFrames.
-				Transition.NextActiveFrameEvents = StateSelectionResult.GetFramesStateSelectionEvents();
-				const EStateTreeRunStatus LastTickStatus = EnterState(Transition);
+				const EStateTreeRunStatus LastTickStatus = EnterState(SelectStateResult, Transition);
 
 				Exec.LastTickStatus = LastTickStatus;
 
@@ -1221,22 +1392,35 @@ EStateTreeRunStatus FStateTreeExecutionContext::Start(FStartParameters Parameter
 				{
 					StateCompleted();
 				}
+
+				// Was not able to enter the root state. Fail selection.
+				if (Exec.ActiveFrames.Num() == 0 || Exec.ActiveFrames[0].ActiveStates.Num() == 0)
+				{
+					GlobalTasksRunStatus = EStateTreeRunStatus::Failed;
+					Exec.TreeRunStatus = EStateTreeRunStatus::Failed;
+					STATETREE_LOG(Error, TEXT("%hs: Failed to enter the initial state on '%s' using StateTree '%s'. Check that the StateTree logic can always select a state at start."),
+						__FUNCTION__, *GetNameSafe(&Owner), *GetFullNameSafe(&RootStateTree));
+				}
 			}
 		}
-
-		InstanceData.ResetTemporaryInstances();
-
-		if (Exec.LastTickStatus == EStateTreeRunStatus::Unset)
+		else
 		{
-			// Should not happen. This may happen if initial state could not be selected.
-			STATETREE_LOG(Error, TEXT("%hs: Failed to select initial state on '%s' using StateTree '%s'. This should not happen, check that the StateTree logic can always select a state at start."),
-				__FUNCTION__, *GetNameSafe(&Owner), *GetFullNameSafe(&RootStateTree));
+			GlobalTasksRunStatus = EStateTreeRunStatus::Failed;
 			Exec.TreeRunStatus = EStateTreeRunStatus::Failed;
+			STATETREE_LOG(Error, TEXT("%hs: Failed to select initial state on '%s' using StateTree '%s'. Check that the StateTree logic can always select a state at start."),
+				__FUNCTION__, *GetNameSafe(&Owner), *GetFullNameSafe(&RootStateTree));
 		}
 	}
 	else
 	{
-		StopEvaluatorsAndGlobalTasks(GlobalTasksRunStatus);
+		Exec.TreeRunStatus = GlobalTasksRunStatus;
+		STATETREE_LOG(VeryVerbose, TEXT("%hs: Start globals on '%s' using StateTree '%s' completes."),
+			__FUNCTION__, *GetNameSafe(&Owner), *GetFullNameSafe(&RootStateTree));
+	}
+
+	if (Exec.TreeRunStatus != EStateTreeRunStatus::Running)
+	{
+		StopTemporaryEvaluatorsAndGlobalTasks(ParentInitFrame, InitFrame, GlobalTasksRunStatus);
 
 		STATETREE_LOG(VeryVerbose, TEXT("%hs: Global tasks completed the StateTree %s on start in status '%s'."),
 			__FUNCTION__, *GetNameSafe(&Owner), *GetFullNameSafe(&RootStateTree), *UEnum::GetDisplayValueAsText(GlobalTasksRunStatus).ToString());
@@ -1249,6 +1433,7 @@ EStateTreeRunStatus FStateTreeExecutionContext::Start(FStartParameters Parameter
 		// We are not considered as running yet so we only set the status without requiring a stop.
 		Exec.TreeRunStatus = GlobalTasksRunStatus;
 	}
+	InstanceData.ResetTemporaryInstances();
 
 	// Reset phase since we are now safe to stop and before potentially stopping the instance.
 	SetUpdatePhaseInExecutionState(Exec, EStateTreeUpdatePhase::Unset);
@@ -1256,7 +1441,8 @@ EStateTreeRunStatus FStateTreeExecutionContext::Start(FStartParameters Parameter
 	// Use local for resulting run state since Stop will reset the instance data.
 	EStateTreeRunStatus Result = Exec.TreeRunStatus;
 	
-	if (Exec.RequestedStop != EStateTreeRunStatus::Unset)
+	if (Exec.RequestedStop != EStateTreeRunStatus::Unset
+		&& Exec.TreeRunStatus == EStateTreeRunStatus::Running)
 	{
 		STATETREE_LOG(VeryVerbose, TEXT("Processing Deferred Stop"));
 		UE_STATETREE_DEBUG_LOG_EVENT(this, Log, TEXT("Processing Deferred Stop"));
@@ -1315,10 +1501,11 @@ EStateTreeRunStatus FStateTreeExecutionContext::Stop(EStateTreeRunStatus Complet
 	if (Exec.TreeRunStatus == EStateTreeRunStatus::Running)
 	{
 		// Transition to Succeeded state.
+		FSelectStateResult EmptySelectionResult;
 		FStateTreeTransitionResult Transition;
 		Transition.TargetState = FStateTreeStateHandle::FromCompletionStatus(CompletionStatus);
 		Transition.CurrentRunStatus = CompletionStatus;	
-		ExitState(Transition);
+		ExitState(EmptySelectionResult, Transition);
 
 		// No active states or global tasks anymore, reset frames.
 		Exec.ActiveFrames.Reset();
@@ -1326,6 +1513,8 @@ EStateTreeRunStatus FStateTreeExecutionContext::Stop(EStateTreeRunStatus Complet
 		Result = CompletionStatus;
 	}
 
+	Exec.TreeRunStatus = CompletionStatus;
+
 	// Trace before resetting the instance data since it is required to provide all the event information
 	UE_STATETREE_DEBUG_ACTIVE_STATES_EVENT(this, {});
 	UE_STATETREE_DEBUG_EXIT_PHASE(this, EStateTreeUpdatePhase::StopTree);
@@ -1516,19 +1705,26 @@ void FStateTreeExecutionContext::TickUpdateTasksInternal(float DeltaTime)
 			if (ExecutionContext::Private::bGlobalTasksCompleteOwningFrame)
 			{
 				// Only set RequestStop if it's the first frame (root)
-				check(Exec.ActiveFrames.Num() > 0);
+				if (Exec.ActiveFrames.Num() > 0)
+				{
 				const UStateTree* StateTree = Exec.ActiveFrames[0].StateTree;
 				check(StateTree == &RootStateTree);
 				const ETaskCompletionStatus GlobalTaskStatus = Exec.ActiveFrames[0].ActiveTasksStatus.GetStatus(StateTree).GetCompletionStatus();
 				const EStateTreeRunStatus GlobalRunStatus = ExecutionContext::CastToRunStatus(GlobalTaskStatus);
 				if (GlobalRunStatus != EStateTreeRunStatus::Running)
 				{
-					// Note. Exec.RequestedStop default value is Unset.
 					Exec.RequestedStop = ExecutionContext::GetPriorityRunStatus(Exec.RequestedStop, GlobalRunStatus);
 					LogRequestStop();
 				}
 			}
 			else
+			{
+					// Start failed and the user called Tick anyway.
+					Exec.RequestedStop = EStateTreeRunStatus::Failed;
+					LogRequestStop();
+				}
+			}
+			else
 			{
 				// Any completion stops the tree execution.
 				Exec.RequestedStop = ExecutionContext::GetPriorityRunStatus(Exec.RequestedStop, EvalAndGlobalTaskStatus);
@@ -1563,33 +1759,36 @@ void FStateTreeExecutionContext::TickTriggerTransitionsInternal()
 		// Trigger conditional transitions or state succeed/failed transitions. First tick transition is handled here too.
 		if (TriggerTransitions())
 		{
+			check(RequestedTransition.IsValid());
 			UE_STATETREE_DEBUG_SCOPED_PHASE(this, EStateTreeUpdatePhase::ApplyTransitions);
-			UE_STATETREE_DEBUG_TRANSITION_EVENT(this, NextTransitionSource, EStateTreeTraceEventType::OnTransition);
-			NextTransitionSource.Reset();
+			UE_STATETREE_DEBUG_TRANSITION_EVENT(this, RequestedTransition->Source, EStateTreeTraceEventType::OnTransition);
 
-			ExitState(NextTransition);
+			ExitState(RequestedTransition->Selection, RequestedTransition->Transition);
 
 			// Tree succeeded or failed.
-			if (NextTransition.TargetState.IsCompletionState())
+			if (RequestedTransition->Transition.TargetState.IsCompletionState())
 			{
 				// Transition to a terminal state (succeeded/failed), or default transition failed.
-				Exec.TreeRunStatus = NextTransition.TargetState.ToCompletionStatus();
+				Exec.TreeRunStatus = RequestedTransition->Transition.TargetState.ToCompletionStatus();
 
-				// Stop evaluators and global tasks.
+				// Stop evaluators and global tasks (handled in ExitState)
+				if (!ensure(Exec.ActiveFrames.Num() == 0))
+				{
 				StopEvaluatorsAndGlobalTasks(Exec.TreeRunStatus);
 
 				// No active states or global tasks anymore, reset frames.
 				Exec.ActiveFrames.Reset();
 
 				RemoveAllDelegateListeners();
+				}
 
 				break;
 			}
 
 			// Enter state tasks can fail/succeed, treat it same as tick.
-			const EStateTreeRunStatus LastTickStatus = EnterState(NextTransition);
+			const EStateTreeRunStatus LastTickStatus = EnterState(RequestedTransition->Selection, RequestedTransition->Transition);
 
-			NextTransition = FStateTreeTransitionResult();
+			RequestedTransition.Reset();
 
 			Exec.LastTickStatus = LastTickStatus;
 
@@ -1697,6 +1896,9 @@ void FStateTreeExecutionContext::ConsumeEvent(const FStateTreeSharedEvent& Event
 
 void FStateTreeExecutionContext::RequestTransition(const FStateTreeTransitionRequest& Request)
 {
+	using namespace UE::StateTree;
+	using namespace UE::StateTree::ExecutionContext;
+
 	CSV_SCOPED_TIMING_STAT_EXCLUSIVE(StateTree_RequestTransition);
 
 	if (!IsValid())
@@ -1714,12 +1916,16 @@ void FStateTreeExecutionContext::RequestTransition(const FStateTreeTransitionReq
 		
 		STATETREE_LOG(Verbose, TEXT("Request transition to '%s' at priority %s"), *GetSafeStateName(*CurrentlyProcessedFrame, Request.TargetState), *UEnum::GetDisplayValueAsText(Request.Priority).ToString());
 
-		if (RequestTransition(*CurrentlyProcessedFrame, Request.TargetState, Request.Priority, /*TransitionEvent*/nullptr, Request.Fallback))
+		const FActiveStateID CurrentlyProcessedStateID = CurrentlyProcessedFrame->ActiveStates.FindStateID(CurrentlyProcessedState);
+		const FStateHandleContext Target = FStateHandleContext(CurrentlyProcessedFrame->StateTree, Request.TargetState);
+		const bool bRequested = RequestTransitionInternal(*CurrentlyProcessedFrame, CurrentlyProcessedStateID, Target, Request.Priority, FStateTreeSharedEvent(), Request.Fallback);
+		if (bRequested)
 		{
-			NextTransitionSource = FStateTreeTransitionSource(CurrentlyProcessedFrame->StateTree, EStateTreeTransitionSourceType::ExternalRequest, Request.TargetState, Request.Priority);
+			check(RequestedTransition.IsValid());
+			RequestedTransition->Source = FStateTreeTransitionSource(CurrentlyProcessedFrame->StateTree, EStateTreeTransitionSourceType::ExternalRequest, Request.TargetState, Request.Priority);
 		}
 	}
-	else
+	else if (Exec.ActiveFrames.Num() > 0)
 	{
 		const FStateTreeExecutionFrame* RootFrame = &Exec.ActiveFrames[0];
 		if (CurrentlyProcessedFrame)
@@ -1727,9 +1933,9 @@ void FStateTreeExecutionContext::RequestTransition(const FStateTreeTransitionReq
 			RootFrame = CurrentlyProcessedFrame;
 		}
 
-		if (!RootFrame)
+		if (RootFrame->ActiveStates.Num() == 0)
 		{
-			STATETREE_LOG(Warning, TEXT("%hs: RequestTransition called on %s using StateTree %s without active state. Start() must be called before requesting transition."),
+			STATETREE_LOG(Warning, TEXT("%hs: RequestTransition called on %s using StateTree %s without an active state. Start() must be called before requesting a transition."),
 				__FUNCTION__, *GetNameSafe(&Owner), *GetFullNameSafe(&RootStateTree));
 			return;
 		}
@@ -1738,13 +1944,19 @@ void FStateTreeExecutionContext::RequestTransition(const FStateTreeTransitionReq
 
 		FStateTreeTransitionRequest RequestWithSource = Request;
 		RequestWithSource.SourceFrameID = RootFrame->FrameID;
+
 		const int32 ActiveStateIndex = RootFrame->ActiveStates.IndexOfReverse(CurrentlyProcessedState);
-		RequestWithSource.SourceStateID = ActiveStateIndex != INDEX_NONE ? RootFrame->ActiveStates.StateIDs[ActiveStateIndex] : UE::StateTree::FActiveStateID::Invalid;
+		RequestWithSource.SourceStateID = ActiveStateIndex != INDEX_NONE ? RootFrame->ActiveStates.StateIDs[ActiveStateIndex] : RootFrame->ActiveStates.StateIDs[0];
 
 		InstanceData.AddTransitionRequest(&Owner, RequestWithSource);
+		ScheduleNextTick();
+	}
+	else
+	{
+		STATETREE_LOG(Warning, TEXT("%hs: RequestTransition called on %s using StateTree %s without an active frame. Start() must be called before requesting a transition."),
+			__FUNCTION__, *GetNameSafe(&Owner), *GetFullNameSafe(&RootStateTree));
+		return;
 	}
-
-	ScheduleNextTick();
 }
 
 void FStateTreeExecutionContext::RequestTransition(FStateTreeStateHandle InTargetState, EStateTreeTransitionPriority InPriority, EStateTreeSelectionFallback InFallback)
@@ -1861,25 +2073,61 @@ void FStateTreeExecutionContext::UpdateCompletedStateList()
 void FStateTreeExecutionContext::MarkStateCompleted(UE::StateTree::FFinishedTask& NewFinishedTask)
 {
 }
-PRAGMA_ENABLE_DEPRECATION_WARNINGS
 
+// Deprecated
 void FStateTreeExecutionContext::UpdateInstanceData(TConstArrayView<FStateTreeExecutionFrame> CurrentActiveFrames, TArrayView<FStateTreeExecutionFrame> NextActiveFrames)
 {
+}
+PRAGMA_ENABLE_DEPRECATION_WARNINGS
+
+void FStateTreeExecutionContext::UpdateInstanceData(FSelectStateResult& Args)
+{
+	// Go though all the state and frame (includes unchanged, sustained, changed).
+	//Keep 2 buffers. InstanceStructs for unchanged and sustained. TempInstanceStructs for changed.
+	//1. Set the frame GlobalParameterDataHandle, GlobalInstanceIndexBase and ActiveInstanceIndexBase
+	//2. Add instances for global frame.
+	// Note that the global parameters are in a different buffer or it's last state parameter or the previous frame (see Start, SelectStateInternal_Linked, SelectStateInternal_LinkedAsset).
+	//3. Add the instances for state parameters (Set the state StateParameterDataHandle).
+	//4. Add instances for state tasks.
+	//5. Remove the previous instances (that are not needed) and copy/move the new TempInstanceStructs.
+
+	using namespace UE::StateTree;
+	using namespace UE::StateTree::ExecutionContext;
+	using namespace UE::StateTree::ExecutionContext::Private;
+
+	FStateTreeExecutionState& Exec = GetExecState();
+
 	// Estimate how many new instance data items we might have.
 	int32 EstimatedNumStructs = 0;
-	for (int32 FrameIndex = 0; FrameIndex < NextActiveFrames.Num(); FrameIndex++)
 	{
-		const FStateTreeExecutionFrame& NextFrame = NextActiveFrames[FrameIndex];
-		if (NextFrame.bIsGlobalFrame)
+		FActiveFrameID CurrentFrameID;
+		const UStateTree* CurrentStateTree = nullptr;
+		for (const FActiveState& SelectedState : Args.SelectedStates)
+		{
+			// Global
+			if (SelectedState.GetFrameID() != CurrentFrameID)
+			{
+				CurrentFrameID = SelectedState.GetFrameID();
+
+				const FStateTreeExecutionFrame* Frame = FindExecutionFrame(CurrentFrameID, MakeArrayView(Exec.ActiveFrames), MakeArrayView(Args.TemporaryFrames));
+				if (ensure(Frame) && Frame->bIsGlobalFrame)
 		{
-			EstimatedNumStructs += NextFrame.StateTree->NumGlobalInstanceData;
+					check(Frame->StateTree);
+
+					CurrentStateTree = Frame->StateTree;
+					EstimatedNumStructs += CurrentStateTree->NumGlobalInstanceData;
+				}
 		}
-		// States
-		for (int32 StateIndex = 0; StateIndex < NextFrame.ActiveStates.Num(); StateIndex++)
+
+			// State
+			if (ensure(CurrentStateTree))
+			{
+				const FCompactStateTreeState* CurrentState = CurrentStateTree->GetStateFromHandle(SelectedState.GetStateHandle());
+				if (ensure(CurrentState))
 		{
-			const FStateTreeStateHandle StateHandle = NextFrame.ActiveStates[StateIndex];
-			const FCompactStateTreeState& State = NextFrame.StateTree->States[StateHandle.Index];
-			EstimatedNumStructs += State.InstanceDataNum;
+					EstimatedNumStructs += CurrentState->InstanceDataNum;
+				}
+			}
 		}
 	}
 	
@@ -1889,18 +2137,20 @@ void FStateTreeExecutionContext::UpdateInstanceData(TConstArrayView<FStateTreeEx
 	TArray<FInstancedStruct*> TempInstanceStructs;
 	TempInstanceStructs.Reserve(EstimatedNumStructs);
 
-	TArray<FCompactStateTreeParameters, TFixedAllocator<FStateSelectionResult::MaxExecutionFrames>> TempParams;
+	constexpr int32 ExpectedAmountOfFrames = 8;
+	TArray<FCompactStateTreeParameters, TInlineAllocator<ExpectedAmountOfFrames>> TempParams;
 
 	TArrayView<FStateTreeTemporaryInstanceData> TempInstances = Storage.GetMutableTemporaryInstances();
-	auto FindInstanceTempData = [&TempInstances](const FStateTreeExecutionFrame& Frame, FStateTreeDataHandle DataHandle)
+	auto FindInstanceTempData = [&TempInstances](FActiveFrameID FrameID, FStateTreeDataHandle DataHandle)
 	{
-		FStateTreeTemporaryInstanceData* TempData = TempInstances.FindByPredicate([&Frame, &DataHandle](const FStateTreeTemporaryInstanceData& Data)
+			FStateTreeTemporaryInstanceData* TempData = TempInstances.FindByPredicate([&FrameID, &DataHandle](const FStateTreeTemporaryInstanceData& Data)
 		{
-			return Data.FrameID == Frame.FrameID && Data.DataHandle == DataHandle;
+					return Data.FrameID == FrameID && Data.DataHandle == DataHandle;
 		});
 		return TempData ? &TempData->Instance : nullptr;
 	};
 	
+
 	// Find next instance data sources and find common/existing section of instance data at start.
 	int32 CurrentGlobalInstanceIndexBase = 0;
 	int32 NumCommonInstanceData = 0;
@@ -1909,97 +2159,115 @@ void FStateTreeExecutionContext::UpdateInstanceData(TConstArrayView<FStateTreeEx
 	FStateTreeDataHandle NextStateParameterDataHandle = FStateTreeDataHandle::Invalid;
 	
 	FStateTreeDataHandle CurrentGlobalParameterDataHandle = FStateTreeDataHandle(EStateTreeDataSourceType::GlobalParameterData);
-
-	bool bAreCommon = true;
-	for (int32 FrameIndex = 0; FrameIndex < NextActiveFrames.Num(); FrameIndex++)
 	{
-		const bool bIsCurrentFrameValid = CurrentActiveFrames.IsValidIndex(FrameIndex)
-						&& CurrentActiveFrames[FrameIndex].IsSameFrame(NextActiveFrames[FrameIndex]);
+		FActiveFrameID CurrentFrameID;
+		FStateTreeExecutionFrame* CurrentFrame = nullptr;
+		const UStateTree* CurrentStateTree = nullptr;
+		int32 CurrentFrameBaseIndex = 0;
+		for (int32 SelectedStateIndex = 0; SelectedStateIndex < Args.SelectedStates.Num(); ++SelectedStateIndex)
+		{
+			const FActiveState& SelectedState = Args.SelectedStates[SelectedStateIndex];
 
-		bAreCommon &= bIsCurrentFrameValid;
+			// Global
+			if (SelectedState.GetFrameID() != CurrentFrameID)
+	{
+				CurrentFrameID = SelectedState.GetFrameID();
+				CurrentFrame = Exec.FindActiveFrame(CurrentFrameID);
 
-		const FStateTreeExecutionFrame* CurrentFrame = bIsCurrentFrameValid ? &CurrentActiveFrames[FrameIndex] : nullptr;
-		FStateTreeExecutionFrame& NextFrame = NextActiveFrames[FrameIndex];
+				// The frame/globals are common (before the transition or sustained from the transition)
+				const bool bIsFrameCommon = CurrentFrame != nullptr;
 
-		check(NextFrame.StateTree);
+				if (CurrentFrame == nullptr)
+				{
+					CurrentFrame = Args.TemporaryFrames.FindByPredicate(
+						[FrameID = SelectedState.GetFrameID()](const FStateTreeExecutionFrame& Frame)
+						{
+							return Frame.FrameID == FrameID;
+						});
+				}
+				check(CurrentFrame);
+				CurrentStateTree = CurrentFrame->StateTree;
 
-		if (NextFrame.bIsGlobalFrame)
+				// Global Nodes
+				if (CurrentFrame->bIsGlobalFrame)
 		{
 			// Handle global tree parameters
 			if (NextStateParameterDataHandle.IsValid())
 			{
 				// Point to the parameter block set by linked state.
-				check(NextStateParameterDataStruct == NextFrame.StateTree->GetDefaultParameters().GetPropertyBagStruct());
+						check(NextStateParameterDataStruct == CurrentStateTree->GetDefaultParameters().GetPropertyBagStruct());
 				CurrentGlobalParameterDataHandle = NextStateParameterDataHandle;
 				NextStateParameterDataHandle = FStateTreeDataHandle::Invalid; // Mark as used.
 			}
 			
-			// Global Evals
 			const int32 BaseIndex = InstanceStructs.Num();
 			CurrentGlobalInstanceIndexBase = BaseIndex;
 			
-			InstanceStructs.AddDefaulted(NextFrame.StateTree->NumGlobalInstanceData);
-			TempInstanceStructs.AddZeroed(NextFrame.StateTree->NumGlobalInstanceData);
+					InstanceStructs.AddDefaulted(CurrentStateTree->NumGlobalInstanceData);
+					TempInstanceStructs.AddZeroed(CurrentStateTree->NumGlobalInstanceData);
 			
-			for (int32 EvalIndex = NextFrame.StateTree->EvaluatorsBegin; EvalIndex < (NextFrame.StateTree->EvaluatorsBegin + NextFrame.StateTree->EvaluatorsNum); EvalIndex++)
+					// Global Evals
+					for (int32 EvalIndex = CurrentStateTree->EvaluatorsBegin; EvalIndex < (CurrentStateTree->EvaluatorsBegin + CurrentStateTree->EvaluatorsNum); EvalIndex++)
 			{
-				const FStateTreeEvaluatorBase& Eval =  NextFrame.StateTree->Nodes[EvalIndex].Get<const FStateTreeEvaluatorBase>();
-				const FConstStructView EvalInstanceData = NextFrame.StateTree->DefaultInstanceData.GetStruct(Eval.InstanceTemplateIndex.Get());
+						const FStateTreeEvaluatorBase& Eval = CurrentStateTree->Nodes[EvalIndex].Get<const FStateTreeEvaluatorBase>();
+						const FConstStructView EvalInstanceData = CurrentStateTree->DefaultInstanceData.GetStruct(Eval.InstanceTemplateIndex.Get());
 				InstanceStructs[BaseIndex + Eval.InstanceDataHandle.GetIndex()] = EvalInstanceData;
-				if (!bAreCommon)
+						if (!bIsFrameCommon)
 				{
-					TempInstanceStructs[BaseIndex + Eval.InstanceDataHandle.GetIndex()] = FindInstanceTempData(NextFrame, Eval.InstanceDataHandle);
+							TempInstanceStructs[BaseIndex + Eval.InstanceDataHandle.GetIndex()] = FindInstanceTempData(CurrentFrameID, Eval.InstanceDataHandle);
 				}
 			}
 
 			// Global tasks
-			for (int32 TaskIndex = NextFrame.StateTree->GlobalTasksBegin; TaskIndex < (NextFrame.StateTree->GlobalTasksBegin + NextFrame.StateTree->GlobalTasksNum); TaskIndex++)
+					for (int32 TaskIndex = CurrentStateTree->GlobalTasksBegin; TaskIndex < (CurrentStateTree->GlobalTasksBegin + CurrentStateTree->GlobalTasksNum); TaskIndex++)
 			{
-				const FStateTreeTaskBase& Task =  NextFrame.StateTree->Nodes[TaskIndex].Get<const FStateTreeTaskBase>();
-				const FConstStructView TaskInstanceData = NextFrame.StateTree->DefaultInstanceData.GetStruct(Task.InstanceTemplateIndex.Get());
+						const FStateTreeTaskBase& Task = CurrentStateTree->Nodes[TaskIndex].Get<const FStateTreeTaskBase>();
+						const FConstStructView TaskInstanceData = CurrentStateTree->DefaultInstanceData.GetStruct(Task.InstanceTemplateIndex.Get());
 				InstanceStructs[BaseIndex + Task.InstanceDataHandle.GetIndex()] = TaskInstanceData;
-				if (!bAreCommon)
+						if (!bIsFrameCommon)
 				{
-					TempInstanceStructs[BaseIndex + Task.InstanceDataHandle.GetIndex()] = FindInstanceTempData(NextFrame, Task.InstanceDataHandle);
+							TempInstanceStructs[BaseIndex + Task.InstanceDataHandle.GetIndex()] = FindInstanceTempData(CurrentFrameID, Task.InstanceDataHandle);
 				}
 			}
 
-			if (bAreCommon)
+					if (bIsFrameCommon)
 			{
 				NumCommonInstanceData = InstanceStructs.Num();
 			}
 		}
 
-		// States
-		const int32 BaseIndex = InstanceStructs.Num();
+				CurrentFrameBaseIndex = InstanceStructs.Num();
 
-		NextFrame.GlobalParameterDataHandle = CurrentGlobalParameterDataHandle;
-		NextFrame.GlobalInstanceIndexBase = FStateTreeIndex16(CurrentGlobalInstanceIndexBase);
-		NextFrame.ActiveInstanceIndexBase = FStateTreeIndex16(BaseIndex);
+				CurrentFrame->GlobalParameterDataHandle = CurrentGlobalParameterDataHandle;
+				CurrentFrame->GlobalInstanceIndexBase = FStateTreeIndex16(CurrentGlobalInstanceIndexBase);
+				CurrentFrame->ActiveInstanceIndexBase = FStateTreeIndex16(CurrentFrameBaseIndex);
+			}
 
-		for (int32 StateIndex = 0; StateIndex < NextFrame.ActiveStates.Num(); StateIndex++)
+			// State Params and Nodes
+			if (ensure(CurrentStateTree))
+			{
+				const FStateTreeStateHandle CurrentStateHandle = SelectedState.GetStateHandle();
+				const FCompactStateTreeState* CurrentState = CurrentStateTree->GetStateFromHandle(CurrentStateHandle);
+				if (ensure(CurrentState))
 		{
-			// Check if the next state is still same as current state, GetStateSafe() will return invalid state if passed out of bounds index.
-			bAreCommon = bAreCommon && (CurrentFrame && CurrentFrame->ActiveStates.GetStateSafe(StateIndex) == NextFrame.ActiveStates[StateIndex]);
+					// Find if the state is common (before the transition or sustained from the transition).
+					//The frame may contain valid states (sustained or not). SelectedState will be added later.
+					const bool bIsStateCommon = CurrentFrame->ActiveStates.Contains(SelectedState.GetStateID());
 
-			const FStateTreeStateHandle StateHandle = NextFrame.ActiveStates[StateIndex];
-			const FCompactStateTreeState& State = NextFrame.StateTree->States[StateHandle.Index];
-
-			InstanceStructs.AddDefaulted(State.InstanceDataNum);
-			TempInstanceStructs.AddZeroed(State.InstanceDataNum);
+					InstanceStructs.AddDefaulted(CurrentState->InstanceDataNum);
+					TempInstanceStructs.AddZeroed(CurrentState->InstanceDataNum);
 
 			bool bCanHaveTempData = false;
-			
-			if (State.Type == EStateTreeStateType::Subtree)
+					if (CurrentState->Type == EStateTreeStateType::Subtree)
 			{
-				check(State.ParameterDataHandle.IsValid());
-				check(State.ParameterTemplateIndex.IsValid());
-				const FConstStructView ParamsInstanceData = NextFrame.StateTree->DefaultInstanceData.GetStruct(State.ParameterTemplateIndex.Get());
+						check(CurrentState->ParameterDataHandle.IsValid());
+						check(CurrentState->ParameterTemplateIndex.IsValid());
+						const FConstStructView ParamsInstanceData = CurrentStateTree->DefaultInstanceData.GetStruct(CurrentState->ParameterTemplateIndex.Get());
 				if (!NextStateParameterDataHandle.IsValid())
 				{
 					// Parameters are not set by a linked state, create instance data.
-					InstanceStructs[BaseIndex + State.ParameterDataHandle.GetIndex()] = ParamsInstanceData;
-					NextFrame.StateParameterDataHandle = State.ParameterDataHandle;
+							InstanceStructs[CurrentFrameBaseIndex + CurrentState->ParameterDataHandle.GetIndex()] = ParamsInstanceData;
+							CurrentFrame->StateParameterDataHandle = CurrentState->ParameterDataHandle;
 					bCanHaveTempData = true;
 				}
 				else
@@ -2009,7 +2277,7 @@ void FStateTreeExecutionContext::UpdateInstanceData(TConstArrayView<FStateTreeEx
 					const UStruct* StateParameterDataStruct = Params ? Params->Parameters.GetPropertyBagStruct() : nullptr;
 					check(NextStateParameterDataStruct == StateParameterDataStruct);
 					
-					NextFrame.StateParameterDataHandle = NextStateParameterDataHandle;
+							CurrentFrame->StateParameterDataHandle = NextStateParameterDataHandle;
 					NextStateParameterDataHandle = FStateTreeDataHandle::Invalid; // Mark as used.
 
 					// This state will not instantiate parameter data, so we don't care about the temp data either.
@@ -2018,17 +2286,17 @@ void FStateTreeExecutionContext::UpdateInstanceData(TConstArrayView<FStateTreeEx
 			}
 			else
 			{
-				if (State.ParameterTemplateIndex.IsValid())
+						if (CurrentState->ParameterTemplateIndex.IsValid())
 				{
 					// Linked state's instance data is the parameters.
-					check(State.ParameterDataHandle.IsValid());
+							check(CurrentState->ParameterDataHandle.IsValid());
 
 					const FCompactStateTreeParameters* Params = nullptr;
-					if (FInstancedStruct* TempParamsInstanceData = FindInstanceTempData(NextFrame, State.ParameterDataHandle))
+							if (FInstancedStruct* TempParamsInstanceData = FindInstanceTempData(CurrentFrameID, CurrentState->ParameterDataHandle))
 					{
 						// If we have temp data for the parameters, then setup the instance data with just a type, so that we can steal the temp data below (TempInstanceStructs).
 						// We expect overridden linked assets to hit this code path. 
-						InstanceStructs[BaseIndex + State.ParameterDataHandle.GetIndex()] = FConstStructView(TempParamsInstanceData->GetScriptStruct());
+								InstanceStructs[CurrentFrameBaseIndex + CurrentState->ParameterDataHandle.GetIndex()] = FConstStructView(TempParamsInstanceData->GetScriptStruct());
 						Params = TempParamsInstanceData->GetPtr<const FCompactStateTreeParameters>();
 						bCanHaveTempData = true;
 					}
@@ -2036,64 +2304,72 @@ void FStateTreeExecutionContext::UpdateInstanceData(TConstArrayView<FStateTreeEx
 					{
 						// If not temp data, use the states or linked assets default values.
 						FConstStructView ParamsInstanceData;
-						if (State.Type == EStateTreeStateType::LinkedAsset)
+								if (CurrentState->Type == EStateTreeStateType::LinkedAsset)
 						{
-							// This state is a container for the linked state tree.Its instance data matches the linked state tree parameters.The linked state tree asset is the next frame.
-							const bool bIsLastFrame = FrameIndex == NextActiveFrames.Num() - 1;
+									// This state is a container for the linked state tree.
+									// Its instance data matches the linked state tree parameters.
+									// The linked state tree asset is the next frame.
+									const bool bIsLastFrame = SelectedStateIndex == Args.SelectedStates.Num() - 1;
 							if (!bIsLastFrame)
 							{
-								FStateTreeExecutionFrame& FollowingNextFrame = NextActiveFrames[FrameIndex + 1];
-								ParamsInstanceData = FConstStructView::Make(TempParams.Emplace_GetRef(FollowingNextFrame.StateTree->GetDefaultParameters()));
+										const FActiveState& FollowingSelectedState = Args.SelectedStates[SelectedStateIndex + 1];
+										const FStateTreeExecutionFrame* FollowingNextFrame = FindExecutionFrame(FollowingSelectedState.GetFrameID(), MakeArrayView(Exec.ActiveFrames), MakeArrayView(Args.TemporaryFrames));
+										if (ensure(FollowingNextFrame))
+										{
+											ParamsInstanceData = FConstStructView::Make(TempParams.Emplace_GetRef(FollowingNextFrame->StateTree->GetDefaultParameters()));
+										}
 							}
 						}
 						if (!ParamsInstanceData.IsValid())
 						{
-							ParamsInstanceData = NextFrame.StateTree->DefaultInstanceData.GetStruct(State.ParameterTemplateIndex.Get());
+									ParamsInstanceData = CurrentStateTree->DefaultInstanceData.GetStruct(CurrentState->ParameterTemplateIndex.Get());
 						}
-						InstanceStructs[BaseIndex + State.ParameterDataHandle.GetIndex()] = ParamsInstanceData;
+								InstanceStructs[CurrentFrameBaseIndex + CurrentState->ParameterDataHandle.GetIndex()] = ParamsInstanceData;
 						Params = ParamsInstanceData.GetPtr<const FCompactStateTreeParameters>();
 						bCanHaveTempData = true;
 					}
 
-					if (State.Type == EStateTreeStateType::Linked
-						|| State.Type == EStateTreeStateType::LinkedAsset)
+							if (CurrentState->Type == EStateTreeStateType::Linked
+								|| CurrentState->Type == EStateTreeStateType::LinkedAsset)
 					{
 						// Store the index of the parameter data, so that we can point the linked state to it.
-						check(State.ParameterDataHandle.GetSource() == EStateTreeDataSourceType::StateParameterData);
+								check(CurrentState->ParameterDataHandle.GetSource() == EStateTreeDataSourceType::StateParameterData);
 						checkf(!NextStateParameterDataHandle.IsValid(), TEXT("NextStateParameterDataIndex not should be set yet when we encounter a linked state."));
-						NextStateParameterDataHandle = State.ParameterDataHandle;
+								NextStateParameterDataHandle = CurrentState->ParameterDataHandle;
 						NextStateParameterDataStruct = Params ? Params->Parameters.GetPropertyBagStruct() : nullptr;
 					}
 				}
 			}
 			
-			if (!bAreCommon && bCanHaveTempData)
+					if (!bIsStateCommon && bCanHaveTempData)
 			{
-				TempInstanceStructs[BaseIndex + State.ParameterDataHandle.GetIndex()] = FindInstanceTempData(NextFrame, State.ParameterDataHandle);
+						TempInstanceStructs[CurrentFrameBaseIndex + CurrentState->ParameterDataHandle.GetIndex()] = FindInstanceTempData(CurrentFrameID, CurrentState->ParameterDataHandle);
 			}
 
-			if (State.EventDataIndex.IsValid())
+					if (CurrentState->EventDataIndex.IsValid())
 			{
-				InstanceStructs[BaseIndex + State.EventDataIndex.Get()] = FConstStructView(FStateTreeSharedEvent::StaticStruct());
+						InstanceStructs[CurrentFrameBaseIndex + CurrentState->EventDataIndex.Get()] = FConstStructView(FStateTreeSharedEvent::StaticStruct());
 			}
 
-			for (int32 TaskIndex = State.TasksBegin; TaskIndex < (State.TasksBegin + State.TasksNum); TaskIndex++)
+					for (int32 TaskIndex = CurrentState->TasksBegin; TaskIndex < (CurrentState->TasksBegin + CurrentState->TasksNum); TaskIndex++)
 			{
-				const FStateTreeTaskBase& Task = NextFrame.StateTree->Nodes[TaskIndex].Get<const FStateTreeTaskBase>();
-				const FConstStructView TaskInstanceData = NextFrame.StateTree->DefaultInstanceData.GetStruct(Task.InstanceTemplateIndex.Get());
-				InstanceStructs[BaseIndex + Task.InstanceDataHandle.GetIndex()] = TaskInstanceData;
-				if (!bAreCommon)
+						const FStateTreeTaskBase& Task = CurrentStateTree->Nodes[TaskIndex].Get<const FStateTreeTaskBase>();
+						const FConstStructView TaskInstanceData = CurrentStateTree->DefaultInstanceData.GetStruct(Task.InstanceTemplateIndex.Get());
+						InstanceStructs[CurrentFrameBaseIndex + Task.InstanceDataHandle.GetIndex()] = TaskInstanceData;
+						if (!bIsStateCommon)
 				{
-					TempInstanceStructs[BaseIndex + Task.InstanceDataHandle.GetIndex()] = FindInstanceTempData(NextFrame, Task.InstanceDataHandle);
+							TempInstanceStructs[CurrentFrameBaseIndex + Task.InstanceDataHandle.GetIndex()] = FindInstanceTempData(CurrentFrameID, Task.InstanceDataHandle);
 				}
 			}
 
-			if (bAreCommon)
+					if (bIsStateCommon)
 			{
 				NumCommonInstanceData = InstanceStructs.Num();
 			}
 		}
 	}
+		}
+	}
 	
 	// Common section should match.
 #if WITH_STATETREE_DEBUG
@@ -2153,34 +2429,36 @@ FStateTreeDataView FStateTreeExecutionContext::GetDataView(const FStateTreeExecu
 	case EStateTreeDataSourceType::StateEvent:
 		{
 			// If state selection is going, return FStateTreeEvent of the event currently captured by the state selection.
-			if (CurrentlyProcessedStateSelectionEvents)
+		if (CurrentlyProcessedStateSelectionResult)
 			{
-				if (const FCompactStateTreeState* State = CurrentFrame.StateTree->GetStateFromHandle(Handle.GetState()))
+			FSelectionEventWithID* FoundEvent = CurrentlyProcessedStateSelectionResult->SelectionEvents.FindByPredicate(
+				[FrameID = CurrentFrame.FrameID, StateHandle = Handle.GetState()]
+				(const FSelectionEventWithID& Event)
+				{
+					return Event.State.GetStateHandle() == StateHandle && Event.State.GetFrameID() == FrameID;
+				});
+			if (FoundEvent)
 				{
 					// Events are read only, but we cannot express that in FStateTreeDataView.
-					if (FStateTreeEvent* Event = CurrentlyProcessedStateSelectionEvents->Events[State->Depth].GetMutable())
-					{
-						return FStateTreeDataView(FStructView::Make(*Event));
+				return FStateTreeDataView(FStructView::Make(*FoundEvent->Event.GetMutable()));
 					}
 				}
 
-				return {};
-			}
-
+		if (UE::StateTree::InstanceData::Private::IsActiveInstanceHandleSourceValid(Storage, CurrentFrame, Handle))
+		{
 			return UE::StateTree::InstanceData::GetDataView(Storage, CurrentlyProcessedSharedInstanceStorage, ParentFrame, CurrentFrame, Handle);
+			}
 		}
 
 	case EStateTreeDataSourceType::ExternalGlobalParameterData:
 		{
 			checkf(false, TEXT("External global parameter data currently not supported for linked state-trees"));
-			break;
+		return nullptr;
 		}
 
 	default:
 		return UE::StateTree::InstanceData::GetDataView(Storage, CurrentlyProcessedSharedInstanceStorage, ParentFrame, CurrentFrame, Handle);
 	}
-
-	return {};
 }
 
 FStateTreeDataView FStateTreeExecutionContext::GetDataView(const FStateTreeExecutionFrame* ParentFrame, const FStateTreeExecutionFrame& CurrentFrame, const FPropertyBindingCopyInfo& CopyInfo)
@@ -2194,8 +2472,14 @@ FStateTreeDataView FStateTreeExecutionContext::GetDataView(const FStateTreeExecu
 	return GetDataView(ParentFrame, CurrentFrame, Handle);
 }
 
+//@todo should return TOptional or TValueOrError. We should not trust what is coming from FRecordedStateTreeTransitionResult.
+//@todo recorded should use something else to identify the frame and state. We should not thrust what is coming from FRecordedStateTreeTransitionResult.
 EStateTreeRunStatus FStateTreeExecutionContext::ForceTransition(const FRecordedStateTreeTransitionResult& Transition)
 {
+	using namespace UE::StateTree;
+	using namespace UE::StateTree::ExecutionContext;
+	using namespace UE::StateTree::ExecutionContext::Private;
+
 	if (!IsValid())
 	{
 		STATETREE_LOG(Warning, TEXT("%hs: StateTree context is not initialized properly ('%s' using StateTree '%s')"),
@@ -2203,28 +2487,93 @@ EStateTreeRunStatus FStateTreeExecutionContext::ForceTransition(const FRecordedS
 		return EStateTreeRunStatus::Failed;
 	}
 
+	const FStateTreeExecutionState& Exec = GetExecState();
+
 	// A reentrant call to ForceTransition or a call from Start, Tick or Stop must be deferred.
-	if (GetExecState().CurrentPhase != EStateTreeUpdatePhase::Unset)
+	if (Exec.CurrentPhase != EStateTreeUpdatePhase::Unset)
 	{
-		UE_STATETREE_DEBUG_LOG_EVENT(this, Warning, TEXT("Can't force a transition while %s"), *UEnum::GetDisplayValueAsText(GetExecState().CurrentPhase).ToString());
+		UE_STATETREE_DEBUG_LOG_EVENT(this, Warning, TEXT("Can't force a transition while %s"), *UEnum::GetDisplayValueAsText(Exec.CurrentPhase).ToString());
 		return EStateTreeRunStatus::Unset;
 	}
 
-	TOptional<FStateTreeTransitionResult> TransitionResult = MakeTransitionResult(Transition);
-	if (!TransitionResult.IsSet())
+	if (!CollectActiveExternalData())
+	{
+		STATETREE_LOG(Warning, TEXT("%hs: Failed to collect external data ('%s' using StateTree '%s')"),
+			__FUNCTION__, *GetNameSafe(&Owner), *GetFullNameSafe(&RootStateTree));
+		return EStateTreeRunStatus::Failed;
+	}
+
+	TOptional<FStateTreeTransitionResult> FoundTransitionResult = MakeTransitionResult(Transition);
+	if (!FoundTransitionResult.IsSet())
 	{
 		return EStateTreeRunStatus::Unset;
 	}
+	FStateTreeTransitionResult& TransitionResult = FoundTransitionResult.GetValue();
+
+	// Create the state argument
+	FSelectStateResult Args;
+
+	int32 CurrentActiveFrameIndex = 0;
+	for (FStateTreeExecutionFrame& Frame : TransitionResult.NextActiveFrames)
+	{
+		// reuse or new temporary
+		if (Exec.ActiveFrames.IsValidIndex(CurrentActiveFrameIndex)
+			&& Exec.ActiveFrames[CurrentActiveFrameIndex].HasSameRoot(Frame))
+		{
+			Frame.FrameID = Exec.ActiveFrames[CurrentActiveFrameIndex].FrameID;
+			++CurrentActiveFrameIndex;
+		}
+		else
+		{
+			Frame.FrameID = FActiveFrameID(Storage.GenerateUniqueId());
+			Args.TemporaryFrames.Add_GetRef(Frame).ActiveStates = FStateTreeActiveStates();
+			CurrentActiveFrameIndex = INDEX_NONE;
+		}
+
+		Args.SelectedFrames.Add(Frame.FrameID);
+	}
+	{
+		TArray<FActiveState, TInlineAllocator<ExpectedActiveStatePathLength>> NewActiveStates;
+		GetActiveStatePath(TransitionResult.NextActiveFrames, NewActiveStates);
+		Args.SelectedStates.Append(NewActiveStates);
+	}
+
+	int32 SelectedStateIndex = 0;
+	for (FStateTreeFrameStateSelectionEvents& SelectionEvent : TransitionResult.NextActiveFrameEvents)
+	{
+		for (FStateTreeSharedEvent& Event : SelectionEvent.Events)
+		{
+			// Event should be valid but they might come from replication and we don't thrust anything that comes from the network.
+			if (Event.IsValid())
+			{
+				if (Args.SelectedStates.IsValidIndex(SelectedStateIndex))
+				{
+					Args.SelectionEvents.Emplace(Args.SelectedStates[SelectedStateIndex], Event);
+				}
+			}
+		}
+		++SelectedStateIndex;
+	}
+
+	const FActiveState* FoundActiveState = Args.SelectedStates.FindByPredicate([&TransitionResult](const FActiveState& ActiveState)
+		{
+			return (ActiveState.GetFrameID() == TransitionResult.SourceFrameID
+				&& ActiveState.GetStateHandle() == TransitionResult.TargetState);
+		});
+	if (FoundActiveState != nullptr)
+	{
+		Args.TargetState = *FoundActiveState;
+	}
 
-	ExitState(TransitionResult.GetValue());
-	return EnterState(TransitionResult.GetValue());
+	ExitState(Args, TransitionResult);
+	return EnterState(Args, TransitionResult);
 }
 
 const FStateTreeExecutionFrame* FStateTreeExecutionContext::FindFrame(const UStateTree* StateTree, FStateTreeStateHandle RootState, TConstArrayView<FStateTreeExecutionFrame> Frames, const FStateTreeExecutionFrame*& OutParentFrame)
 {
-	const int32 FrameIndex = Frames.IndexOfByPredicate([&StateTree, RootState](const FStateTreeExecutionFrame& Frame)
+	const int32 FrameIndex = Frames.IndexOfByPredicate([StateTree, RootState](const FStateTreeExecutionFrame& Frame)
 	{
-		return Frame.StateTree == StateTree && Frame.RootState == RootState;
+		return Frame.HasRoot(StateTree, RootState);
 	});
 
 	if (FrameIndex == INDEX_NONE)
@@ -2259,10 +2608,8 @@ bool FStateTreeExecutionContext::IsHandleSourceValid(const FStateTreeExecutionFr
 		return CurrentlyProcessedTransitionEvent != nullptr;
 
 	case EStateTreeDataSourceType::StateEvent:
-		return CurrentlyProcessedStateSelectionEvents != nullptr
-			|| (CurrentFrame.ActiveInstanceIndexBase.IsValid()
-			&& CurrentFrame.ActiveStates.Contains(Handle.GetState())
-			&& Storage.IsValidIndex(CurrentFrame.ActiveInstanceIndexBase.Get() + Handle.GetIndex()));
+		return CurrentlyProcessedStateSelectionResult != nullptr
+			|| UE::StateTree::InstanceData::Private::IsActiveInstanceHandleSourceValid(Storage, CurrentFrame, Handle);
 
 	case EStateTreeDataSourceType::ExternalGlobalParameterData:
 	{
@@ -2594,11 +2941,21 @@ bool FStateTreeExecutionContext::SetGlobalParameters(const FInstancedPropertyBag
 	return false;
 }
 
+// Deprecated
 void FStateTreeExecutionContext::CaptureNewStateEvents(TConstArrayView<FStateTreeExecutionFrame> PrevFrames, TConstArrayView<FStateTreeExecutionFrame> NewFrames, TArrayView<FStateTreeFrameStateSelectionEvents> FramesStateSelectionEvents)
 {
+}
+
+void FStateTreeExecutionContext::CaptureNewStateEvents(FSelectStateResult& Args)
+{
+	using namespace UE::StateTree;
+	using namespace UE::StateTree::ExecutionContext::Private;
+
+	FStateTreeExecutionState& Exec = GetExecState();
+
 	// Mark the events from delayed transitions as in use, so that each State will receive unique copy of the event struct. 
 	TArray<FStateTreeSharedEvent, TInlineAllocator<16>> EventsInUse;
-	for (const FStateTreeTransitionDelayedState& DelayedTransition : GetExecState().DelayedTransitions)
+	for (const FStateTreeTransitionDelayedState& DelayedTransition : Exec.DelayedTransitions)
 	{
 		if (DelayedTransition.CapturedEvent.IsValid())
 		{
@@ -2606,41 +2963,39 @@ void FStateTreeExecutionContext::CaptureNewStateEvents(TConstArrayView<FStateTre
 		}
 	}
 
-	for (int32 FrameIndex = 0; FrameIndex < NewFrames.Num(); ++FrameIndex)
+	// For each state that are changed (not sustained)
+	FActiveFrameID CurrentFrameID;
+	const FStateTreeExecutionFrame* CurrentFrame = nullptr;
+	for (int32 SelectedStateIndex = 0; SelectedStateIndex < Args.SelectedStates.Num(); ++SelectedStateIndex)
 	{
-		const FStateTreeExecutionFrame& NewFrame = NewFrames[FrameIndex];
+		const FActiveState& SelectedState = Args.SelectedStates[SelectedStateIndex];
 
-		// Find states that are unique to the new frame.
-		TConstArrayView<FStateTreeStateHandle> UniqueStates = NewFrame.ActiveStates;
-		if (PrevFrames.IsValidIndex(FrameIndex))
+		// Global
+		if (SelectedState.GetFrameID() != CurrentFrameID)
 		{
-			const FStateTreeExecutionFrame& PrevFrame = PrevFrames[FrameIndex];
-
-			if (PrevFrame.FrameID == NewFrame.FrameID)
-			{
-			checkf(PrevFrame.RootState == NewFrame.RootState && PrevFrame.StateTree == NewFrame.StateTree, TEXT("If the Id matches, then the root and the tree must also match."));
+			CurrentFrameID = SelectedState.GetFrameID();
+			CurrentFrame = FindExecutionFrame(SelectedState.GetFrameID(), MakeConstArrayView(Exec.ActiveFrames), MakeConstArrayView(Args.TemporaryFrames));
+		}
 
-			for (int32 StateIndex = 0; StateIndex < NewFrame.ActiveStates.Num(); ++StateIndex)
+		check(CurrentFrame);
+		const bool bIsStateCommon = CurrentFrame->ActiveStates.Contains(SelectedState.GetStateID());
+		if (bIsStateCommon)
 			{
-				if (!PrevFrame.ActiveStates.IsValidIndex(StateIndex) || PrevFrame.ActiveStates.StateIDs[StateIndex] != NewFrame.ActiveStates.StateIDs[StateIndex])
-				{
-					UniqueStates = TConstArrayView<FStateTreeStateHandle>(&NewFrame.ActiveStates[StateIndex], NewFrame.ActiveStates.Num() - StateIndex);
-					break;
-				}
-			}
-		}
+			continue;
 		}
 
-		// Capture events for the new states.
-		for (const FStateTreeStateHandle StateHandle : UniqueStates)
+		if (const FCompactStateTreeState* State = CurrentFrame->StateTree->GetStateFromHandle(SelectedState.GetStateHandle()))
+			{
+			if (State->EventDataIndex.IsValid())
 		{
-			if (const FCompactStateTreeState* State = NewFrame.StateTree->GetStateFromHandle(StateHandle))
+				FStateTreeSharedEvent& StateTreeEvent = Storage.GetMutableStruct(CurrentFrame->ActiveInstanceIndexBase.Get() + State->EventDataIndex.Get()).Get<FStateTreeSharedEvent>();
+				const FSelectionEventWithID* EventToCapturePtr = Args.SelectionEvents.FindByPredicate([SelectedState](const FSelectionEventWithID& EventID)
 			{
-				if (State->EventDataIndex.IsValid())
+						return EventID.State == SelectedState;
+					});
+				if (ensureAlways(EventToCapturePtr))
 				{
-					FStateTreeSharedEvent& StateTreeEvent = Storage.GetMutableStruct(NewFrame.ActiveInstanceIndexBase.Get() + State->EventDataIndex.Get()).Get<FStateTreeSharedEvent>();
-					
-					const FStateTreeSharedEvent& EventToCapture = FramesStateSelectionEvents[FrameIndex].Events[State->Depth];
+					const FStateTreeSharedEvent& EventToCapture = EventToCapturePtr->Event;
 					if (EventsInUse.Contains(EventToCapture))
 					{
 						// Event is already spoken for, make a copy.
@@ -2649,7 +3004,7 @@ void FStateTreeExecutionContext::CaptureNewStateEvents(TConstArrayView<FStateTre
 					else
 					{
 						// Event not in use, steal it.
-						StateTreeEvent = FramesStateSelectionEvents[FrameIndex].Events[State->Depth];
+						StateTreeEvent = EventToCapture;
 						EventsInUse.Add(EventToCapture);
 					}
 				}
@@ -2658,11 +3013,27 @@ void FStateTreeExecutionContext::CaptureNewStateEvents(TConstArrayView<FStateTre
 	}
 }
 
+// Deprecated
 EStateTreeRunStatus FStateTreeExecutionContext::EnterState(FStateTreeTransitionResult& Transition)
 {
+	return EStateTreeRunStatus::Failed;
+}
+
+EStateTreeRunStatus FStateTreeExecutionContext::EnterState(FSelectStateResult& Args, const FStateTreeTransitionResult& Transition)
+{
+	// 1. Update the data instance data of all frames and selected states.
+	// The data won't be available until the node is reached or the state is added to the FStateTreeExecutionFrame::ActiveState.
+	// 2. Capture StateEvents
+	// 3. Call Enter state for each selected state from Target.
+	//Note, no need to update the bindings for previously active state, the bindings are updated in ExitState or there was no binding before (See Start).
+
+	using namespace UE::StateTree;
+	using namespace UE::StateTree::ExecutionContext;
+	using namespace UE::StateTree::ExecutionContext::Private;
+
 	CSV_SCOPED_TIMING_STAT_EXCLUSIVE(StateTree_EnterState);
 
-	if (Transition.NextActiveFrames.IsEmpty())
+	if (Args.SelectedStates.IsEmpty())
 	{
 		return EStateTreeRunStatus::Failed;
 	}
@@ -2674,131 +3045,157 @@ EStateTreeRunStatus FStateTreeExecutionContext::EnterState(FStateTreeTransitionR
 		RecordedTransitions.Add(MakeRecordedTransitionResult(Transition));
 	}
 
-	// Allocate new tasks.
-	UpdateInstanceData(Exec.ActiveFrames, Transition.NextActiveFrames);
+	UpdateInstanceData(Args);
+	CaptureNewStateEvents(Args);
 
-	CaptureNewStateEvents(Exec.ActiveFrames, Transition.NextActiveFrames, Transition.NextActiveFrameEvents);
+	++Exec.StateChangeCount;
 
-	Exec.StateChangeCount++;
-	
-	// On target branch means that the state is the target of current transition or child of it.
-	// States which were active before and will remain active, but are not on target branch will not get
-	// EnterState called. That is, a transition is handled as "replan from this state".
-	bool bOnTargetBranch = false;
 	FStateTreeTransitionResult CurrentTransition = Transition;
 	EStateTreeRunStatus Result = EStateTreeRunStatus::Running;
 
-	STATETREE_LOG(Log, TEXT("Enter state '%s' (%d)"), *DebugGetStatePath(Transition.NextActiveFrames), Exec.StateChangeCount);
+	STATETREE_LOG(Log, TEXT("Enter state '%s' (%d)")
+		, *UE::StateTree::ExecutionContext::Private::GetStatePathAsString(&RootStateTree, Args.SelectedStates)
+		, Exec.StateChangeCount
+	);
 	UE_STATETREE_DEBUG_ENTER_PHASE(this, EStateTreeUpdatePhase::EnterStates);
 
-	// The previous active frames are needed for state enter logic.
-	TArray<FStateTreeExecutionFrame> PreviousActiveFrames;
-	PreviousActiveFrames = Exec.ActiveFrames;
-
-	// Reset the current active frames, new ones are added one by one.
-	Exec.ActiveFrames.Reset();
-
-	// Track any changed state (i.e., not sustained) to prevent reused subtrees from being considered as sustained states since the whole
-	// tree is reused and should be considered as a new instance.
-	bool bAnyParentStateChanged = false;
-
-	for (int32 FrameIndex = 0; FrameIndex < Transition.NextActiveFrames.Num() && Result != EStateTreeRunStatus::Failed; FrameIndex++)
+	bool bTargetReached = false;
+	FActiveFrameID CurrentFrameID;
+	const TArrayView<const UE::StateTree::FActiveState> SelectedStates = Args.SelectedStates;
+	for (int32 SelectedStateIndex = 0; SelectedStateIndex < SelectedStates.Num(); ++SelectedStateIndex)
 	{
-		const FStateTreeExecutionFrame& NextFrame = Transition.NextActiveFrames[FrameIndex];
-		
-		FStateTreeExecutionFrame* CurrentParentFrame = !Exec.ActiveFrames.IsEmpty() ? &Exec.ActiveFrames.Last() : nullptr;
-		FStateTreeExecutionFrame& CurrentFrame = Exec.ActiveFrames.Add_GetRef(NextFrame);
-		const UStateTree* CurrentStateTree = NextFrame.StateTree;
-		const UE::StateTree::FActiveFrameID CurrentFrameID = NextFrame.FrameID;
-		
-		// We'll add new states one by one, so that active states contain only the states which have EnterState called.
-		CurrentFrame.ActiveStates.Reset();
+		const UE::StateTree::FActiveState& SelectedState = SelectedStates[SelectedStateIndex];
 
-		if (!ensureMsgf(CurrentFrame.ActiveTasksStatus.IsValid(), TEXT("Frame is not formed correct.")))
+		const bool bIsCurrentFrameNew = CurrentFrameID != SelectedState.GetFrameID();
+		CurrentFrameID = SelectedState.GetFrameID();
+
+		bool bFrameAdded = false;
+		int32 FrameIndex = Exec.IndexOfActiveFrame(SelectedState.GetFrameID());
+		if (FrameIndex == INDEX_NONE)
 		{
-			// Create the status. We lost the previous tasks' completion status.
-			const FCompactStateTreeFrame* FrameInfo = NextFrame.StateTree->GetFrameFromHandle(NextFrame.RootState);
+			check(bIsCurrentFrameNew);
+			FStateTreeExecutionFrame* FoundNewFrame = Args.TemporaryFrames.FindByPredicate(
+				[FrameID = SelectedState.GetFrameID()](const FStateTreeExecutionFrame& Frame)
+				{
+					return Frame.FrameID == FrameID;
+				});
+			if (!ensure(FoundNewFrame != nullptr))
+			{
+				return EStateTreeRunStatus::Failed;
+			}
+
+			// Add it to the active list
+			FrameIndex = Exec.ActiveFrames.Add(MoveTemp(*FoundNewFrame));
+			bFrameAdded = true;
+		}
+
+		FStateTreeExecutionFrame& CurrentFrame = Exec.ActiveFrames[FrameIndex];
+
+		int32 ActiveStateIndex = CurrentFrame.ActiveStates.IndexOfReverse(SelectedState.GetStateID());
+		// The state is the target of current transition or is a sustained state
+		const bool bAreCommon = ActiveStateIndex != INDEX_NONE;
+		bTargetReached = bTargetReached || SelectedState == Args.TargetState;
+		const bool bSustainedState = bAreCommon && bTargetReached;
+
+		// States which were active before and will remain active, but are not on target branch will not get
+		// EnterState called. That is, a transition is handled as "replan from this state".
+		if (bAreCommon && !bSustainedState)
+	{
+			// Is already in the active and we do not need to call EnterState with Sustained
+			continue;
+		}
+		
+		FStateTreeExecutionFrame* CurrentParentFrame = FrameIndex > 0 ? &Exec.ActiveFrames[FrameIndex - 1] : nullptr;
+		const UStateTree* CurrentStateTree = CurrentFrame.StateTree;
+
+		// Create the status. We lost the previous tasks' completion status.
+		if (!ensureMsgf(CurrentFrame.ActiveTasksStatus.IsValid(), TEXT("Frame is not formed correct.")))
+		{
+			const FCompactStateTreeFrame* FrameInfo = CurrentFrame.StateTree->GetFrameFromHandle(CurrentFrame.RootState);
 			ensureAlwaysMsgf(FrameInfo, TEXT("The compiled data is invalid. It should contains the information for the root frame."));
 			CurrentFrame.ActiveTasksStatus = FrameInfo ? FStateTreeTasksCompletionStatus(*FrameInfo) : FStateTreeTasksCompletionStatus();
 		}
 
-		// Get previous active states, they are used to calculate transition type.
-		FStateTreeActiveStates PreviousActiveStates;
-		if (PreviousActiveFrames.IsValidIndex(FrameIndex)
-			&& PreviousActiveFrames[FrameIndex].IsSameFrame(NextFrame))
+		if (bIsCurrentFrameNew && !bFrameAdded && !CurrentFrame.bHaveEntered)
 		{
-			PreviousActiveStates = PreviousActiveFrames[FrameIndex].ActiveStates;
+			// EnterState on Changed global was called during selection with StartTemporaryEvaluatorsAndGlobalTasks
+			CurrentTransition.CurrentState = FStateTreeStateHandle::Invalid;
+			CurrentTransition.ChangeType = EStateTreeStateChangeType::Sustained;
+
+			EStateTreeRunStatus GlobalTaskRunStatus = StartGlobalsForFrameOnActiveInstances(CurrentParentFrame, CurrentFrame, CurrentTransition);
+			Result = GetPriorityRunStatus(Result, GlobalTaskRunStatus);
+			if (Result == EStateTreeRunStatus::Failed)
+			{
+				break;
+			}
 		}
 
 		FCurrentlyProcessedFrameScope FrameScope(*this, CurrentParentFrame, CurrentFrame);
-
-		for (int32 StateIndex = 0; StateIndex < NextFrame.ActiveStates.Num() && Result != EStateTreeRunStatus::Failed; StateIndex++)
 		{
-			const FStateTreeStateHandle CurrentHandle = NextFrame.ActiveStates[StateIndex];
-			const UE::StateTree::FActiveStateID CurrentStateID = NextFrame.ActiveStates.StateIDs[StateIndex];
-			const FStateTreeStateHandle PreviousHandle = PreviousActiveStates.GetStateSafe(StateIndex);
-			const FCompactStateTreeState& State = CurrentStateTree->States[CurrentHandle.Index];
+			const FStateTreeStateHandle CurrentStateHandle = SelectedState.GetStateHandle();
+			const FCompactStateTreeState& CurrentState = CurrentStateTree->States[CurrentStateHandle.Index];
+			const UE::StateTree::FActiveStateID CurrentStateID = SelectedState.GetStateID();
 
 			// Add only enabled States to the list of active States
-			if (State.bEnabled && !CurrentFrame.ActiveStates.Push(CurrentHandle, NextFrame.ActiveStates.StateIDs[StateIndex]))
-			{
-				STATETREE_LOG(Error, TEXT("%hs: Reached max execution depth when trying to enter state '%s'.  '%s' using StateTree '%s'."),
-					__FUNCTION__, *GetStateStatusString(Exec), *GetNameSafe(&Owner), *GetFullNameSafe(&RootStateTree));
+			if (!CurrentState.bEnabled)
+			{
+				ensureMsgf(false, TEXT("%hs: A selected state is not enabled. '%s' should not be in the selected list. '%s' using StateTree '%s'."),
+					__FUNCTION__,
+					*GetStateStatusString(Exec),
+					*GetNameSafe(&Owner),
+					*GetFullNameSafe(&RootStateTree)
+				);
 				break;
 			}
-			//@todo push the same state as previously handle
-			CurrentFrame.ActiveTasksStatus.Push(State);
 
-			UE::StateTree::FTasksCompletionStatus CurrentStateTasksStatus = CurrentFrame.ActiveTasksStatus.GetStatus(State);
-			CurrentFrame.NumCurrentlyActiveStates = static_cast<uint8>(CurrentFrame.ActiveStates.Num());
-
-			FCurrentlyProcessedStateScope StateScope(*this, CurrentHandle);
-
-			if (State.Type == EStateTreeStateType::Linked
-				|| State.Type == EStateTreeStateType::LinkedAsset)
+			// New state. Add it
+			if (!bSustainedState)
 			{
-				if (State.ParameterDataHandle.IsValid()
-					&& State.ParameterBindingsBatch.IsValid())
+				CurrentFrame.ActiveTasksStatus.Push(CurrentState);
+				if (!CurrentFrame.ActiveStates.Push(CurrentStateHandle, CurrentStateID))
 				{
-					const FStateTreeDataView StateParamsDataView = GetDataView(CurrentParentFrame, CurrentFrame, State.ParameterDataHandle);
-					CopyBatchOnActiveInstances(CurrentParentFrame, CurrentFrame, StateParamsDataView, State.ParameterBindingsBatch);
+					// The ActiveStates array supports a max of 8 states. The depth is verified at compilation.
+					ensureMsgf(false, TEXT("Reached max execution depth when trying to enter state '%s'. '%s' using StateTree '%s'."),
+						*GetStateStatusString(Exec),
+						*GetNameSafe(&Owner),
+						*GetFullNameSafe(&RootStateTree)
+					);
+					break;
 				}
+				ActiveStateIndex = CurrentFrame.ActiveStates.Num() - 1;
 			}
 
-			if (CurrentFrame.FrameID == Transition.SourceFrameID
-				&& CurrentHandle == Transition.TargetState)
-			{
-				bOnTargetBranch = true;
-			}
-
-			const bool bWasActive = (PreviousHandle == CurrentHandle) && !bAnyParentStateChanged;
+			FCurrentlyProcessedStateScope StateScope(*this, CurrentStateHandle);
 
-			const EStateTreeStateChangeType ChangeType = bWasActive ? EStateTreeStateChangeType::Sustained : EStateTreeStateChangeType::Changed;
-			if (ChangeType == EStateTreeStateChangeType::Changed)
+			if (CurrentState.Type == EStateTreeStateType::Linked
+				|| CurrentState.Type == EStateTreeStateType::LinkedAsset)
 			{
-				bAnyParentStateChanged = true;
+				if (CurrentState.ParameterDataHandle.IsValid()
+					&& CurrentState.ParameterBindingsBatch.IsValid())
+				{
+					const FStateTreeDataView StateParamsDataView = GetDataView(CurrentParentFrame, CurrentFrame, CurrentState.ParameterDataHandle);
+					CopyBatchOnActiveInstances(CurrentParentFrame, CurrentFrame, StateParamsDataView, CurrentState.ParameterBindingsBatch);
+			}
 			}
 
-			CurrentTransition.CurrentState = CurrentHandle;
+			const EStateTreeStateChangeType ChangeType = bSustainedState ? EStateTreeStateChangeType::Sustained : EStateTreeStateChangeType::Changed;
+			CurrentTransition.CurrentState = CurrentStateHandle;
 			CurrentTransition.ChangeType = ChangeType;
 
-			// Do not enter a disabled State tasks but maintain property bindings
-			const bool bIsEnteringState = (!bWasActive || bOnTargetBranch) && State.bEnabled;
-
-			if (bIsEnteringState)
-			{
-				UE_STATETREE_DEBUG_STATE_EVENT(this, CurrentHandle, EStateTreeTraceEventType::OnEntering);
-				STATETREE_LOG(Log, TEXT("%*sState '%s' (%s)"), (FrameIndex + StateIndex + 1)*UE::StateTree::Debug::IndentSize, TEXT("")
-					, *GetSafeStateName(NextFrame, CurrentHandle)
-					, *UEnum::GetDisplayValueAsText(CurrentTransition.ChangeType).ToString());
-			}
+			UE_STATETREE_DEBUG_STATE_EVENT(this, CurrentStateHandle, EStateTreeTraceEventType::OnEntering);
+			STATETREE_LOG(Log, TEXT("%*sState '%s' (%s)"),
+				(FrameIndex + ActiveStateIndex + 1) * UE::StateTree::Debug::IndentSize, TEXT(""),
+				*GetSafeStateName(CurrentStateTree, CurrentStateHandle),
+				*UEnum::GetDisplayValueAsText(CurrentTransition.ChangeType).ToString()
+			);
 
 			// Call state change events on conditions if needed.
-			if (bIsEnteringState && State.bHasStateChangeConditions)
+			if (CurrentState.bHasStateChangeConditions)
 			{
-				for (int32 ConditionIndex = State.EnterConditionsBegin; ConditionIndex < (State.EnterConditionsBegin + State.EnterConditionsNum); ConditionIndex++)
+				const int32 EnterConditionsEnd = CurrentState.EnterConditionsBegin + CurrentState.EnterConditionsNum;
+				for (int32 ConditionIndex = CurrentState.EnterConditionsBegin; ConditionIndex < EnterConditionsEnd; ++ConditionIndex)
 				{
-					const FStateTreeConditionBase& Cond = CurrentFrame.StateTree->Nodes[ConditionIndex].Get<const FStateTreeConditionBase>();
+					const FStateTreeConditionBase& Cond = CurrentStateTree->Nodes[ConditionIndex].Get<const FStateTreeConditionBase>();
 					if (Cond.bHasShouldCallStateChangeEvents)
 					{
 						const bool bShouldCallEnterState = CurrentTransition.ChangeType == EStateTreeStateChangeType::Changed
@@ -2816,7 +3213,7 @@ EStateTreeRunStatus FStateTreeExecutionContext::EnterState(FStateTreeTransitionR
 							}
 
 							UE_STATETREE_DEBUG_CONDITION_ENTER_STATE(this, CurrentFrame.StateTree, FStateTreeIndex16(ConditionIndex));
-							Cond.EnterState(*this, Transition);
+							Cond.EnterState(*this, CurrentTransition);
 
 							// Reset copied properties that might contain object references.
 							if (Cond.BindingsBatch.IsValid())
@@ -2829,10 +3226,11 @@ EStateTreeRunStatus FStateTreeExecutionContext::EnterState(FStateTreeTransitionR
 			}
 
 			// Activate tasks on current state.
-			for (int32 StateTaskIndex = 0; StateTaskIndex < State.TasksNum; ++StateTaskIndex)
+			UE::StateTree::FTasksCompletionStatus CurrentStateTasksStatus = CurrentFrame.ActiveTasksStatus.GetStatus(CurrentState);
+			for (int32 StateTaskIndex = 0; StateTaskIndex < CurrentState.TasksNum; ++StateTaskIndex)
 			{
-				const int32 AssetTaskIndex = State.TasksBegin + StateTaskIndex;
-				const FStateTreeTaskBase& Task = NextFrame.StateTree->Nodes[AssetTaskIndex].Get<const FStateTreeTaskBase>();
+				const int32 AssetTaskIndex = CurrentState.TasksBegin + StateTaskIndex;
+				const FStateTreeTaskBase& Task = CurrentStateTree->Nodes[AssetTaskIndex].Get<const FStateTreeTaskBase>();
 				const FStateTreeDataView TaskInstanceView = GetDataView(CurrentParentFrame, CurrentFrame, Task.InstanceDataHandle);
 
 				FNodeInstanceDataScope DataScope(*this, &Task, AssetTaskIndex, Task.InstanceDataHandle, TaskInstanceView);
@@ -2855,10 +3253,10 @@ EStateTreeRunStatus FStateTreeExecutionContext::EnterState(FStateTreeTransitionR
 				const bool bShouldCallEnterState = CurrentTransition.ChangeType == EStateTreeStateChangeType::Changed
 													|| (CurrentTransition.ChangeType == EStateTreeStateChangeType::Sustained && Task.bShouldStateChangeOnReselect);
 
-				if (bIsEnteringState && bShouldCallEnterState)
+				if (bShouldCallEnterState)
 				{
-					STATETREE_LOG(Verbose, TEXT("%*sTask '%s'.EnterState()"), (FrameIndex + StateIndex + 1)*UE::StateTree::Debug::IndentSize, TEXT(""), *Task.Name.ToString());
-					UE_STATETREE_DEBUG_TASK_ENTER_STATE(this, NextFrame.StateTree, FStateTreeIndex16(AssetTaskIndex));
+					STATETREE_LOG(Verbose, TEXT("%*sTask '%s'.EnterState()"), (FrameIndex + ActiveStateIndex + 1) * UE::StateTree::Debug::IndentSize, TEXT(""), *Task.Name.ToString());
+					UE_STATETREE_DEBUG_TASK_ENTER_STATE(this, CurrentStateTree, FStateTreeIndex16(AssetTaskIndex));
 
 					EStateTreeRunStatus TaskRunStatus = EStateTreeRunStatus::Unset;
 					{
@@ -2868,240 +3266,136 @@ EStateTreeRunStatus FStateTreeExecutionContext::EnterState(FStateTreeTransitionR
 						TaskRunStatus = Task.EnterState(*this, CurrentTransition);
 					}
 
-					UE::StateTree::ETaskCompletionStatus TaskStatus = UE::StateTree::ExecutionContext::CastToTaskStatus(TaskRunStatus);
+					UE::StateTree::ETaskCompletionStatus TaskStatus = CastToTaskStatus(TaskRunStatus);
 					TaskStatus = CurrentStateTasksStatus.SetStatusWithPriority(StateTaskIndex, TaskStatus);
 
-					TaskRunStatus = UE::StateTree::ExecutionContext::CastToRunStatus(TaskStatus);
+					TaskRunStatus = CastToRunStatus(TaskStatus);
 					UE_STATETREE_DEBUG_TASK_EVENT(this, AssetTaskIndex, TaskInstanceView, EStateTreeTraceEventType::OnEntered, TaskRunStatus);
 
 					if (CurrentStateTasksStatus.IsConsideredForCompletion(StateTaskIndex))
 					{
-						Result = UE::StateTree::ExecutionContext::GetPriorityRunStatus(Result, TaskRunStatus);
-						if (TaskRunStatus == EStateTreeRunStatus::Failed)
+						Result = GetPriorityRunStatus(Result, TaskRunStatus);
+						if (Result == EStateTreeRunStatus::Failed)
 						{
 							break;
 						}
 					}
 				}
 			}
+			UE_STATETREE_DEBUG_STATE_EVENT(this, CurrentStateHandle, EStateTreeTraceEventType::OnEntered);
+		}
 
-			if (bIsEnteringState)
+		if (Result == EStateTreeRunStatus::Failed)
 			{
-				UE_STATETREE_DEBUG_STATE_EVENT(this, CurrentHandle, EStateTreeTraceEventType::OnEntered);
-			}
+			break;
 		}
 	}
 
 	UE_STATETREE_DEBUG_EXIT_PHASE(this, EStateTreeUpdatePhase::EnterStates);
-
 	UE_STATETREE_DEBUG_ACTIVE_STATES_EVENT(this, Exec.ActiveFrames);
 
 	Exec.bHasPendingCompletedState = Result != EStateTreeRunStatus::Running;
 	return Result;
 }
 
+// Deprecated
 void FStateTreeExecutionContext::ExitState(const FStateTreeTransitionResult& Transition)
 {
+}
+
+void FStateTreeExecutionContext::ExitState(const FSelectStateResult& Args, const FStateTreeTransitionResult& Transition)
+{
+	// 1. Copy all bindings on all active nodes.
+	// 2. Call ExitState on all active global tasks, global evaluators, state tasks and state condition that are affected by the transition.
+	//  If the frame/state is before the target
+	//    then it is not affected.
+	//  Else
+	//    Excluding the target, if the active state is not in the new selection, then call ExitState with EStateTreeStateChangeType::Changed.
+	//    Including the target, if the active state is in the new selection then call ExitState with EStateTreeStateChangeType::Sustained.
+	// 3. If the state/frame is not in the new selection, clean up memory, delegate...
+
+	using namespace UE::StateTree;
+	using namespace UE::StateTree::ExecutionContext;
+	using namespace UE::StateTree::ExecutionContext::Private;
+
 	CSV_SCOPED_TIMING_STAT_EXCLUSIVE(StateTree_ExitState);
 
 	FStateTreeExecutionState& Exec = GetExecState();
-
 	if (Exec.ActiveFrames.IsEmpty())
 	{
 		return;
 	}
 
-	// On target branch means that the state is the target of current transition or child of it.
-	// States which were active before and will remain active, but are not on target branch will not get
-	// EnterState called. That is, a transition is handled as "replan from this state".
-	bool bOnTargetBranch = false;
-
-	struct FExitStateCall
-	{
-		FExitStateCall() = default;
-		FExitStateCall(const EStateTreeStateChangeType InChangeType, const bool bInShouldCall)
-			: ChangeType(InChangeType)
-			, bShouldCall(bInShouldCall)
-		{
-		}
-
-		EStateTreeStateChangeType ChangeType = EStateTreeStateChangeType::None;
-		bool bShouldCall = false; 
-	};
+	CopyAllBindingsOnActiveInstances(ECopyBindings::ExitState);
 
-	TArray<FExitStateCall> ExitStateCalls;
+	STATETREE_LOG(Log, TEXT("Exit state '%s' (%d)"), *DebugGetStatePath(Exec.ActiveFrames), Exec.StateChangeCount);
+	UE_STATETREE_DEBUG_SCOPED_PHASE(this, EStateTreeUpdatePhase::ExitStates);
 
-	// Track any changed state (i.e., not sustained) to prevent reused subtrees from being considered as sustained states since the whole
-	// tree is reused and should be considered as a new instance.
-	bool bAnyParentStateChanged = false;
+	TArray<FActiveState, TInlineAllocator<ExpectedActiveStatePathLength>> ActiveStates;
+	GetActiveStatePath(GetExecState().ActiveFrames, ActiveStates);
 
-	int32 FirstModifiedFrame = INDEX_NONE;
-	int32 FirstModifiedState = INDEX_NONE;
+	const TArrayView<const FActiveState> CommonStates = FActiveStatePath::Intersect(MakeConstArrayView(ActiveStates), MakeConstArrayView(Args.SelectedStates));
+	const TArrayView<const FActiveState> ChangedStates = MakeConstArrayView(ActiveStates).Mid(CommonStates.Num());
+	const bool bHasSustained = CommonStates.Contains(Args.TargetState);
 
-	for (int32 FrameIndex = 0; FrameIndex < Exec.ActiveFrames.Num(); FrameIndex++)
+	bool bContinue = true;
+	FStateTreeTransitionResult CurrentTransition = Transition;
+	for (int32 FrameIndex = Exec.ActiveFrames.Num() - 1; bContinue && FrameIndex >= 0; FrameIndex--)
 	{
 		FStateTreeExecutionFrame* CurrentParentFrame = FrameIndex > 0 ? &Exec.ActiveFrames[FrameIndex - 1] : nullptr;
 		FStateTreeExecutionFrame& CurrentFrame = Exec.ActiveFrames[FrameIndex];
 		const UStateTree* CurrentStateTree = CurrentFrame.StateTree;
 
 		FCurrentlyProcessedFrameScope FrameScope(*this, CurrentParentFrame, CurrentFrame);
-
-		const FStateTreeExecutionFrame* NextFrame = nullptr;
-		if (Transition.NextActiveFrames.IsValidIndex(FrameIndex)
-			&& Transition.NextActiveFrames[FrameIndex].IsSameFrame(CurrentFrame))
+		for (int32 StateIndex = CurrentFrame.ActiveStates.Num() - 1; bContinue && StateIndex >= 0; StateIndex--)
 		{
-			NextFrame = &Transition.NextActiveFrames[FrameIndex];
-		}
-
-		const bool bShouldCallOnEvaluatorsAndGlobalTasks = NextFrame == nullptr && CurrentFrame.bIsGlobalFrame;
-		ExitStateCalls.Emplace(EStateTreeStateChangeType::Changed, bShouldCallOnEvaluatorsAndGlobalTasks);
+			const FStateTreeStateHandle CurrentHandle = CurrentFrame.ActiveStates[StateIndex];
+			const FActiveStateID CurrentStateID = CurrentFrame.ActiveStates.StateIDs[StateIndex];
+			const FCompactStateTreeState& CurrentState = CurrentStateTree->States[CurrentHandle.Index];
 		
-		if (bShouldCallOnEvaluatorsAndGlobalTasks)
-		{
-			for (int32 EvalIndex = CurrentStateTree->EvaluatorsBegin; EvalIndex < (CurrentStateTree->EvaluatorsBegin + CurrentStateTree->EvaluatorsNum); EvalIndex++)
-			{
-				const FStateTreeEvaluatorBase& Eval = CurrentStateTree->Nodes[EvalIndex].Get<const FStateTreeEvaluatorBase>();
-				const FStateTreeDataView EvalInstanceView = GetDataView(CurrentParentFrame, CurrentFrame, Eval.InstanceDataHandle);
-				FNodeInstanceDataScope DataScope(*this, &Eval, EvalIndex, Eval.InstanceDataHandle, EvalInstanceView);
-
-				if (Eval.BindingsBatch.IsValid())
-				{
-					CopyBatchOnActiveInstances(CurrentParentFrame, CurrentFrame, EvalInstanceView, Eval.BindingsBatch);
-				}
-			}
-
-			for (int32 TaskIndex = CurrentStateTree->GlobalTasksBegin; TaskIndex < (CurrentStateTree->GlobalTasksBegin + CurrentStateTree->GlobalTasksNum); TaskIndex++)
-			{
-				const FStateTreeTaskBase& Task = CurrentStateTree->Nodes[TaskIndex].Get<const FStateTreeTaskBase>();
-				const FStateTreeDataView TaskInstanceView = GetDataView(CurrentParentFrame, CurrentFrame, Task.InstanceDataHandle);
-				FNodeInstanceDataScope DataScope(*this, &Task, TaskIndex, Task.InstanceDataHandle, TaskInstanceView);
-
-				if (Task.BindingsBatch.IsValid() && Task.bShouldCopyBoundPropertiesOnExitState)
-				{
-					CopyBatchOnActiveInstances(CurrentParentFrame, CurrentFrame, TaskInstanceView, Task.BindingsBatch);
-				}
-			}
-		}
-
-		for (int32 Index = 0; Index < CurrentFrame.ActiveStates.Num(); Index++)
-		{
-			const FStateTreeStateHandle CurrentHandle = CurrentFrame.ActiveStates[Index];
-			const FStateTreeStateHandle NextHandle = NextFrame ? NextFrame->ActiveStates.GetStateSafe(Index) : FStateTreeStateHandle::Invalid;
-			const FCompactStateTreeState& State = CurrentStateTree->States[CurrentHandle.Index];
-
-			FCurrentlyProcessedStateScope StateScope(*this, CurrentHandle);
+			const FActiveState CurrentActiveState(CurrentFrame.FrameID, CurrentStateID, CurrentHandle);
+			const bool bAreCommon = !ChangedStates.Contains(CurrentActiveState);
 
-			if (State.Type == EStateTreeStateType::Linked
-				|| State.Type == EStateTreeStateType::LinkedAsset)
-			{
-				if (State.ParameterDataHandle.IsValid()
-					&& State.ParameterBindingsBatch.IsValid())
+			// Can it stop calling Changed
+			if (bAreCommon && !bHasSustained)
 				{
-					const FStateTreeDataView StateParamsDataView = GetDataView(CurrentParentFrame, CurrentFrame, State.ParameterDataHandle);
-					CopyBatchOnActiveInstances(CurrentParentFrame, CurrentFrame, StateParamsDataView, State.ParameterBindingsBatch);
-				}
-			}
-
-			if (CurrentFrame.FrameID == Transition.SourceFrameID
-				&& CurrentHandle == Transition.TargetState)
-			{
-				bOnTargetBranch = true;
-			}
-
-			const bool bRemainsActive = (NextHandle == CurrentHandle) && !bAnyParentStateChanged;
-
-			const EStateTreeStateChangeType ChangeType = bRemainsActive ? EStateTreeStateChangeType::Sustained : EStateTreeStateChangeType::Changed;
-			if (ChangeType == EStateTreeStateChangeType::Changed)
-			{
-				bAnyParentStateChanged = true;
-			}
-			
-			// Should call ExitState() on this state.
-			// (No need to test for 'State.bEnabled' like EnterState since we can't enter a disabled state)
-			const bool bShouldCallExitState = (!bRemainsActive || bOnTargetBranch);
-			ExitStateCalls.Emplace(ChangeType, bShouldCallExitState);
-
-			if (bShouldCallExitState && FirstModifiedFrame == INDEX_NONE)
-			{
-				FirstModifiedFrame = FrameIndex;
-				FirstModifiedState = Index;
+				// It will also stop the frame loop.
+				bContinue = false;
+				break;
 			}
 
-			// Do property copies, ExitState() is called below.
-			for (int32 TaskIndex = State.TasksBegin; TaskIndex < (State.TasksBegin + State.TasksNum); TaskIndex++)
+			if (bAreCommon && bHasSustained && CurrentActiveState == Args.TargetState)
 			{
-				const FStateTreeTaskBase& Task = CurrentStateTree->Nodes[TaskIndex].Get<const FStateTreeTaskBase>();
-				const FStateTreeDataView TaskInstanceView = GetDataView(CurrentParentFrame, CurrentFrame, Task.InstanceDataHandle);
-
-				// Copy bound properties.
-				if (Task.BindingsBatch.IsValid() && Task.bShouldCopyBoundPropertiesOnExitState)
-				{
-					CopyBatchOnActiveInstances(CurrentParentFrame, CurrentFrame, TaskInstanceView, Task.BindingsBatch);
-				}
-			}
+				// this is the last sustained
+				bContinue = false;
 		}
-	}
-
-	// Call in reverse order.
-	STATETREE_LOG(Log, TEXT("Exit state '%s' (%d)"), *DebugGetStatePath(Exec.ActiveFrames), Exec.StateChangeCount);
-	UE_STATETREE_DEBUG_SCOPED_PHASE(this, EStateTreeUpdatePhase::ExitStates);
-
-	FStateTreeTransitionResult CurrentTransition = Transition;
-	int32 CallIndex = ExitStateCalls.Num() - 1;
-
-	for (int32 FrameIndex = Exec.ActiveFrames.Num() - 1; FrameIndex >= 0; FrameIndex--)
-	{
-		FStateTreeExecutionFrame* CurrentParentFrame = FrameIndex > 0 ? &Exec.ActiveFrames[FrameIndex - 1] : nullptr;
-		FStateTreeExecutionFrame& CurrentFrame = Exec.ActiveFrames[FrameIndex];
-		const UStateTree* CurrentStateTree = CurrentFrame.StateTree;
-
-		FCurrentlyProcessedFrameScope FrameScope(*this, CurrentParentFrame, CurrentFrame);
 
-		for (int32 StateIndex = CurrentFrame.ActiveStates.Num() - 1; StateIndex >= 0; StateIndex--)
-		{
-			const FStateTreeStateHandle CurrentHandle = CurrentFrame.ActiveStates[StateIndex];
-			const UE::StateTree::FActiveStateID CurrentStateID = CurrentFrame.ActiveStates.StateIDs[StateIndex];
-			const FCompactStateTreeState& State = CurrentStateTree->States[CurrentHandle.Index];
+			const bool bIsStateSustained = bAreCommon;
 
-			const FExitStateCall& ExitCall = ExitStateCalls[CallIndex--];
-			CurrentTransition.ChangeType = ExitCall.ChangeType;
+			TGuardValue<FStateTreeStateHandle> CurrentTransitionCurrentStateGuard(CurrentTransition.CurrentState, CurrentHandle);
+			CurrentTransition.ChangeType = bIsStateSustained ? EStateTreeStateChangeType::Sustained : EStateTreeStateChangeType::Changed;
 
 			STATETREE_LOG(Log, TEXT("%*sState '%s' (%s)"), (FrameIndex + StateIndex + 1)*UE::StateTree::Debug::IndentSize, TEXT("")
 				, *GetSafeStateName(CurrentFrame, CurrentHandle)
 				, *UEnum::GetDisplayValueAsText(CurrentTransition.ChangeType).ToString());
 
-			if (ExitCall.bShouldCall)
-			{
 				FCurrentlyProcessedStateScope StateScope(*this, CurrentHandle);
 				UE_STATETREE_DEBUG_STATE_EVENT(this, CurrentHandle, EStateTreeTraceEventType::OnExiting);
 
-				// Remove any delayed transitions that belong to this state.
-				Exec.DelayedTransitions.RemoveAllSwap(
-					[CurrentStateID, Begin = State.TransitionsBegin, End = State.TransitionsBegin + State.TransitionsNum](const FStateTreeTransitionDelayedState& DelayedState)
+			for (int32 TaskIndex = (CurrentState.TasksBegin + CurrentState.TasksNum) - 1; TaskIndex >= CurrentState.TasksBegin; TaskIndex--)
 					{
-						return  DelayedState.StateID == CurrentStateID
-							&& DelayedState.TransitionIndex.Get() >= Begin
-							&& DelayedState.TransitionIndex.Get() < End;
-					});
-
-				CurrentTransition.CurrentState = CurrentHandle;
+				const FStateTreeTaskBase& Task = CurrentStateTree->Nodes[TaskIndex].Get<const FStateTreeTaskBase>();
 
-				// Do property copies, ExitState() is called below.
-				for (int32 TaskIndex = (State.TasksBegin + State.TasksNum) - 1; TaskIndex >= State.TasksBegin; TaskIndex--)
+				// Ignore disabled task
+				if (Task.bTaskEnabled)
 				{
-
 					// Call task completed only if EnterState() was called.
 					// The task order in the tree (BF) allows us to use the comparison.
 					if (TaskIndex <= CurrentFrame.ActiveTaskNodeIndex.AsInt32())
 					{
-						const FStateTreeTaskBase& Task = CurrentStateTree->Nodes[TaskIndex].Get<const FStateTreeTaskBase>();
 						const FStateTreeDataView TaskInstanceView = GetDataView(CurrentParentFrame, CurrentFrame, Task.InstanceDataHandle);
-
 						FNodeInstanceDataScope DataScope(*this, &Task, TaskIndex, Task.InstanceDataHandle, TaskInstanceView);
 
-						// Ignore disabled task
-						if (Task.bTaskEnabled)
-						{
 							const bool bShouldCallStateChange = CurrentTransition.ChangeType == EStateTreeStateChangeType::Changed
 									|| (CurrentTransition.ChangeType == EStateTreeStateChangeType::Sustained && Task.bShouldStateChangeOnReselect);
 
@@ -3117,18 +3411,17 @@ void FStateTreeExecutionContext::ExitState(const FStateTreeTransitionResult& Tra
 								UE_STATETREE_DEBUG_TASK_EVENT(this, TaskIndex, TaskInstanceView, EStateTreeTraceEventType::OnExited, Transition.CurrentRunStatus);
 							}
 						}
+				}
 						else
 						{
 							STATETREE_LOG(VeryVerbose, TEXT("%*sSkipped 'ExitState' for disabled Task: '%s'"), UE::StateTree::Debug::IndentSize, TEXT(""), *Task.Name.ToString());
 						}
-					}
 					CurrentFrame.ActiveTaskNodeIndex = FStateTreeIndex16(TaskIndex - 1);
-				}
 
 				// Call state change events on conditions if needed.
-				if (State.bHasStateChangeConditions)
+				if (CurrentState.bHasStateChangeConditions)
 				{
-					for (int32 ConditionIndex = (State.EnterConditionsBegin + State.EnterConditionsNum) - 1; ConditionIndex >= State.EnterConditionsBegin; ConditionIndex--)
+					for (int32 ConditionIndex = (CurrentState.EnterConditionsBegin + CurrentState.EnterConditionsNum) - 1; ConditionIndex >= CurrentState.EnterConditionsBegin; ConditionIndex--)
 					{
 						const FStateTreeConditionBase& Cond = CurrentFrame.StateTree->Nodes[ConditionIndex].Get<const FStateTreeConditionBase>();
 						if (Cond.bHasShouldCallStateChangeEvents)
@@ -3148,7 +3441,7 @@ void FStateTreeExecutionContext::ExitState(const FStateTreeTransitionResult& Tra
 								}
 
 								UE_STATETREE_DEBUG_CONDITION_EXIT_STATE(this, CurrentFrame.StateTree, FStateTreeIndex16(ConditionIndex));
-								Cond.ExitState(*this, Transition);
+								Cond.ExitState(*this, CurrentTransition);
 
 								// Reset copied properties that might contain object references.
 								if (Cond.BindingsBatch.IsValid())
@@ -3159,25 +3452,46 @@ void FStateTreeExecutionContext::ExitState(const FStateTreeTransitionResult& Tra
 						}
 					}
 				}
+			}
+
+			// Reset the completed state. This is to keep the wrong UE5.6 behavior.
+			if (!EnumHasAnyFlags(CurrentFrame.StateTree->GetStateSelectionRules(), EStateTreeStateSelectionRules::CompletedTransitionStatesCreateNewStates)
+				&& bIsStateSustained)
+			{
+				CurrentFrame.ActiveTasksStatus.GetStatus(CurrentState).ResetStatus(CurrentState.TasksNum);
+			}
 
-				// Delegate Listeners Cleanup
-				GetExecState().DelegateActiveListeners.RemoveAll(CurrentFrame.ActiveStates.StateIDs[StateIndex]);
+			// Remove the state from the active states if it "changed".
+			if (!bIsStateSustained)
+			{
+				CleanState(Exec, CurrentStateID);
 
-				UE_STATETREE_DEBUG_STATE_EVENT(this, CurrentHandle, EStateTreeTraceEventType::OnExited);
+				// Remove state
+				const FStateTreeStateHandle PoppedState = CurrentFrame.ActiveStates.Pop();
+				check(PoppedState == CurrentHandle);
 			}
+
+			UE_STATETREE_DEBUG_STATE_EVENT(this, CurrentHandle, EStateTreeTraceEventType::OnExited);
 		}
 
 		// Frame exit call
+		if (bContinue)
 		{
-			const FExitStateCall& ExitCall = ExitStateCalls[CallIndex--];
-			if (ExitCall.bShouldCall)
+			// The frame is not part of the new path
+			const bool bFrameSustained = Args.SelectedFrames.Contains(CurrentFrame.FrameID);
+			if (CurrentFrame.bIsGlobalFrame)
 			{
 				CurrentTransition.CurrentState = FStateTreeStateHandle::Invalid;
-				CurrentTransition.ChangeType = ExitCall.ChangeType;
+				CurrentTransition.ChangeType = bFrameSustained ? EStateTreeStateChangeType::Sustained : EStateTreeStateChangeType::Changed;
 				StopGlobalsForFrameOnActiveInstances(CurrentParentFrame, CurrentFrame, CurrentTransition);
+			}
 
-				// Delegate Listeners Cleanup
-				GetExecState().DelegateActiveListeners.RemoveAll(CurrentFrame.FrameID);
+			if (!bFrameSustained)
+			{
+				checkf(CurrentFrame.ActiveStates.Num() == 0, TEXT("All states must received ExitState first."));
+				CleanFrame(Exec, CurrentFrame.FrameID);
+				check(Exec.ActiveFrames.Num() - 1 == FrameIndex);
+				Exec.ActiveFrames.RemoveAt(FrameIndex);
 			}
 		}
 	}
@@ -3230,8 +3544,6 @@ void FStateTreeExecutionContext::StateCompleted()
 				if (TaskIndex <= CurrentActiveTaskNodeIndex)
 				{
 					const FStateTreeTaskBase& Task = CurrentStateTree->Nodes[TaskIndex].Get<const FStateTreeTaskBase>();
-					const FStateTreeDataView TaskInstanceView = GetDataView(CurrentParentFrame, CurrentFrame, Task.InstanceDataHandle);
-					FNodeInstanceDataScope DataScope(*this, &Task, TaskIndex, Task.InstanceDataHandle, TaskInstanceView);
 
 					// Ignore disabled task
 					if (Task.bTaskEnabled == false)
@@ -3240,6 +3552,9 @@ void FStateTreeExecutionContext::StateCompleted()
 						continue;
 					}
 
+					const FStateTreeDataView TaskInstanceView = GetDataView(CurrentParentFrame, CurrentFrame, Task.InstanceDataHandle);
+					FNodeInstanceDataScope DataScope(*this, &Task, TaskIndex, Task.InstanceDataHandle, TaskInstanceView);
+
 					STATETREE_LOG(Verbose, TEXT("%*sTask '%s'.StateCompleted()"), (FrameIndex + StateIndex + 1)*UE::StateTree::Debug::IndentSize, TEXT(""), *Task.Name.ToString());
 					Task.StateCompleted(*this, Exec.LastTickStatus, CurrentFrame.ActiveStates);
 				}
@@ -3432,15 +3747,17 @@ EStateTreeRunStatus FStateTreeExecutionContext::StartGlobalsForFrameInternal(con
 	}
 
 	CSV_SCOPED_TIMING_STAT_EXCLUSIVE(StateTree_StartEvaluators);
+	UE_STATETREE_DEBUG_SCOPED_PHASE(this, EStateTreeUpdatePhase::StartGlobalTasks);
 
 	FStateTreeExecutionState& Exec = GetExecState();
-
 	FCurrentlyProcessedFrameScope FrameScope(*this, CurrentParentFrame, CurrentFrame);
 	if constexpr (!bOnActiveInstances)
 	{
 		UE_STATETREE_DEBUG_ENTER_PHASE(this, EStateTreeUpdatePhase::StartGlobalTasksForSelection);
 	}
 
+	CurrentFrame.bHaveEntered = true;
+
 	EStateTreeRunStatus Result = EStateTreeRunStatus::Running;
 	const UE::StateTree::FActiveFrameID CurrentFrameID = CurrentFrame.FrameID;
 	const UStateTree* CurrentStateTree = CurrentFrame.StateTree;
@@ -3550,7 +3867,7 @@ EStateTreeRunStatus FStateTreeExecutionContext::StartGlobalsForFrameInternal(con
 			Result = UE::StateTree::ExecutionContext::GetPriorityRunStatus(Result, TaskRunStatus);
 			if (Result == EStateTreeRunStatus::Failed)
 			{
-				return Result;
+				break;
 			}
 		}
 	}
@@ -3656,7 +3973,8 @@ void FStateTreeExecutionContext::StopGlobalsForFrameWithValidation(const FStateT
 template<bool bOnActiveInstances>
 void FStateTreeExecutionContext::StopGlobalsForFrameInternal(const FStateTreeExecutionFrame* ParentFrame, FStateTreeExecutionFrame& Frame, const FStateTreeTransitionResult& Transition)
 {
-	if (!Frame.bIsGlobalFrame)
+	// Special case when we select a new root. See SelectState. We don't want to stop the globals.
+	if (!Frame.bIsGlobalFrame || !Frame.bHaveEntered)
 	{
 		return;
 	}
@@ -3717,6 +4035,8 @@ void FStateTreeExecutionContext::StopGlobalsForFrameInternal(const FStateTreeExe
 		}
 	}
 
+	Frame.bHaveEntered = false;
+
 	if constexpr (bOnActiveInstances)
 	{
 		UE_STATETREE_DEBUG_EXIT_PHASE(this, EStateTreeUpdatePhase::StopGlobalTasksForSelection);
@@ -3864,6 +4184,7 @@ EStateTreeRunStatus FStateTreeExecutionContext::TickTasks(const float DeltaTime)
 	for (int32 FrameIndex = 0; FrameIndex < Exec.ActiveFrames.Num(); ++FrameIndex)
 	{
 		using namespace UE::StateTree::ExecutionContext;
+		using namespace UE::StateTree::ExecutionContext::Private;
 
 		const FStateTreeExecutionFrame& CurrentFrame = Exec.ActiveFrames[FrameIndex];
 		const UStateTree* CurrentStateTree = CurrentFrame.StateTree;
@@ -3898,7 +4219,7 @@ EStateTreeRunStatus FStateTreeExecutionContext::TickTasks(const float DeltaTime)
 	{
 		// No enabled tasks, done ticking.
 		//Complete the the bottom state in the bottom frame (to trigger the completion transitions).
-		if (ensureMsgf(Exec.ActiveFrames.Num() > 0, TEXT("No task is allowed to clear/stop/transition. Those actions should be delayed inside the execution context.")))
+		if (Exec.ActiveFrames.Num() > 0)
 		{
 			FStateTreeExecutionFrame& LastFrame = Exec.ActiveFrames.Last();
 			const int32 NumberOfActiveState = LastFrame.ActiveStates.Num();
@@ -3915,7 +4236,7 @@ EStateTreeRunStatus FStateTreeExecutionContext::TickTasks(const float DeltaTime)
 		}
 		else
 		{
-			Exec.RequestedStop = EStateTreeRunStatus::Stopped;
+			Exec.RequestedStop = ExecutionContext::GetPriorityRunStatus(Exec.RequestedStop, EStateTreeRunStatus::Stopped);
 		}
 		StateResult = EStateTreeRunStatus::Succeeded;
 	}
@@ -4014,9 +4335,13 @@ FStateTreeExecutionContext::FTickTaskResult FStateTreeExecutionContext::TickTask
 			TaskRunStatus != EStateTreeRunStatus::Running ? EStateTreeTraceEventType::OnTaskCompleted : EStateTreeTraceEventType::OnTicked,
 			TaskRunStatus);
 
-		if (TaskRunStatus == EStateTreeRunStatus::Failed && Args.TasksCompletionStatus->IsConsideredForCompletion(OwnerTaskIndex))
+		if (Args.TasksCompletionStatus->IsConsideredForCompletion(OwnerTaskIndex))
+		{
+			if (TaskRunStatus == EStateTreeRunStatus::Failed)
 		{
 			bShouldTickTasks = false;
+				break; // Stop copy binding.
+			}
 		}
 	}
 
@@ -4334,6 +4659,7 @@ FString FStateTreeExecutionContext::DebugGetEventsAsString() const
 	return StrBuilder.ToString();
 }
 
+// Deprecated
 bool FStateTreeExecutionContext::RequestTransition(
 	const FStateTreeExecutionFrame& CurrentFrame,
 	const FStateTreeStateHandle NextState,
@@ -4341,56 +4667,152 @@ bool FStateTreeExecutionContext::RequestTransition(
 	const FStateTreeSharedEvent* TransitionEvent,
 	const EStateTreeSelectionFallback Fallback)
 {
+	using namespace UE::StateTree;
+	using namespace UE::StateTree::ExecutionContext;
+
+	const FStateHandleContext TargetState = FStateHandleContext(CurrentFrame.StateTree, NextState);
+	const FStateTreeSharedEvent Event = TransitionEvent != nullptr ? *TransitionEvent : FStateTreeSharedEvent();
+	return RequestTransitionInternal(CurrentFrame, FActiveStateID(), TargetState, Priority, Event, Fallback);
+}
+
+bool FStateTreeExecutionContext::RequestTransitionInternal(
+	const FStateTreeExecutionFrame& SourceFrame,
+	const UE::StateTree::FActiveStateID SourceStateID,
+	const UE::StateTree::ExecutionContext::FStateHandleContext TargetState,
+	const EStateTreeTransitionPriority Priority,
+	const FStateTreeSharedEvent TransitionEvent,
+	const EStateTreeSelectionFallback Fallback)
+{
+	using namespace UE::StateTree;
+	using namespace UE::StateTree::ExecutionContext;
+	using namespace UE::StateTree::ExecutionContext::Private;
+
 	// Skip lower priority transitions.
-	if (NextTransition.Priority >= Priority)
+	if (RequestedTransition.IsValid() && RequestedTransition->Transition.Priority >= Priority)
 	{
 		return false;
 	}
 
-	if (NextState.IsCompletionState())
+	if (TargetState.StateHandle.IsCompletionState())
 	{
-		SetupNextTransition(CurrentFrame, NextState, Priority);
+		SetupNextTransition(SourceFrame, SourceStateID, TargetState, Priority);
 		STATETREE_LOG(Verbose, TEXT("Transition on state '%s' -> state '%s'"),
-			*GetSafeStateName(CurrentFrame, CurrentFrame.ActiveStates.Last()), *NextState.Describe());
+			*GetSafeStateName(SourceFrame, SourceFrame.ActiveStates.Last()), *TargetState.StateHandle.Describe());
 		return true;
 	}
-	if (!NextState.IsValid())
+	if (!TargetState.StateHandle.IsValid() || TargetState.StateTree == nullptr)
 	{
 		// NotSet is no-operation, but can be used to mask a transition at parent state. Returning unset keeps updating current state.
-		SetupNextTransition(CurrentFrame, FStateTreeStateHandle::Invalid, Priority);
+		SetupNextTransition(SourceFrame, SourceStateID, FStateHandleContext(), Priority);
 		return true;
 	}
 
-	FStateSelectionResult StateSelectionResult;
-	if (SelectState(CurrentFrame, NextState, StateSelectionResult, TransitionEvent, Fallback))
+	TArray<FActiveState, TInlineAllocator<ExpectedActiveStatePathLength>> CurrentActiveStatePath;
+	GetActiveStatePath(GetExecState().ActiveFrames, CurrentActiveStatePath);
+
+	FSelectStateArguments SelectStateArgs;
+	SelectStateArgs.ActiveStates = MakeConstArrayView(CurrentActiveStatePath);
+
+	// Source can be a GlobalTask (not in active state)
+	FStateTreeStateHandle SourceStateHandle;
+	if (SourceStateID.IsValid())
+	{
+		SourceStateHandle = SourceFrame.ActiveStates.FindStateHandle(SourceStateID);
+		if (!ensure(SourceStateHandle.IsValid()))
+		{
+			return false;
+		}
+		SelectStateArgs.SourceState = FActiveState(SourceFrame.FrameID, SourceStateID, SourceStateHandle);
+	}
+	else
 	{
-		SetupNextTransition(CurrentFrame, NextState, Priority);
-		NextTransition.NextActiveFrames = StateSelectionResult.GetSelectedFrames();
-		NextTransition.NextActiveFrameEvents = StateSelectionResult.GetFramesStateSelectionEvents();
+		SelectStateArgs.SourceState = FActiveState(SourceFrame.FrameID, FActiveStateID(), FStateTreeStateHandle());
+	}
+	SelectStateArgs.TargetState = TargetState;
+	SelectStateArgs.TransitionEvent = TransitionEvent;
+	SelectStateArgs.Fallback = Fallback;
+	SelectStateArgs.Behavior = ESelectStateBehavior::StateTransition;
 
-		// Consume events from states, if required. 
-		for (int32 FrameIndex = 0; FrameIndex < NextTransition.NextActiveFrames.Num(); FrameIndex++)
+	FSelectStateResult StateSelectionResult;
+	if (SelectState(SelectStateArgs, StateSelectionResult))
 		{
-			const FStateTreeExecutionFrame& Frame = NextTransition.NextActiveFrames[FrameIndex];
-			const FStateTreeFrameStateSelectionEvents& FrameEvents = NextTransition.NextActiveFrameEvents[FrameIndex];
+		const FStateTreeExecutionState& Exec = GetExecState();
+
+		SetupNextTransition(SourceFrame, SourceStateID, TargetState, Priority);
+		RequestedTransition->Selection = MoveTemp(StateSelectionResult);
 
-			for (int32 StateIndex = 0; StateIndex < Frame.ActiveStates.Num(); StateIndex++)
+		//@todo remove this when NextActiveFrames & NextActiveFrameEvents are deprecated
+		// Fill NextActiveFrames & NextActiveFrameEvents for backward compatibility with ForceTransition
+		{
+			RequestedTransition->Transition.NextActiveFrames.Reset();
+			RequestedTransition->Transition.NextActiveFrameEvents.Reset();
+			FActiveFrameID PreviousFrameID;
+			FStateTreeExecutionFrame* CurrentFrame = nullptr;
+			FStateTreeFrameStateSelectionEvents* CurrentSelectionEvents = nullptr;
+			int32 CurrentSelectionEventsIndex = 0;
+			for (const FActiveState& SelectedState : RequestedTransition->Selection.SelectedStates)
 			{
-				if (FrameEvents.Events[StateIndex].IsValid())
+				if (SelectedState.GetFrameID() != PreviousFrameID)
+				{
+					PreviousFrameID = SelectedState.GetFrameID();
+					const FStateTreeExecutionFrame* Frame = FindExecutionFrame(SelectedState.GetFrameID(), MakeConstArrayView(Exec.ActiveFrames), MakeConstArrayView(RequestedTransition->Selection.TemporaryFrames));
+					if (ensure(Frame))
+					{
+						RequestedTransition->Transition.NextActiveFrames.Add(*Frame);
+						CurrentFrame = &RequestedTransition->Transition.NextActiveFrames.Last();
+
+						RequestedTransition->Transition.NextActiveFrameEvents.AddDefaulted();
+						CurrentSelectionEvents = &RequestedTransition->Transition.NextActiveFrameEvents.Last();
+						CurrentSelectionEventsIndex = 0;
+
+						CurrentFrame->ActiveStates = FStateTreeActiveStates();
+					}
+				}
+				check(CurrentFrame);
+				check(CurrentSelectionEvents);
+				CurrentFrame->ActiveStates.Push(SelectedState.GetStateHandle(), SelectedState.GetStateID());
+
+				const FSelectionEventWithID* FoundEvent = RequestedTransition->Selection.SelectionEvents.FindByPredicate(
+					[SelectedState](const FSelectionEventWithID& Event)
 				{
-					const FCompactStateTreeState& State = Frame.StateTree->States[StateIndex];
-					if (State.bConsumeEventOnSelect)
+						return Event.State == SelectedState;
+					});
+				if (FoundEvent)
 					{
-						ConsumeEvent(FrameEvents.Events[StateIndex]);
+					CurrentSelectionEvents->Events[CurrentSelectionEventsIndex] = FoundEvent->Event;
+				}
+			}
 					}
+
+		//@todo In TriggerTransitions, the events are processed bottom up. If an event with a higher priority is process after, the event will already be consumed.
+		// Consume events from states, if required. The transition might also want to consume the event. (Not a bug.)
+		for (FSelectionEventWithID& Event : RequestedTransition->Selection.SelectionEvents)
+		{
+			const FStateTreeExecutionFrame* Frame = FindExecutionFrame(Event.State.GetFrameID(), MakeConstArrayView(Exec.ActiveFrames), MakeConstArrayView(RequestedTransition->Selection.TemporaryFrames));
+			if (ensure(Frame))
+			{
+				const FCompactStateTreeState* State = Frame->StateTree->GetStateFromHandle(Event.State.GetStateHandle());
+				if (ensure(State) && State->bConsumeEventOnSelect)
+				{
+					ConsumeEvent(Event.Event);
 				}
 			}
 		}
 
-		STATETREE_LOG(Verbose, TEXT("Transition on state '%s' -[%s]-> state '%s'"),
-			*GetSafeStateName(CurrentFrame, CurrentFrame.ActiveStates.Last()),
-			*GetSafeStateName(CurrentFrame, NextState),
-			*GetSafeStateName(NextTransition.NextActiveFrames.Last(), NextTransition.NextActiveFrames.Last().ActiveStates.Last()));
+PRAGMA_DISABLE_DEPRECATION_WARNINGS
+UE_SUPPRESS(LogStateTree, Verbose,)
+		{
+			check(RequestedTransition->Selection.SelectedStates.Num() > 0);
+			const FActiveState& LastSelectedState = RequestedTransition->Selection.SelectedStates.Last();
+			const FStateTreeExecutionFrame* SelectedFrame = FindExecutionFrame(LastSelectedState.GetFrameID(), MakeConstArrayView(Exec.ActiveFrames), MakeConstArrayView(RequestedTransition->Selection.TemporaryFrames));
+			check(SelectedFrame);
+			STATETREE_LOG(Verbose, TEXT("Transition Request. Source:'%s'. Target:'%s'. Selected:'%s'"),
+				*GetSafeStateName(SourceFrame, SourceStateHandle),
+				*GetSafeStateName(TargetState.StateTree, TargetState.StateHandle),
+				*GetSafeStateName(*SelectedFrame, LastSelectedState.GetStateHandle())
+			);
+		}
+PRAGMA_ENABLE_DEPRECATION_WARNINGS
 
 		return true;
 	}
@@ -4398,44 +4820,65 @@ bool FStateTreeExecutionContext::RequestTransition(
 	return false;
 }
 
+// Deprecated
 void FStateTreeExecutionContext::SetupNextTransition(const FStateTreeExecutionFrame& CurrentFrame, const FStateTreeStateHandle NextState, const EStateTreeTransitionPriority Priority)
+{
+	using namespace UE::StateTree;
+	using namespace UE::StateTree::ExecutionContext;
+
+	const FStateHandleContext TargetState = FStateHandleContext(CurrentFrame.StateTree, NextState);
+	SetupNextTransition(CurrentFrame, FActiveStateID(), TargetState, Priority);
+}
+
+void FStateTreeExecutionContext::SetupNextTransition(
+	const FStateTreeExecutionFrame& SourceFrame,
+	const UE::StateTree::FActiveStateID SourceStateID,
+	const UE::StateTree::ExecutionContext::FStateHandleContext TargetState,
+	const EStateTreeTransitionPriority Priority)
+{
+	RequestedTransition = MakeUnique<FRequestTransitionResult>();
+	SetupNextTransition(SourceFrame, SourceStateID, TargetState, Priority, RequestedTransition->Transition);
+}
+
+void FStateTreeExecutionContext::SetupNextTransition(
+	const FStateTreeExecutionFrame& SourceFrame,
+	const UE::StateTree::FActiveStateID SourceStateID,
+	const UE::StateTree::ExecutionContext::FStateHandleContext TargetState,
+	const EStateTreeTransitionPriority Priority,
+	FStateTreeTransitionResult& OutTransitionResult)
 {
 	const FStateTreeExecutionState& Exec = GetExecState();
 
-	NextTransition.SourceFrameID = CurrentFrame.FrameID;
-	NextTransition.SourceStateID = UE::StateTree::FActiveStateID();
-	if (CurrentlyProcessedState.IsValid())
-	{
-		const int32 CurrentStateIndex = CurrentFrame.ActiveStates.IndexOfReverse(CurrentlyProcessedState);
-		if (CurrentStateIndex != INDEX_NONE)
+	OutTransitionResult.SourceFrameID = SourceFrame.FrameID;
+	OutTransitionResult.SourceStateID = SourceStateID;
+PRAGMA_DISABLE_DEPRECATION_WARNINGS
+	if (SourceStateID.IsValid())
 		{
-			NextTransition.SourceStateID = CurrentFrame.ActiveStates.StateIDs[CurrentStateIndex];
+		OutTransitionResult.SourceState = SourceFrame.ActiveStates.FindStateHandle(SourceStateID);
 		}
-	}
-PRAGMA_DISABLE_DEPRECATION_WARNINGS
-	NextTransition.SourceState = CurrentlyProcessedState;
-	NextTransition.SourceStateTree = CurrentFrame.StateTree;
-	NextTransition.SourceRootState = CurrentFrame.ActiveStates.GetStateSafe(0);
+	OutTransitionResult.SourceStateTree = SourceFrame.StateTree;
+	OutTransitionResult.SourceRootState = SourceFrame.RootState;
 PRAGMA_ENABLE_DEPRECATION_WARNINGS
 
-	NextTransition.TargetState = NextState;
-	NextTransition.CurrentState = FStateTreeStateHandle::Invalid;
-	NextTransition.CurrentRunStatus = Exec.LastTickStatus;
-	NextTransition.ChangeType = EStateTreeStateChangeType::Changed; 
-	NextTransition.Priority = Priority;
-
-	FStateTreeExecutionFrame& NewFrame = NextTransition.NextActiveFrames.AddDefaulted_GetRef();
-	NewFrame.StateTree = CurrentFrame.StateTree;
-	NewFrame.RootState = CurrentFrame.RootState;
-	NewFrame.ActiveTasksStatus = CurrentFrame.ActiveTasksStatus;
+	OutTransitionResult.TargetState = TargetState.StateHandle;
+	OutTransitionResult.CurrentState = FStateTreeStateHandle::Invalid;
+	OutTransitionResult.CurrentRunStatus = Exec.LastTickStatus;
+	OutTransitionResult.ChangeType = EStateTreeStateChangeType::Changed;
+	OutTransitionResult.Priority = Priority;
 
-	if (NextState == FStateTreeStateHandle::Invalid)
+	//@todo remove this when NextActiveFrames is deprecated
+	//For backward compatibility with ForceTransition
+	FStateTreeExecutionFrame& NewFrame = OutTransitionResult.NextActiveFrames.AddDefaulted_GetRef();
+	NewFrame.StateTree = SourceFrame.StateTree;
+	NewFrame.RootState = SourceFrame.RootState;
+	NewFrame.ActiveTasksStatus = SourceFrame.ActiveTasksStatus;
+	if (TargetState.StateHandle == FStateTreeStateHandle::Invalid)
 	{
 		NewFrame.ActiveStates = {};
 	}
 	else
 	{
-		NewFrame.ActiveStates = FStateTreeActiveStates(NextState, UE::StateTree::FActiveStateID::Invalid);
+		NewFrame.ActiveStates = FStateTreeActiveStates(TargetState.StateHandle, UE::StateTree::FActiveStateID::Invalid);
 	}
 }
 
@@ -4452,7 +4895,9 @@ bool FStateTreeExecutionContext::TriggerTransitions()
 	CSV_SCOPED_TIMING_STAT_EXCLUSIVE(StateTree_TriggerTransition);
 	UE_STATETREE_DEBUG_SCOPED_PHASE(this, EStateTreeUpdatePhase::TriggerTransitions);
 
-	FAllowDirectTransitionsScope AllowDirectTransitionsScope(*this); // Set flag for the scope of this function to allow direct transitions without buffering.
+	// Set flag for the scope of this function to allow direct transitions without buffering.
+	FAllowDirectTransitionsScope AllowDirectTransitionsScope(*this);
+
 	FStateTreeExecutionState& Exec = GetExecState();
 
 	if (EventQueue && EventQueue->HasEvents())
@@ -4461,7 +4906,7 @@ bool FStateTreeExecutionContext::TriggerTransitions()
 		UE_STATETREE_DEBUG_LOG_EVENT(this, Log, TEXT("Trigger transitions with events: %s"), *DebugGetEventsAsString());
 	}
 
-	NextTransition = FStateTreeTransitionResult();
+	RequestedTransition.Reset();
 
 	//
 	// Process transition requests
@@ -4470,11 +4915,13 @@ bool FStateTreeExecutionContext::TriggerTransitions()
 	{
 		// Find frame associated with the request.
 		const FStateTreeExecutionFrame* CurrentFrame = Exec.FindActiveFrame(Request.SourceFrameID);
-		if (CurrentFrame)
+		if (CurrentFrame && CurrentFrame->ActiveStates.Contains(Request.SourceStateID))
 		{
-			if (RequestTransition(*CurrentFrame, Request.TargetState, Request.Priority, /*TransitionEvent*/nullptr, Request.Fallback))
+			UE::StateTree::ExecutionContext::FStateHandleContext TargetState(CurrentFrame->StateTree, Request.TargetState);
+			if (RequestTransitionInternal(*CurrentFrame, Request.SourceStateID, TargetState, Request.Priority, FStateTreeSharedEvent(), Request.Fallback))
 			{
-				NextTransitionSource = FStateTreeTransitionSource(CurrentFrame->StateTree, EStateTreeTransitionSourceType::ExternalRequest, Request.TargetState, Request.Priority);
+				check(RequestedTransition.IsValid());
+				RequestedTransition->Source = FStateTreeTransitionSource(CurrentFrame->StateTree, EStateTreeTransitionSourceType::ExternalRequest, Request.TargetState, Request.Priority);
 			}
 		}
 	}
@@ -4679,7 +5126,7 @@ bool FStateTreeExecutionContext::TriggerTransitions()
 				}
 				
 				// No need to test the transition if same or higher priority transition has already been processed.
-				if (Transition.Priority <= NextTransition.Priority)
+				if (RequestedTransition.IsValid() && Transition.Priority <= RequestedTransition->Transition.Priority)
 				{
 					continue;
 				}
@@ -4702,7 +5149,8 @@ bool FStateTreeExecutionContext::TriggerTransitions()
 								*GetSafeStateName(CurrentFrame, CurrentFrame.ActiveStates.Last()), *State.Name.ToString(), *GetSafeStateName(CurrentFrame, Transition.State));
 
 							// Trigger Delayed Transition when the delay has passed.
-							if (RequestTransition(CurrentFrame, Transition.State, Transition.Priority, &DelayedTransition.CapturedEvent, Transition.Fallback))
+							UE::StateTree::ExecutionContext::FStateHandleContext TargetState(CurrentFrame.StateTree, Transition.State);
+							if (RequestTransitionInternal(CurrentFrame, Handler.StateID, TargetState, Transition.Priority, DelayedTransition.CapturedEvent, Transition.Fallback))
 							{
 								// If the transition was successfully requested with a specific event, consume and remove the event, it's been used.
 								if (DelayedTransition.CapturedEvent.IsValid() && Transition.bConsumeEventOnSelect)
@@ -4710,7 +5158,7 @@ bool FStateTreeExecutionContext::TriggerTransitions()
 									ConsumeEvent(DelayedTransition.CapturedEvent);
 								}
 
-								NextTransitionSource = FStateTreeTransitionSource(CurrentFrame.StateTree, FStateTreeIndex16(TransitionIndex), Transition.State, Transition.Priority);
+								RequestedTransition->Source = FStateTreeTransitionSource(CurrentFrame.StateTree, FStateTreeIndex16(TransitionIndex), Transition.State, Transition.Priority);
 								bTriggeredDelayedTransition = true;
 								break;
 							}
@@ -4723,45 +5171,14 @@ bool FStateTreeExecutionContext::TriggerTransitions()
 					}
 				}
 
-				TArray<const FStateTreeSharedEvent*, TInlineAllocator<8>> TransitionEvents;
-
-				if (Transition.Trigger == EStateTreeTransitionTrigger::OnEvent)
-				{
-					check(Transition.RequiredEvent.IsValid());
-
-					TConstArrayView<FStateTreeSharedEvent> EventsQueue = GetEventsToProcessView();
-					for (const FStateTreeSharedEvent& Event : EventsQueue)
-					{
-						check(Event.IsValid());
-						if (Transition.RequiredEvent.DoesEventMatchDesc(*Event))
-						{
-							TransitionEvents.Emplace(&Event);
-						}
-					}
-				}
-				else if (EnumHasAnyFlags(Transition.Trigger, EStateTreeTransitionTrigger::OnTick))
-				{
-					// Dummy event to make sure we iterate to loop below once.
-					TransitionEvents.Emplace(nullptr);
-				}
-				else if (EnumHasAnyFlags(Transition.Trigger, EStateTreeTransitionTrigger::OnDelegate))
-				{
-					if (Storage.IsDelegateBroadcasted(Transition.RequiredDelegateDispatcher))
-					{
-						// Dummy event to make sure we iterate to loop below once.
-						TransitionEvents.Emplace(nullptr);
-					}
-				}
-				else
-				{
-					ensureMsgf(false, TEXT("The trigger type is not supported."));
-				}
+				UE::StateTree::ExecutionContext::Private::FSharedEventInlineArray TransitionEvents;
+				UE::StateTree::ExecutionContext::Private::GetTriggerTransitionEvent(Transition, Storage, FStateTreeSharedEvent(), GetEventsToProcessView(), TransitionEvents);
 
-				for (const FStateTreeSharedEvent* TransitionEvent : TransitionEvents)
+				for (const FStateTreeSharedEvent& TransitionEvent : TransitionEvents)
 				{
 					bool bPassed = false; 
 					{
-						FCurrentlyProcessedTransitionEventScope TransitionEventScope(*this, TransitionEvent ? TransitionEvent->Get() : nullptr);
+						FCurrentlyProcessedTransitionEventScope TransitionEventScope(*this, TransitionEvent.IsValid() ? TransitionEvent.Get() : nullptr);
 						UE_STATETREE_DEBUG_TRANSITION_EVENT(this, FStateTreeTransitionSource(CurrentFrame.StateTree, FStateTreeIndex16(TransitionIndex), Transition.State, Transition.Priority), EStateTreeTraceEventType::OnEvaluating);
 						UE_STATETREE_DEBUG_SCOPED_PHASE(this, EStateTreeUpdatePhase::TransitionConditions);
 						bPassed = TestAllConditionsOnActiveInstances(CurrentParentFrame, CurrentFrame, Handler.StateHandle, Transition.ConditionsBegin, Transition.ConditionsNum);
@@ -4773,9 +5190,9 @@ bool FStateTreeExecutionContext::TriggerTransitions()
 						if (Transition.HasDelay())
 						{
 							uint32 TransitionEventHash = 0u;
-							if (TransitionEvent && TransitionEvent->IsValid())
+							if (TransitionEvent.IsValid())
 							{		
-								TransitionEventHash = GetTypeHash(*TransitionEvent->Get());
+								TransitionEventHash = GetTypeHash(TransitionEvent.Get());
 							}
 
 							const bool bIsDelayedTransitionExisting = Exec.DelayedTransitions.ContainsByPredicate([CurrentStateID =  Handler.StateID, TransitionIndex, TransitionEventHash](const FStateTreeTransitionDelayedState& DelayedState)
@@ -4799,9 +5216,9 @@ PRAGMA_DISABLE_DEPRECATION_WARNINGS
 PRAGMA_ENABLE_DEPRECATION_WARNINGS
 									DelayedState.TransitionIndex = FStateTreeIndex16(TransitionIndex);
 									DelayedState.TimeLeft = DelayDuration;
-									if (TransitionEvent && TransitionEvent->IsValid())
+									if (TransitionEvent.IsValid())
 									{
-										DelayedState.CapturedEvent = *TransitionEvent;
+										DelayedState.CapturedEvent = TransitionEvent;
 										DelayedState.CapturedEventHash = TransitionEventHash;
 									}
 
@@ -4822,15 +5239,16 @@ PRAGMA_ENABLE_DEPRECATION_WARNINGS
 						}
 
 						UE_STATETREE_DEBUG_TRANSITION_EVENT(this, FStateTreeTransitionSource(CurrentFrame.StateTree, FStateTreeIndex16(TransitionIndex), Transition.State, Transition.Priority), EStateTreeTraceEventType::OnRequesting);
-						if (RequestTransition(CurrentFrame, Transition.State, Transition.Priority, TransitionEvent, Transition.Fallback))
+						UE::StateTree::ExecutionContext::FStateHandleContext TargetState(CurrentFrame.StateTree, Transition.State);
+						if (RequestTransitionInternal(CurrentFrame, Handler.StateID, TargetState, Transition.Priority, TransitionEvent, Transition.Fallback))
 						{
 							// If the transition was successfully requested with a specific event, consume and remove the event, it's been used.
-							if (TransitionEvent && Transition.bConsumeEventOnSelect)
+							if (TransitionEvent.IsValid() && Transition.bConsumeEventOnSelect)
 							{
-								ConsumeEvent(*TransitionEvent);
+								ConsumeEvent(TransitionEvent);
 							}
 							
-							NextTransitionSource = FStateTreeTransitionSource(CurrentFrame.StateTree, FStateTreeIndex16(TransitionIndex), Transition.State, Transition.Priority);
+							RequestedTransition->Source = FStateTreeTransitionSource(CurrentFrame.StateTree, FStateTreeIndex16(TransitionIndex), Transition.State, Transition.Priority);
 							break;
 						}
 					}
@@ -4852,7 +5270,7 @@ PRAGMA_ENABLE_DEPRECATION_WARNINGS
 	//
 	bool bProcessSubTreeCompletion = true;
 
-	if (NextTransition.Priority == EStateTreeTransitionPriority::None
+	if (!RequestedTransition.IsValid()
 		&& (Exec.LastTickStatus != EStateTreeRunStatus::Running || Exec.bHasPendingCompletedState))
 	{
 		// Find the pending completed frame/state. Don't cache the result because this function is reentrant.
@@ -4863,6 +5281,8 @@ PRAGMA_ENABLE_DEPRECATION_WARNINGS
 		for (int32 FrameIndex = 0; FrameIndex < Exec.ActiveFrames.Num(); ++FrameIndex)
 		{
 			using namespace UE::StateTree;
+			using namespace UE::StateTree::ExecutionContext;
+			using namespace UE::StateTree::ExecutionContext::Private;
 
 			const FStateTreeExecutionFrame& CurrentFrame = Exec.ActiveFrames[FrameIndex];
 			const UStateTree* CurrentStateTree = CurrentFrame.StateTree;
@@ -4872,10 +5292,10 @@ PRAGMA_ENABLE_DEPRECATION_WARNINGS
 				if (FrameIndex == 0)
 				{
 					// If first frame, then complete the tree execution.
-					Exec.RequestedStop = ExecutionContext::GetPriorityRunStatus(Exec.RequestedStop, ExecutionContext::CastToRunStatus(FrameTasksStatus));
+					Exec.RequestedStop = GetPriorityRunStatus(Exec.RequestedStop, CastToRunStatus(FrameTasksStatus));
 					break;
 				}
-				else if (ExecutionContext::Private::bGlobalTasksCompleteOwningFrame)
+				else if (bGlobalTasksCompleteOwningFrame)
 				{
 					const int32 ParentFrameIndex = FrameIndex - 1;
 					FStateTreeExecutionFrame& ParentFrame = Exec.ActiveFrames[ParentFrameIndex];
@@ -4885,14 +5305,14 @@ PRAGMA_ENABLE_DEPRECATION_WARNINGS
 						// Set the parent linked state as last completed state, and update tick status to the status from the transition.
 						STATETREE_LOG(Verbose, TEXT("Completed subtree '%s' from global: %s"),
 							*GetSafeStateName(ParentFrame, ParentLinkedState),
-							*UEnum::GetDisplayValueAsText(ExecutionContext::CastToRunStatus(FrameTasksStatus)).ToString()
+							*UEnum::GetDisplayValueAsText(CastToRunStatus(FrameTasksStatus)).ToString()
 							);
 
 						const FCompactStateTreeState& State = ParentFrame.StateTree->States[ParentLinkedState.Index];
 						ParentFrame.ActiveTasksStatus.GetStatus(State).SetCompletionStatus(FrameTasksStatus);
 						Exec.bHasPendingCompletedState = true;
 
-						CurrentStatus = ExecutionContext::CastToRunStatus(FrameTasksStatus);
+						CurrentStatus = CastToRunStatus(FrameTasksStatus);
 						FrameIndexToStart = ParentFrameIndex;
 						StateIndexToStart = ParentFrame.ActiveStates.Num() - 1;
 						break;
@@ -4904,10 +5324,10 @@ PRAGMA_ENABLE_DEPRECATION_WARNINGS
 			{
 				const FStateTreeStateHandle CurrentHandle = CurrentFrame.ActiveStates[StateIndex];
 				const FCompactStateTreeState& State = CurrentStateTree->States[CurrentHandle.Index];
-				const UE::StateTree::ETaskCompletionStatus StateTasksStatus = CurrentFrame.ActiveTasksStatus.GetStatus(State).GetCompletionStatus();
-				if (StateTasksStatus != UE::StateTree::ETaskCompletionStatus::Running)
+				const ETaskCompletionStatus StateTasksStatus = CurrentFrame.ActiveTasksStatus.GetStatus(State).GetCompletionStatus();
+				if (StateTasksStatus != ETaskCompletionStatus::Running)
 				{
-					CurrentStatus = ExecutionContext::CastToRunStatus(StateTasksStatus);
+					CurrentStatus = CastToRunStatus(StateTasksStatus);
 					FrameIndexToStart = FrameIndex;
 					StateIndexToStart = StateIndex;
 					break;
@@ -4943,18 +5363,19 @@ PRAGMA_ENABLE_DEPRECATION_WARNINGS
 				// Check completion transitions
 				for (int32 StateIndex = CurrentStateIndexToStart; StateIndex >= 0; --StateIndex)
 				{
-					const FStateTreeStateHandle StateHandle = CurrentFrame.ActiveStates[StateIndex];
-					const FCompactStateTreeState& State = CurrentStateTree->States[StateHandle.Index];
+					const FStateTreeStateHandle CurrentStateHandle = CurrentFrame.ActiveStates[StateIndex];
+					const UE::StateTree::FActiveStateID CurrentStateID = CurrentFrame.ActiveStates.StateIDs[StateIndex];
+					const FCompactStateTreeState& CurrentState = CurrentStateTree->States[CurrentStateHandle.Index];
 
-					if (State.ShouldTickCompletionTransitions(bIsCurrentStatusSucceeded, bIsCurrentStatusFailed))
+					if (CurrentState.ShouldTickCompletionTransitions(bIsCurrentStatusSucceeded, bIsCurrentStatusFailed))
 					{
-						FCurrentlyProcessedStateScope StateScope(*this, StateHandle);
-						UE_STATETREE_DEBUG_SCOPED_STATE_PHASE(this, StateHandle, EStateTreeUpdatePhase::TriggerTransitions);
+						FCurrentlyProcessedStateScope StateScope(*this, CurrentStateHandle);
+						UE_STATETREE_DEBUG_SCOPED_STATE_PHASE(this, CurrentStateHandle, EStateTreeUpdatePhase::TriggerTransitions);
 
-						for (uint8 TransitionCounter = 0; TransitionCounter < State.TransitionsNum; ++TransitionCounter)
+						for (uint8 TransitionCounter = 0; TransitionCounter < CurrentState.TransitionsNum; ++TransitionCounter)
 						{
 							// All transition conditions must pass
-							const int16 TransitionIndex = State.TransitionsBegin + TransitionCounter;
+							const int16 TransitionIndex = CurrentState.TransitionsBegin + TransitionCounter;
 							const FCompactStateTransition& Transition = CurrentStateTree->Transitions[TransitionIndex];
 
 							// Skip disabled transitions
@@ -4972,7 +5393,7 @@ PRAGMA_ENABLE_DEPRECATION_WARNINGS
 								{
 									UE_STATETREE_DEBUG_TRANSITION_EVENT(this, FStateTreeTransitionSource(CurrentFrame.StateTree, FStateTreeIndex16(TransitionIndex), Transition.State, Transition.Priority), EStateTreeTraceEventType::OnEvaluating);
 									UE_STATETREE_DEBUG_SCOPED_PHASE(this, EStateTreeUpdatePhase::TransitionConditions);
-									bPassed = TestAllConditionsOnActiveInstances(CurrentParentFrame, CurrentFrame, StateHandle, Transition.ConditionsBegin, Transition.ConditionsNum);
+									bPassed = TestAllConditionsOnActiveInstances(CurrentParentFrame, CurrentFrame, CurrentStateHandle, Transition.ConditionsBegin, Transition.ConditionsNum);
 								}
 
 								if (bPassed)
@@ -4980,16 +5401,17 @@ PRAGMA_ENABLE_DEPRECATION_WARNINGS
 									// No delay allowed on completion conditions.
 									// No priority on completion transitions, use the priority to signal that state is selected.
 									UE_STATETREE_DEBUG_TRANSITION_EVENT(this, FStateTreeTransitionSource(CurrentFrame.StateTree, FStateTreeIndex16(TransitionIndex), Transition.State, Transition.Priority), EStateTreeTraceEventType::OnRequesting);
-									if (RequestTransition(CurrentFrame, Transition.State, EStateTreeTransitionPriority::Normal, /*TransitionEvent*/nullptr, Transition.Fallback))
+									UE::StateTree::ExecutionContext::FStateHandleContext TargetState(CurrentFrame.StateTree, Transition.State);
+									if (RequestTransitionInternal(CurrentFrame, CurrentStateID, TargetState, EStateTreeTransitionPriority::Normal, FStateTreeSharedEvent(), Transition.Fallback))
 									{
-										NextTransitionSource = FStateTreeTransitionSource(CurrentFrame.StateTree, FStateTreeIndex16(TransitionIndex), Transition.State, Transition.Priority);
+										RequestedTransition->Source = FStateTreeTransitionSource(CurrentFrame.StateTree, FStateTreeIndex16(TransitionIndex), Transition.State, Transition.Priority);
 										break;
 									}
 								}
 							}
 						}
 
-						if (NextTransition.Priority != EStateTreeTransitionPriority::None) //-V547
+						if (RequestedTransition.IsValid())
 						{
 							break;
 						}
@@ -4998,33 +5420,35 @@ PRAGMA_ENABLE_DEPRECATION_WARNINGS
 
 				// if a valid completion transition has already been found, the remaining transitions in parent frames won't have a higher priority than the found one
 				// so skip the remainder. this also prevented false positive warnings and ensures from STDebugger
-				if (NextTransition.Priority != EStateTreeTransitionPriority::None) //-V547
+				if (RequestedTransition.IsValid())
 				{
 					break;
 				}
 			}
 
 			// Handle the case where no transition was found.
-			if (NextTransition.Priority == EStateTreeTransitionPriority::None) //-V547
+			if (!RequestedTransition.IsValid())
 			{
 				STATETREE_LOG(Verbose, TEXT("Could not trigger completion transition, jump back to root state."));
 				UE_STATETREE_DEBUG_LOG_EVENT(this, Warning, TEXT("Could not trigger completion transition, jump back to root state."));
 
 				check(!Exec.ActiveFrames.IsEmpty());
 				FStateTreeExecutionFrame& RootFrame = Exec.ActiveFrames[0];
+				check(RootFrame.ActiveStates.Num() != 0);
 				FCurrentlyProcessedFrameScope RootFrameScope(*this, nullptr, RootFrame);
 				FCurrentlyProcessedStateScope RootStateScope(*this, FStateTreeStateHandle::Root);
 				
-				if (RequestTransition(RootFrame, FStateTreeStateHandle::Root, EStateTreeTransitionPriority::Normal))
+				UE::StateTree::ExecutionContext::FStateHandleContext TargetState(RootFrame.StateTree, FStateTreeStateHandle::Root);
+				if (RequestTransitionInternal(RootFrame, RootFrame.ActiveStates.StateIDs[0], TargetState))
 				{
-					NextTransitionSource = FStateTreeTransitionSource(RootFrame.StateTree, EStateTreeTransitionSourceType::Internal, FStateTreeStateHandle::Root, EStateTreeTransitionPriority::Normal);
+					RequestedTransition->Source = FStateTreeTransitionSource(RootFrame.StateTree, EStateTreeTransitionSourceType::Internal, FStateTreeStateHandle::Root, EStateTreeTransitionPriority::Normal);
 				}
 				else
 				{
 					STATETREE_LOG(Warning, TEXT("Failed to select root state. Stopping the tree with failure."));
 					UE_STATETREE_DEBUG_LOG_EVENT(this, Error, TEXT("Failed to select root state. Stopping the tree with failure."));
 
-					SetupNextTransition(RootFrame, FStateTreeStateHandle::Failed, EStateTreeTransitionPriority::Critical);
+					Exec.RequestedStop = UE::StateTree::ExecutionContext::GetPriorityRunStatus(Exec.RequestedStop, EStateTreeRunStatus::Failed);
 
 					// In this case we don't want to complete subtrees, we want to force the whole tree to stop.
 					bProcessSubTreeCompletion = false;
@@ -5034,10 +5458,10 @@ PRAGMA_ENABLE_DEPRECATION_WARNINGS
 	}
 
 	// Check if the transition was succeed/failed, if we're on a sub-tree, complete the subtree instead of transition.
-	if (NextTransition.TargetState.IsCompletionState() && bProcessSubTreeCompletion)
+	if (RequestedTransition.IsValid() && RequestedTransition->Transition.TargetState.IsCompletionState() && bProcessSubTreeCompletion)
 	{
 		// Check that the transition source frame is a sub-tree, the first frame (0 index) is not a subtree. 
-		const int32 SourceFrameIndex = Exec.IndexOfActiveFrame(NextTransition.SourceFrameID);
+		const int32 SourceFrameIndex = Exec.IndexOfActiveFrame(RequestedTransition->Transition.SourceFrameID);
 		if (SourceFrameIndex > 0)
 		{
 			const FStateTreeExecutionFrame& SourceFrame = Exec.ActiveFrames[SourceFrameIndex];
@@ -5047,10 +5471,10 @@ PRAGMA_ENABLE_DEPRECATION_WARNINGS
 
 			if (ParentLinkedState.IsValid())
 			{
-				const EStateTreeRunStatus RunStatus = NextTransition.TargetState.ToCompletionStatus();
+				const EStateTreeRunStatus RunStatus = RequestedTransition->Transition.TargetState.ToCompletionStatus();
 
 #if ENABLE_VISUAL_LOG
-				const int32 NextTransitionSourceIndex = SourceFrame.ActiveStates.IndexOfReverse(NextTransition.SourceStateID);
+				const int32 NextTransitionSourceIndex = SourceFrame.ActiveStates.IndexOfReverse(RequestedTransition->Transition.SourceStateID);
 				const FStateTreeStateHandle NextTransitionSourceState = NextTransitionSourceIndex != INDEX_NONE
 					? SourceFrame.ActiveStates[NextTransitionSourceIndex]
 					: FStateTreeStateHandle::Invalid;
@@ -5073,13 +5497,13 @@ PRAGMA_ENABLE_DEPRECATION_WARNINGS
 				// now starting from the linked parent state. If we run out of retires in the selection loop (e.g. very deep hierarchy)
 				// we will continue on next tick.
 				TriggerTransitionsFromFrameIndex = ParentFrameIndex;
-				NextTransition = FStateTreeTransitionResult();
+				RequestedTransition.Reset();
 				return false;
 			}
 		}
 	}
 
-	return NextTransition.TargetState.IsValid();
+	return RequestedTransition.IsValid();
 }
 
 TOptional<FStateTreeTransitionResult> FStateTreeExecutionContext::MakeTransitionResult(const FRecordedStateTreeTransitionResult& RecordedTransition) const
@@ -5112,6 +5536,7 @@ TOptional<FStateTreeTransitionResult> FStateTreeExecutionContext::MakeTransition
 		ExecutionFrame.ActiveStates = RecordedExecutionFrame.ActiveStates;
 		ExecutionFrame.ActiveTasksStatus = FStateTreeTasksCompletionStatus(*CompactFrame);
 		ExecutionFrame.bIsGlobalFrame = RecordedExecutionFrame.bIsGlobalFrame;
+		ExecutionFrame.ExternalDataBaseIndex = const_cast<FStateTreeExecutionContext*>(this)->CollectExternalData(RecordedExecutionFrame.StateTree);
 
 		FStateTreeFrameStateSelectionEvents& StateTreeFrameStateSelectionEvents = Result.NextActiveFrameEvents.AddDefaulted_GetRef();
 		for (int32 EventIdx = 0; EventIdx < RecordedExecutionFrame.EventIndices.Num(); EventIdx++)
@@ -5145,7 +5570,7 @@ TOptional<FStateTreeTransitionResult> FStateTreeExecutionContext::MakeTransition
 	const FStateTreeExecutionState& Exec = GetExecState();
 	const FStateTreeExecutionFrame* ExecFrame = Exec.ActiveFrames.FindByPredicate([StateTree = RecordedTransition.SourceStateTree, RootState = RecordedTransition.SourceRootState](const FStateTreeExecutionFrame& Frame)
 		{
-			return Frame.StateTree == StateTree && Frame.RootState == RootState;
+			return Frame.HasRoot(StateTree, RootState);
 		});
 	if (ExecFrame)
 	{
@@ -5207,162 +5632,329 @@ FRecordedStateTreeTransitionResult FStateTreeExecutionContext::MakeRecordedTrans
 			Result.SourceState = FoundSourceFrame->ActiveStates[ActiveStateIndex];
 		}
 	}
+	else
+	{
+PRAGMA_DISABLE_DEPRECATION_WARNINGS
+		Result.SourceStateTree = Transition.SourceStateTree;
+		Result.SourceRootState = Transition.SourceRootState;
+		Result.SourceState = Transition.SourceState;
+PRAGMA_ENABLE_DEPRECATION_WARNINGS
+	}
 	Result.TargetState = Transition.TargetState;
 	Result.Priority = Transition.Priority;
 
 	return Result;
 }
 
+// Deprecated
 bool FStateTreeExecutionContext::SelectState(const FStateTreeExecutionFrame& CurrentFrame,
 											const FStateTreeStateHandle NextState,
-											FStateSelectionResult& OutSelectionResult,
+	FSelectStateResult& OutSelectionResult,
 											const FStateTreeSharedEvent* TransitionEvent,
 											const EStateTreeSelectionFallback Fallback)
 {
-	TGuardValue<const FStateSelectionResult*> GuardValue(CurrentSelectionResult, &OutSelectionResult);
+	using namespace UE::StateTree;
+	using namespace UE::StateTree::ExecutionContext;
+	using namespace UE::StateTree::ExecutionContext::Private;
+
+	TArray<FActiveState, TInlineAllocator<ExpectedActiveStatePathLength>> CurrentActiveStatePath;
+	GetActiveStatePath(GetExecState().ActiveFrames, CurrentActiveStatePath);
+
+	FSelectStateArguments SelectStateArgs;
+	SelectStateArgs.ActiveStates = MakeConstArrayView(CurrentActiveStatePath);
+	SelectStateArgs.SourceState = FActiveState();
+	SelectStateArgs.TargetState = FStateHandleContext(CurrentFrame.StateTree, NextState);
+	SelectStateArgs.TransitionEvent = TransitionEvent != nullptr ? *TransitionEvent : FStateTreeSharedEvent();
+	SelectStateArgs.Fallback = Fallback;
+	return SelectState(SelectStateArgs, OutSelectionResult);
+}
+
+bool FStateTreeExecutionContext::SelectState(const FSelectStateArguments& Args, FSelectStateResult& OutSelectionResult)
+{
+	// 1. Find the target root state.
+	// 2. If the the target frame is active, then copy all previous states from the previous frames.
+	//   If any state inside the frame is completed (before the target), then the transition fails.
+	//   See EStateTreeStateSelectionRules::CompletedStateBeforeTargetFailTransition.
+	// 3. In the target frame, start adding states that match the source/target.
+	//   If the state is completed (before the target), then the transition fails.
+	//   See EStateTreeStateSelectionRules::CompletedStateBeforeTargetFailTransition.
+	// 4. New/Sustained states need to be reevaluated (see SelectStateInternal).
+	// 5. Else, handle fallback.
+	// 
+	// Source is from where the transition request occurs.
+	//  The source frame is valid but the source state can be invalid. It will be the top root state if needed.
+	//  It doesn't need to be in the same frame as the target.
+	// Target is where we wish to go. The selection can stop or select another state (depending on the state's type).
+	// Selected is the selection result.
+	//
+	// ExitState: If the state >= Target and is in the selected, then the transition is "sustained" and the instance data is untouched.
+	// ExitState: If the state not in the selected (removed state), then the transition is "changed" and the state is removed.
+	// EnterState: If the state >= Target and it is in the actives, then the transition is "sustained" and the instance data is untouched.
+	// EnterState: If the state >= Target and it is not in the actives (new state), then the transition is "changed" and the state is added.
+	// 
+	// Rules might impact the results.
+	// In the examples,  "New State if completed" is only valid if EStateTreeStateSelectionRules::CompletedTransitionStatesCreateNewStates
+	// Examples: Active | Source | Target | Selected
+	//
+	//           ABCD   | ABCD   | ABCDE  | ABCDEF, ABCD'EF
+	//  New D if completed, always new EF.
+	//  ExitState: Sustained -. Changed -, D. EnterState: Sustained -. Changed EF, D'EF.
+	//
+	//           ABCD   | AB     | ABI    | ABIJ, AB'IJ
+	//  New B if completed, always new IJ.
+	//  ExitState: Sustained -. Changed CD, BCD. EnterState: Sustained -. Changed IJ, B'IJ.
+	// 
+	//           ABCD   | AB     | ABC    | ABCD, ABCD', ABC'D', AB'C'D'
+	//  New D, CD, BCD if completed.
+	//  ExitState: Sustained CD, C, -, -. Changed -, D, CD, BCD. EnterState: Sustained CD, C, -, -. Changed -, D', C'D', B'C'D'.
+	//
+	// (if ReselectedParentCreatesNewChildStates)
+	//           ABCD   | ABC    | AB     | ABC'D', AB'C'D'
+	//  Select a parent state. New B always new CD
+	//  ExitState: Sustained B, -. Changed CD, BCD. EnterState: Sustained B, -. Changed C'D', B'C'D'
+	//
+	// (if not ReselectedStateCreatesNewStates)
+	//           ABCD   | ABC     | AB    | ABCD, ABCD', ABC'D', AB'C'D'
+	//  New D, CD, BCD if completed.
+	//  ExitState: Sustained BCD, BC, B, -. Changed -, D, CD, BCD. EnterState: Sustained BCD, BC, B, -. Changed -, D', C'D', B'C'D'.
+	// 
+	//           ABCD   | AB     | AX     | AXY
+	//  Source is not in target. New XY.
+	//  ExitState: Sustained -. Changed B. EnterState: Sustained -. Changed X
+	//
+	//           ABCD   | AB     | AB     | ABCD, ABCD', ABC'D', AB'C'D'
+	//  Source is target. New D, CD, BCD if completed.
+	//  ExitState: Sustained BCD, BC, B, -. Changed -, D, CD, BCD. EnterState: Sustained BCD, BC, B, -. Changed -, D', C'D', B'C'D'.
 
-	const FStateTreeExecutionState& Exec = GetExecState();
+	using namespace UE::StateTree;
+	using namespace UE::StateTree::ExecutionContext;
+	using namespace UE::StateTree::ExecutionContext::Private;
 
-	if (Exec.ActiveFrames.IsEmpty())
+	FStateTreeExecutionState& Exec = GetExecState();
+
+	// Not a valid target
+	if (!ensure(Args.TargetState.StateHandle.IsValid() && !Args.TargetState.StateHandle.IsCompletionState()))
 	{
-		STATETREE_LOG(Error, TEXT("%hs: SelectState can only be called on initialized tree.  '%s' using StateTree '%s'."),
-			__FUNCTION__, *GetNameSafe(&Owner), *GetFullNameSafe(&RootStateTree));
 		return false;
 	}
-	
-	if (!NextState.IsValid())
+	if (!ensure(Args.TargetState.StateTree && Args.TargetState.StateTree->GetStateFromHandle(Args.TargetState.StateHandle)))
 	{
 		return false;
 	}
-
-	// Walk towards the root from current state.
-	TArray<FStateTreeStateHandle, TInlineAllocator<FStateTreeActiveStates::MaxStates>> PathToNextState;
-	FStateTreeStateHandle CurrState = NextState;
-	while (CurrState.IsValid())
+	// Not a valid source. Note Source can be a GlobalTask (not in active state)
+	if (!ensure(Args.SourceState.GetFrameID().IsValid()))
 	{
-		if (PathToNextState.Num() == FStateTreeActiveStates::MaxStates)
+		return false;
+	}
+
+	TArray<FStateHandleContext, TInlineAllocator<FStateTreeActiveStates::MaxStates>> PathToTargetState;
+	if (!GetStatesListToState(Args.TargetState, PathToTargetState))
 		{
 			STATETREE_LOG(Error, TEXT("%hs: Reached max execution depth when trying to select state %s from '%s'.  '%s' using StateTree '%s'."),
-				__FUNCTION__, *GetSafeStateName(CurrentFrame, NextState), *GetStateStatusString(Exec), *GetNameSafe(&Owner), *GetFullNameSafe(&RootStateTree));
+			__FUNCTION__,
+			*GetSafeStateName(Args.TargetState.StateTree, Args.TargetState.StateHandle),
+			*GetStateStatusString(Exec),
+			*GetNameSafe(&Owner),
+			*GetFullNameSafe(&RootStateTree)
+		);
 			return false;
 		}
-		// Store the states that are in between the 'NextState' and common ancestor. 
-		PathToNextState.Push(CurrState);
-		CurrState = CurrentFrame.StateTree->States[CurrState.Index].Parent;
-	}
+	check(!PathToTargetState.IsEmpty());
 
-	Algo::Reverse(PathToNextState);
+	const FExecutionFrameKey TargetFrameKey = FExecutionFrameKey(Args.TargetState.StateTree, PathToTargetState[0].StateHandle);
+	const FCompactStateTreeFrame* TargetFrame = TargetFrameKey.StateTree->GetFrameFromHandle(TargetFrameKey.RootState);
+	checkf(TargetFrame, TEXT("The frame was not compiled."));
 
-	const UStateTree* NextStateTree = CurrentFrame.StateTree;
-	const FStateTreeStateHandle NextRootState = PathToNextState[0]; 
-
-	// Find the frame that the next state belongs to.
-	int32 CurrentFrameIndex = INDEX_NONE;
-	int32 CurrentStateTreeIndex = INDEX_NONE;
+	// Build the source path from the active path. Includes all the states from the previous frames.
+	TArrayView<const FActiveState> SourceStates;
+	{
+		if (Args.SourceState.GetStateID().IsValid())
+		{
+			const int32 FoundIndex = FActiveStatePath::IndexOf(Args.ActiveStates, Args.SourceState);
+			if (FoundIndex == INDEX_NONE)
+			{
+				STATETREE_LOG(Error, TEXT("%hs: The source do not exist in the active path when trying to select state %s from '%s'. '%s' using StateTree '%s'."),
+					__FUNCTION__,
+					*GetSafeStateName(Args.TargetState.StateTree, Args.TargetState.StateHandle),
+					*GetStateStatusString(Exec),
+					*GetNameSafe(&Owner),
+					*GetFullNameSafe(&RootStateTree)
+				);
+				return false;
+			}
 
-	for (int32 FrameIndex = Exec.ActiveFrames.Num() - 1; FrameIndex >= 0; FrameIndex--)
+			SourceStates = Args.ActiveStates.Left(FoundIndex + 1);
+		}
+		else
 	{
-		const FStateTreeExecutionFrame& Frame = Exec.ActiveFrames[FrameIndex]; 
-		if (Frame.StateTree == NextStateTree)
+			// Pick the source frame first state.
+			//This is usually when the request is out of the scope or from a global task. It should start from the root of the frame.
+			int32 ActiveStateIndex = 0;
+			for (; ActiveStateIndex < Args.ActiveStates.Num(); ++ActiveStateIndex)
 		{
-			CurrentStateTreeIndex = FrameIndex;
-			if (Frame.RootState == NextRootState)
+				const FActiveFrameID CurrentActiveFrameID = Args.ActiveStates[ActiveStateIndex].GetFrameID();
+				if (CurrentActiveFrameID == Args.SourceState.GetFrameID())
 			{
-				CurrentFrameIndex = FrameIndex;
 				break;
 			}
 		}
+			if (ActiveStateIndex < Args.ActiveStates.Num())
+			{
+				SourceStates = Args.ActiveStates.Left(ActiveStateIndex + 1);
 	}
-
-	// Copy common frames over.
-	// ReferenceCurrentFrame is the original of the last copied frame. It will be used to keep track if we are following the current active frames and states.
-	const FStateTreeExecutionFrame* CurrentFrameInActiveFrames  = nullptr;
-	if (CurrentFrameIndex != INDEX_NONE)
+			else
+			{
+				// SourceFrame exists, so the state must be in the active states.
+				if (!ensure(Args.ActiveStates.Num() == 0))
 	{
-		const int32 NumCommonFrames = CurrentFrameIndex + 1;
-		OutSelectionResult = FStateSelectionResult(MakeArrayView(Exec.ActiveFrames.GetData(), NumCommonFrames));
-		CurrentFrameInActiveFrames  = &Exec.ActiveFrames[CurrentFrameIndex];
+					return false;
+				}
+			}
 	}
-	else if (CurrentStateTreeIndex != INDEX_NONE)
-	{
-		// If we could not find a common frame, we assume that we jumped to different subtree in same asset.
-		const int32 NumCommonFrames = CurrentStateTreeIndex + 1;
-		OutSelectionResult = FStateSelectionResult(MakeArrayView(Exec.ActiveFrames.GetData(), NumCommonFrames));
-		CurrentFrameInActiveFrames  = &Exec.ActiveFrames[CurrentStateTreeIndex];
 	}
-	else
+	checkf(SourceStates.Num() == 0 || FActiveStatePath::StartsWith(Args.ActiveStates, SourceStates), TEXT("Source is part of the active path."));
+	FStateTreeExecutionFrame* SourceExecFrame = FindExecutionFrame(Args.SourceState.GetFrameID(), MakeArrayView(Exec.ActiveFrames), MakeArrayView(OutSelectionResult.TemporaryFrames));
+	if (!ensure(SourceExecFrame))
 	{
-		STATETREE_LOG(Error, TEXT("%hs: Encountered unrecognized state %s during state selection from '%s'.  '%s' using StateTree '%s'."),
-			__FUNCTION__, *GetNameSafe(NextStateTree), *GetStateStatusString(Exec), *GetNameSafe(&Owner), *GetFullNameSafe(NextStateTree));
 		return false;
 	}
 	
-	// Append in between state in reverse order, they were collected from leaf towards the root.
-	// Note: NextState will be added by SelectStateInternal() if conditions pass.
-	const int32 LastFrameIndex = OutSelectionResult.FramesNum() - 1;
-	FStateTreeExecutionFrame& LastFrame = OutSelectionResult.GetSelectedFrames()[LastFrameIndex];
+	const int32 SelectedStatesSlack = Args.ActiveStates.Num() + 10; // Number to help with buffer reallocation
+	OutSelectionResult.SelectedStates.Empty(SelectedStatesSlack);
 
-	// Find index of the first state to be evaluated.
-	int32 FirstNewStateIndex = 0;
-	if (CurrentFrameIndex != INDEX_NONE)
+	// Find the Target inside the ActiveStates.
+	//Look in the SourceStates because you can select in a different frame up the tree but not down the tree.
+	FActiveFrameID TargetFrameID;
 	{
-		// If LastFrame.ActiveStates is a subset of PathToNextState (e.g when someone use "TryEnter" selection behavior and then make a transition to it's child or if one is reentering the same state).
-		// In such case loop below won't break on anything and FirstNewStateIndex will be incorrectly 0, thus we initialize it to be right after the shorter range.
-		FirstNewStateIndex = FMath::Max(0, FMath::Min(PathToNextState.Num(), LastFrame.ActiveStates.Num()) - 1);
-		for (int32 Index = 0; Index < FMath::Min(PathToNextState.Num(), LastFrame.ActiveStates.Num()); ++Index)
+		if (SourceExecFrame->HasRoot(TargetFrameKey))
 		{
-			if (LastFrame.ActiveStates[Index] != PathToNextState[Index])
+			TargetFrameID = SourceExecFrame->FrameID;
+
+			// Copy active to the new selected until we reach the frame.
+			FActiveFrameID CurrentFrameID;
+			for (const FActiveState& ActiveState : Args.ActiveStates)
+	{
+				if (CurrentFrameID != ActiveState.GetFrameID())
+		{
+					CurrentFrameID = ActiveState.GetFrameID();
+					if (ActiveState.GetFrameID() == Args.SourceState.GetFrameID())
+					{
+						break;
+					}
+					OutSelectionResult.SelectedFrames.Add(CurrentFrameID);
+				}
+				OutSelectionResult.SelectedStates.Add(ActiveState);
+			}
+		}
+		else if (SourceStates.Num() >= 0)
+		{
+			// Can jump to a state from a previous frame.
+			//Find the common frame or the first frame with the same tree asset.
+			int32 FoundActiveStateIndex = SourceStates.Num() - 1;
+			{
+				bool bFoundRootState = false;
+				bool bFoundStateTree = false;
+				FActiveFrameID CurrentFrameID;
+				for (; FoundActiveStateIndex >= 0; --FoundActiveStateIndex)
+				{
+					const FActiveFrameID CurrentActiveFrameID = SourceStates[FoundActiveStateIndex].GetFrameID();
+					if (CurrentActiveFrameID != CurrentFrameID)
+					{
+						FStateTreeExecutionFrame* ProcessFrame = FindExecutionFrame(CurrentActiveFrameID, MakeArrayView(Exec.ActiveFrames), MakeArrayView(OutSelectionResult.TemporaryFrames));
+						if (!ensure(ProcessFrame))
+						{
+							return false;
+						}
+						if (ProcessFrame->StateTree == TargetFrameKey.StateTree)
+						{
+							bFoundStateTree = true;
+							if (ProcessFrame->RootState == TargetFrameKey.RootState)
+							{
+								bFoundRootState = true;
+							}
+							else if (bFoundRootState)
+							{
+								TargetFrameID = CurrentActiveFrameID;
+								break;
+							}
+						}
+						else if (bFoundStateTree)
 			{
-				FirstNewStateIndex = Index;
 				break;
 			}
 		}
 	}
 
-	ensureMsgf(LastFrame.ActiveStates.Num() >= FirstNewStateIndex, TEXT("ActiveTasksStatus won't be in sync with the amount of states."));
-	LastFrame.ActiveStates.SetNum(FirstNewStateIndex);
-
-	// Existing state's data is safe to access during select.
-	LastFrame.NumCurrentlyActiveStates = static_cast<uint8>(LastFrame.ActiveStates.Num());
-
-	FStateSelectionResult InitialSelection;
+				if (FoundActiveStateIndex < 0 && !bFoundStateTree)
+				{
+					STATETREE_LOG(Error, TEXT("%hs: Encountered unrecognized state %s during state selection from '%s'.  '%s' using StateTree '%s'."),
+						__FUNCTION__,
+						*GetSafeStateName(Args.TargetState.StateTree, Args.TargetState.StateHandle),
+						*GetStateStatusString(Exec),
+						*GetNameSafe(&Owner),
+						*GetFullNameSafe(&RootStateTree)
+					);
+					return false;
+				}
 
-	if (Fallback == EStateTreeSelectionFallback::NextSelectableSibling)
+				// Copy active states to the new selected until we reach the frame.
+				for (int32 ActiveStateIndex = 0; ActiveStateIndex < FoundActiveStateIndex+1; ++ActiveStateIndex)
+				{
+					const FActiveState& ActiveState = Args.ActiveStates[ActiveStateIndex];
+					if (CurrentFrameID != ActiveState.GetFrameID())
+					{
+						CurrentFrameID = ActiveState.GetFrameID();
+						OutSelectionResult.SelectedFrames.Add(CurrentFrameID);
+					}
+					OutSelectionResult.SelectedStates.Add(ActiveState);
+				}
+			}
+		}
+		else
 	{
-		InitialSelection = OutSelectionResult;
+			// SourceExecFrame do not matches and there are no source states.
+			STATETREE_LOG(Error, TEXT("%hs: Encountered out of range state %s during state selection from '%s'.  '%s' using StateTree '%s'."),
+				__FUNCTION__,
+				*GetSafeStateName(Args.TargetState.StateTree, Args.TargetState.StateHandle),
+				*GetStateStatusString(Exec),
+				*GetNameSafe(&Owner),
+				*GetFullNameSafe(&RootStateTree)
+			);
+			return false;
+		}
 	}
 	
-	// We take copy of the last frame and assign it later, as SelectStateInternal() might change the array and invalidate the pointer.
-	const FStateTreeExecutionFrame* CurrentParentFrame = LastFrameIndex > 0 ? &OutSelectionResult.GetSelectedFrames()[LastFrameIndex - 1] : nullptr;
-
-	// Path from the first new state up to the NextState
-	TConstArrayView<FStateTreeStateHandle> NewStatesPathToNextState(&PathToNextState[FirstNewStateIndex], PathToNextState.Num() - FirstNewStateIndex);
+	FSelectStateInternalArguments InternalArgs;
+	InternalArgs.MissingActiveStates = Args.ActiveStates.Mid(OutSelectionResult.SelectedStates.Num());
+	InternalArgs.MissingSourceStates = SourceStates.Mid(OutSelectionResult.SelectedStates.Num());
+	InternalArgs.MissingStatesToReachTarget = PathToTargetState;
+	InternalArgs.MissingSourceFrameID = TargetFrameID;
 
-	if (SelectStateInternal(CurrentParentFrame, OutSelectionResult.GetSelectedFrames()[LastFrameIndex], CurrentFrameInActiveFrames, NewStatesPathToNextState, OutSelectionResult, TransitionEvent))
+	// Add the state and check the prerequisites.
+	if (SelectStateInternal(Args, InternalArgs, OutSelectionResult))
 	{
 		return true;
 	}
 
-	// Failed to Select Next State, handle fallback here
-	// Return true on the first next sibling that gets selected successfully
-	if (Fallback == EStateTreeSelectionFallback::NextSelectableSibling && PathToNextState.Num() >= 2)
+	// Failed to Select Target State, handle fallback here
+	if (Args.Fallback == EStateTreeSelectionFallback::NextSelectableSibling
+		&& PathToTargetState.Num() >= 2 // we are not selecting the root
+		&& TargetFrameID.IsValid()) // we are not selecting a new frame on purpose
 	{
-		const FStateTreeStateHandle Parent = PathToNextState.Last(1);
-		if (Parent.IsValid())
+		const FStateHandleContext Parent = PathToTargetState.Last(1);
+		FStateHandleContext& Target = PathToTargetState.Last();
+		if (ensure(Parent.StateTree && Parent.StateHandle.IsValid() && Parent.StateTree == Target.StateTree && Target.StateHandle.IsValid()))
 		{
-			const FCompactStateTreeState& ParentState = CurrentFrame.StateTree->States[Parent.Index];
-
-			uint16 ChildState = CurrentFrame.StateTree->States[NextState.Index].GetNextSibling();
-			for (; ChildState < ParentState.ChildrenEnd; ChildState = CurrentFrame.StateTree->States[ChildState].GetNextSibling())
+			// Get the next sibling
+			const FCompactStateTreeState& ParentState = Parent.StateTree->States[Parent.StateHandle.Index];
+			for (uint16 ChildStateIndex = Target.StateTree->States[Target.StateHandle.Index].GetNextSibling(); ChildStateIndex < ParentState.ChildrenEnd; ChildStateIndex = Target.StateTree->States[Target.StateHandle.Index].GetNextSibling())
 			{
-				FStateTreeStateHandle ChildStateHandle = FStateTreeStateHandle(ChildState);
+				const FStateTreeStateHandle NextChildStateHandle = FStateTreeStateHandle(ChildStateIndex);
 
-				// Start selection from blank slate.
-				OutSelectionResult = InitialSelection;
-	
-				// We take copy of the last frame and assign it later, as SelectStateInternal() might change the array and invalidate the pointer.
-				CurrentParentFrame = LastFrameIndex > 0 ? &OutSelectionResult.GetSelectedFrames()[LastFrameIndex - 1] : nullptr; 
-				if (SelectStateInternal(CurrentParentFrame, OutSelectionResult.GetSelectedFrames()[LastFrameIndex], CurrentFrameInActiveFrames, {ChildStateHandle}, OutSelectionResult))
+				FSelectStateArguments NewArgs = Args;
+				NewArgs.TargetState.StateHandle = NextChildStateHandle;
+				Target = FStateHandleContext(Parent.StateTree, NextChildStateHandle);
+				if (SelectStateInternal(NewArgs, InternalArgs, OutSelectionResult))
 				{
 					return true;
 				}
@@ -5373,54 +5965,168 @@ bool FStateTreeExecutionContext::SelectState(const FStateTreeExecutionFrame& Cur
 	return false;
 }
 
+// Deprecated
 bool FStateTreeExecutionContext::SelectStateInternal(
 	const FStateTreeExecutionFrame* CurrentParentFrame,
 	FStateTreeExecutionFrame& CurrentFrame,
 	const FStateTreeExecutionFrame* CurrentFrameInActiveFrames,
 	TConstArrayView<FStateTreeStateHandle> PathToNextState,
-	FStateSelectionResult& OutSelectionResult,
+	FSelectStateResult& OutSelectionResult,
 	const FStateTreeSharedEvent* TransitionEvent)
+{
+	return false;
+}
+
+bool FStateTreeExecutionContext::SelectStateInternal(
+	const FSelectStateArguments& Args,
+	const FSelectStateInternalArguments& InternalArgs,
+	FSelectStateResult& OutSelectionResult)
 {
 	CSV_SCOPED_TIMING_STAT_EXCLUSIVE(StateTree_SelectState);
 
+	using namespace UE::StateTree;
+	using namespace UE::StateTree::ExecutionContext;
+	using namespace UE::StateTree::ExecutionContext::Private;
+
 	const FStateTreeExecutionState& Exec = GetExecState();
 
-	check(!PathToNextState.IsEmpty());
-	const FStateTreeStateHandle NextStateHandle = PathToNextState[0];
-	if (!NextStateHandle.IsValid())
+	if (InternalArgs.MissingStatesToReachTarget.IsEmpty())
 	{
-		// Trying to select non-existing state.
-		STATETREE_LOG(Error, TEXT("%hs: Trying to select invalid state from '%s'.  '%s' using StateTree '%s'."),
-			__FUNCTION__, *GetStateStatusString(Exec), *GetNameSafe(&Owner), *GetFullNameSafe(CurrentFrame.StateTree));
-		return false;
+		return true;
 	}
 
-	FCurrentlyProcessedFrameScope FrameScope(*this, CurrentParentFrame, CurrentFrame);
-	FCurrentlyProcessedStateScope NextStateScope(*this, NextStateHandle);
-	FCurrentFrameStateSelectionEventsScope CapturedEventsScope(*this, OutSelectionResult.GetFramesStateSelectionEvents().Last());
-
-	const UStateTree* CurrentStateTree = CurrentFrame.StateTree;
-	const FCompactStateTreeState& NextState = CurrentStateTree->States[NextStateHandle.Index];
+	const FStateHandleContext& NextStateContext = InternalArgs.MissingStatesToReachTarget[0];
+	const UStateTree* NextStateTree = NextStateContext.StateTree;
+	const FStateTreeStateHandle NextStateHandle = NextStateContext.StateHandle;
+	if (!ensure(NextStateTree != nullptr && NextStateHandle.IsValid()))
+	{
+		return false;
+	}
 
-	if (NextState.bEnabled == false)
+	const FCompactStateTreeState& NextState = InternalArgs.MissingStatesToReachTarget[0].StateTree->States[NextStateHandle.Index];
+	if (!NextState.bEnabled)
 	{
-		// Do not select disabled state
 		STATETREE_LOG(VeryVerbose, TEXT("%hs: Ignoring disabled state '%s'.  '%s' using StateTree '%s'."),
-			__FUNCTION__, *GetSafeStateName(CurrentFrame, NextStateHandle), *GetNameSafe(&Owner), *GetFullNameSafe(CurrentFrame.StateTree));
+			__FUNCTION__,
+			*NextState.Name.ToString(),
+			*GetNameSafe(&Owner),
+			*GetFullNameSafe(NextStateTree)
+		);
+		return false;
+	}
+	if (NextState.SelectionBehavior == EStateTreeStateSelectionBehavior::None)
+	{
+		STATETREE_LOG(VeryVerbose, TEXT("%hs: Selection Behavior is none for state '%s'.  '%s' using StateTree '%s'."),
+			__FUNCTION__,
+			*NextState.Name.ToString(),
+			*GetNameSafe(&Owner),
+			*GetFullNameSafe(NextStateTree)
+		);
 		return false;
 	}
 
-	UE_STATETREE_DEBUG_SCOPED_STATE_PHASE(this, NextStateHandle, EStateTreeUpdatePhase::StateSelection);
+	// Does it need a new root tree. It can create a new frame or reused an existing one.
+	bool bNewFrameCreated = false;
+	{
+		if (InternalArgs.MissingSourceStates.Num() == 0 && !InternalArgs.MissingSourceFrameID.IsValid())
+		{
+			// Must be in a valid selection (state are already been selected).
+			if (OutSelectionResult.SelectedStates.Num() == 0)
+			{
+				STATETREE_LOG(Verbose, TEXT("%hs: No root state to select '%s'.  '%s' using StateTree '%s'."),
+					__FUNCTION__,
+					*NextState.Name.ToString(),
+					*GetNameSafe(&Owner),
+					*GetFullNameSafe(NextStateTree)
+				);
+				return false;
+			}
 
-	// The state cannot be directly selected.
-	if (NextState.SelectionBehavior == EStateTreeStateSelectionBehavior::None)
+			const FActiveState PreviousSelectedState = OutSelectionResult.SelectedStates.Last();
+			FStateTreeExecutionFrame* PreviousFramePtr = FindExecutionFrame(PreviousSelectedState.GetFrameID(), MakeArrayView(GetExecState().ActiveFrames), MakeArrayView(OutSelectionResult.TemporaryFrames));
+			check(PreviousFramePtr);
+			check(PreviousFramePtr->StateTree == NextStateTree);
+
+			if (PreviousSelectedState.GetStateHandle() != NextState.Parent)
+			{
+				const FExecutionFrameKey FrameKey = FExecutionFrameKey(NextStateContext.StateTree, NextStateContext.StateHandle);
+				const FStateTreeExecutionFrame& NewFrame = OutSelectionResult.MakeAndAddTemporaryFrameWithNewRoot(FActiveFrameID(Storage.GenerateUniqueId()), FrameKey, *PreviousFramePtr);
+				OutSelectionResult.SelectedFrames.Add(NewFrame.FrameID);
+				bNewFrameCreated = true;
+			}
+		}
+		else
+		{
+			// We are building the path toward the desired target.
+			const FActiveFrameID MissingFrameID = InternalArgs.MissingSourceStates.Num() != 0 ? InternalArgs.MissingSourceStates[0].GetFrameID() : InternalArgs.MissingSourceFrameID;
+			FStateTreeExecutionFrame* MissingFramePtr = FindExecutionFrame(MissingFrameID, MakeArrayView(GetExecState().ActiveFrames), MakeArrayView(OutSelectionResult.TemporaryFrames));
+			check(MissingFramePtr);
+			if (!ensure(MissingFramePtr->StateTree == NextStateTree))
 	{
 		return false;
 	}
 
-	const UStateTree* NextLinkedStateAsset = NextState.LinkedAsset;
+			FActiveFrameID NextFrameID = MissingFrameID;
+			if (InternalArgs.MissingSourceStates.Num() != 0)
+			{
+				if (!NextState.Parent.IsValid() && MissingFramePtr->RootState != NextStateHandle)
+				{
+					const FExecutionFrameKey FrameKey = FExecutionFrameKey(NextStateContext.StateTree, NextStateContext.StateHandle);
+					const FStateTreeExecutionFrame& NewFrame = OutSelectionResult.MakeAndAddTemporaryFrameWithNewRoot(FActiveFrameID(Storage.GenerateUniqueId()), FrameKey, *MissingFramePtr);
+					NextFrameID = NewFrame.FrameID;
+					bNewFrameCreated = true;
+				}
+			}
+			else if (MissingFramePtr->RootState != NextStateHandle)
+			{
+				const FExecutionFrameKey FrameKey = FExecutionFrameKey(NextStateContext.StateTree, NextStateContext.StateHandle);
+				const FStateTreeExecutionFrame& NewFrame = OutSelectionResult.MakeAndAddTemporaryFrameWithNewRoot(FActiveFrameID(Storage.GenerateUniqueId()), FrameKey, *MissingFramePtr);
+				NextFrameID = NewFrame.FrameID;
+				bNewFrameCreated = true;
+			}
+
+			// Add new frames.
+			if (OutSelectionResult.SelectedFrames.Num() > 0)
+			{
+				if (OutSelectionResult.SelectedFrames.Last() != NextFrameID)
+				{
+					OutSelectionResult.SelectedFrames.Add(NextFrameID);
+				}
+			}
+			else 
+			{
+				OutSelectionResult.SelectedFrames.Add(NextFrameID);
+			}
+
+			TOptional<bool> bSelected = SelectStateFromSourceInternal(Args, InternalArgs, OutSelectionResult, *MissingFramePtr, NextState, NextStateHandle, bNewFrameCreated);
+			if (bSelected.IsSet())
+			{
+				return bSelected.GetValue();
+			}
+		}
+	}
+
+	// Use the ID to get the ExecutionFrame because it is saved in an array (and this is a recursive function), the array can grow and the pointer won't be valid.
+	check(OutSelectionResult.SelectedFrames.Num() != 0);
+	UE::StateTree::FActiveFrameID NextFrameID = OutSelectionResult.SelectedFrames.Last();
+	const FStateTreeExecutionFrame* NextFramePtr = FindExecutionFrame(NextFrameID, MakeConstArrayView(Exec.ActiveFrames), MakeConstArrayView(OutSelectionResult.TemporaryFrames));
+	check(NextFramePtr);
+
+	const FStateTreeExecutionFrame& NextFrame = *NextFramePtr;
+	const FStateTreeExecutionFrame* NextParentFrame = nullptr;
+	if (OutSelectionResult.SelectedFrames.Num() > 1)
+	{
+		NextParentFrame = FindExecutionFrame(OutSelectionResult.SelectedFrames.Last(1), MakeConstArrayView(Exec.ActiveFrames), MakeConstArrayView(OutSelectionResult.TemporaryFrames));
+		check(NextParentFrame);
+	}
+
+	// Save the current result to use SelectionEvents in GetDataView
+	TGuardValue<FSelectStateResult*> GuardCurrentlyProcessedStateSelectionResult(CurrentlyProcessedStateSelectionResult, &OutSelectionResult);
+
+	UE_STATETREE_DEBUG_SCOPED_STATE_PHASE(this, NextStateHandle, EStateTreeUpdatePhase::StateSelection);
 
 	// Look up linked state overrides
+	const UStateTree* NextLinkedStateAsset = NextState.LinkedAsset;
 	const FInstancedPropertyBag* NextLinkedStateParameterOverride = nullptr;
 	if (NextState.Type == EStateTreeStateType::LinkedAsset)
 	{
@@ -5430,16 +6136,24 @@ bool FStateTreeExecutionContext::SelectStateInternal(
 			NextLinkedStateParameterOverride = &Override->GetParameters();
 
 			STATETREE_LOG(VeryVerbose, TEXT("%hs: In state '%s', overriding linked asset '%s' with '%s'. '%s' using StateTree '%s'."),
-					__FUNCTION__, *GetSafeStateName(CurrentFrame, NextStateHandle),
-					*GetFullNameSafe(NextState.LinkedAsset), *GetFullNameSafe(NextLinkedStateAsset),
-					*GetNameSafe(&Owner), *GetFullNameSafe(CurrentFrame.StateTree));
+				__FUNCTION__,
+				*GetSafeStateName(NextFrame, NextStateHandle),
+				*GetFullNameSafe(NextState.LinkedAsset),
+				*GetFullNameSafe(NextLinkedStateAsset),
+				*GetNameSafe(&Owner),
+				*GetFullNameSafe(NextFrame.StateTree)
+			);
 		}
 	}
 
+	// Update state parameters.
 	if (NextState.ParameterDataHandle.IsValid())
 	{
-		// Instantiate state parameters if not done yet.
-		FStateTreeDataView NextStateParametersView = GetDataViewOrTemporary(CurrentParentFrame, CurrentFrame, NextState.ParameterDataHandle);
+		FCurrentlyProcessedFrameScope FrameScope(*this, NextParentFrame, NextFrame);
+		FCurrentlyProcessedStateScope NextStateScope(*this, NextStateHandle);
+
+		// Instantiate temporary state parameters if not done yet.
+		FStateTreeDataView NextStateParametersView = GetDataViewOrTemporary(NextParentFrame, NextFrame, NextState.ParameterDataHandle);
 		if (!NextStateParametersView.IsValid())
 		{
 			// Allocate temporary instance for parameters if the state has params.
@@ -5448,7 +6162,7 @@ bool FStateTreeExecutionContext::SelectStateInternal(
 			if (NextLinkedStateParameterOverride)
 			{
 				// Create from an override.
-				FStateTreeDataView TempStateParametersView = AddTemporaryInstance(CurrentFrame, FStateTreeIndex16::Invalid, NextState.ParameterDataHandle, FConstStructView(TBaseStructure<FCompactStateTreeParameters>::Get()));
+				FStateTreeDataView TempStateParametersView = AddTemporaryInstance(NextFrame, FStateTreeIndex16::Invalid, NextState.ParameterDataHandle, FConstStructView(TBaseStructure<FCompactStateTreeParameters>::Get()));
 				check(TempStateParametersView.IsValid());
 				FCompactStateTreeParameters& StateParams = TempStateParametersView.GetMutable<FCompactStateTreeParameters>();
 				StateParams.Parameters = *NextLinkedStateParameterOverride;
@@ -5457,8 +6171,8 @@ bool FStateTreeExecutionContext::SelectStateInternal(
 			else
 			{
 				// Create from template in the asset.
-				const FConstStructView DefaultStateParamsInstanceData = CurrentFrame.StateTree->DefaultInstanceData.GetStruct(NextState.ParameterTemplateIndex.Get());
-				FStateTreeDataView TempStateParametersView = AddTemporaryInstance(CurrentFrame, FStateTreeIndex16::Invalid, NextState.ParameterDataHandle, DefaultStateParamsInstanceData);
+				const FConstStructView DefaultStateParamsInstanceData = NextFrame.StateTree->DefaultInstanceData.GetStruct(NextState.ParameterTemplateIndex.Get());
+				FStateTreeDataView TempStateParametersView = AddTemporaryInstance(NextFrame, FStateTreeIndex16::Invalid, NextState.ParameterDataHandle, DefaultStateParamsInstanceData);
 				check(TempStateParametersView.IsValid());
 				FCompactStateTreeParameters& StateParams = TempStateParametersView.GetMutable<FCompactStateTreeParameters>();
 				NextStateParametersView = FStateTreeDataView(StateParams.Parameters.GetMutableValue());
@@ -5472,23 +6186,26 @@ bool FStateTreeExecutionContext::SelectStateInternal(
 		{
 			// Note: the parameters are for the current (linked) state, stored in current frame.
 			// The copy can fail, if the overridden parameters do not match, this is by design.
-			CopyBatchWithValidation(CurrentParentFrame, CurrentFrame, NextStateParametersView, NextState.ParameterBindingsBatch);
+			CopyBatchWithValidation(NextParentFrame, NextFrame, NextStateParametersView, NextState.ParameterBindingsBatch);
 		}
 	}
 
-	const bool bIsDestinationState = PathToNextState.Num() < 2;
-	const bool bShouldPrerequisitesBeChecked = bIsDestinationState || NextState.bCheckPrerequisitesWhenActivatingChildDirectly;
-	TArray<const FStateTreeSharedEvent*, TInlineAllocator<FStateTreeEventQueue::MaxActiveEvents>> StateSelectionEvents;
+	const bool bIsTargetState = InternalArgs.MissingStatesToReachTarget.Num() <= 1;
+	const bool bShouldPrerequisitesBeChecked = Args.Behavior != ESelectStateBehavior::Forced
+		&& (bIsTargetState || NextState.bCheckPrerequisitesWhenActivatingChildDirectly);
+
+	// Check if the events are accepted.
+	TArray<FStateTreeSharedEvent, TInlineAllocator<FStateTreeEventQueue::MaxActiveEvents>> StateSelectionEvents;
 	if (NextState.EventDataIndex.IsValid())
 	{
-		check(NextState.RequiredEventToEnter.IsValid());
-
+		if (ensure(NextState.RequiredEventToEnter.IsValid()))
+		{
 		// Use the same event as performed transition unless it didn't lead to this state as only state selected by the transition should get it's event.
-		if (TransitionEvent && TransitionEvent->IsValid() && bIsDestinationState)
+			if (Args.TransitionEvent.IsValid() && bIsTargetState)
 		{
-			if (NextState.RequiredEventToEnter.DoesEventMatchDesc(*TransitionEvent->Get()))
+				if (NextState.RequiredEventToEnter.DoesEventMatchDesc(*Args.TransitionEvent))
 			{
-				StateSelectionEvents.Emplace(TransitionEvent);
+					StateSelectionEvents.Emplace(Args.TransitionEvent);
 			}
 		}
 		else
@@ -5499,7 +6216,7 @@ bool FStateTreeExecutionContext::SelectStateInternal(
 				check(Event.IsValid());
 				if (NextState.RequiredEventToEnter.DoesEventMatchDesc(*Event))
 				{
-					StateSelectionEvents.Emplace(&Event);
+						StateSelectionEvents.Emplace(Event);
 				}
 			}
 
@@ -5509,6 +6226,7 @@ bool FStateTreeExecutionContext::SelectStateInternal(
 				StateSelectionEvents.Emplace();
 			}
 		}
+		}
 
 		if (StateSelectionEvents.IsEmpty())
 		{
@@ -5517,392 +6235,665 @@ bool FStateTreeExecutionContext::SelectStateInternal(
 	}
 	else
 	{
-		StateSelectionEvents.Emplace();
+		StateSelectionEvents.Add(FStateTreeSharedEvent());
+	}
+
+	bool bShouldCreateNewState = true;
+	const bool bIsNextTargetStateInActiveStates = !bNewFrameCreated
+		&& InternalArgs.MissingActiveStates.Num() > 0
+		&& InternalArgs.MissingActiveStates[0].GetStateHandle() == NextStateHandle
+		&& InternalArgs.MissingActiveStates[0].GetFrameID() == NextFrame.FrameID;
+	if (EnumHasAllFlags(NextFrame.StateTree->StateSelectionRules, EStateTreeStateSelectionRules::CompletedTransitionStatesCreateNewStates))
+	{
+		// Request a new state if it doesn't match the previous states or if the state is completed.
+		if (bIsNextTargetStateInActiveStates)
+		{
+			bShouldCreateNewState = NextFrame.ActiveTasksStatus.GetStatus(NextState).IsCompleted();
+		}
+	}
+	else
+	{
+		bShouldCreateNewState = !bIsNextTargetStateInActiveStates;
 	}
 
-	// Activate/Push a new state
-	if (!CurrentFrame.ActiveStates.Push(NextStateHandle, UE::StateTree::FActiveStateID(Storage.GenerateUniqueId())))
+	bool bSucceededToSelectState = false;
+	ON_SCOPE_EXIT{
+		if (bSucceededToSelectState == false)
+	{
+			OutSelectionResult.SelectedStates.Pop();
+	}
+	};
+
+	// Add the state to the selected list. It will be removed if the selection fails.
+	if (bShouldCreateNewState)
+	{
+		OutSelectionResult.SelectedStates.Emplace(NextFrame.FrameID, FActiveStateID(Storage.GenerateUniqueId()), NextStateHandle);
+
+		int32 StateInFrameCounter = 0;
+		for (int32 Index = OutSelectionResult.SelectedStates.Num() - 1; Index >= 0; --Index)
+		{
+			if (OutSelectionResult.SelectedStates[Index].GetFrameID() != NextFrame.FrameID)
+			{
+				break;
+			}
+			++StateInFrameCounter;
+	}
+
+		if (StateInFrameCounter > FStateTreeActiveStates::MaxStates)
 	{
 		STATETREE_LOG(Error, TEXT("%hs: Reached max execution depth when trying to select state %s from '%s'.  '%s' using StateTree '%s'."),
-			__FUNCTION__, *GetSafeStateName(CurrentFrame, NextStateHandle), *GetStateStatusString(Exec), *GetNameSafe(&Owner), *GetFullNameSafe(CurrentFrame.StateTree));
+				__FUNCTION__,
+				*GetSafeStateName(NextFrame, NextStateHandle),
+				*GetFullNameSafe(NextFrame.StateTree),
+				*GetNameSafe(&Owner),
+				*GetFullNameSafe(&RootStateTree)
+			);
 		return false;
 	}
-	CurrentFrame.ActiveTasksStatus.Push(NextState);
+	}
+	else
+	{
+		check(InternalArgs.MissingActiveStates.Num() > 0); // bIsNextTargetStateInActiveStates is false
+		OutSelectionResult.SelectedStates.Add(InternalArgs.MissingActiveStates[0]);
+	}
 
-	// Check if we're still tracking on the current active frame and state.
-	// If we are, update the NumCurrentlyActiveStates to indicate that this state's instance data can be accessed. 
-	const uint8 PrevNumCurrentlyActiveStates = CurrentFrame.NumCurrentlyActiveStates; 
-	if (CurrentFrame.ActiveInstanceIndexBase.IsValid()
-		&& CurrentFrameInActiveFrames)
+	// Set the target, since this is recursive (can be called for linked state), it only set once the original target is found.
+	if (bIsTargetState && !OutSelectionResult.TargetState.IsValid())
 	{
-		const int32 CurrentStateIndex = CurrentFrame.ActiveStates.Num() - 1;
-		const FStateTreeStateHandle MatchingActiveHandle = CurrentFrameInActiveFrames->ActiveStates.GetStateSafe(CurrentStateIndex);
-		if (MatchingActiveHandle == NextStateHandle)
-		{
-			CurrentFrame.NumCurrentlyActiveStates = static_cast<uint8>(CurrentFrame.ActiveStates.Num());
-		}
+		OutSelectionResult.TargetState = OutSelectionResult.SelectedStates.Last();
 	}
 
-	bool bSucceededToSelectState = false;
+	for (const FStateTreeSharedEvent& StateSelectionEvent : StateSelectionEvents)
+		{
+		// Add the event for GetDataView and for consuming it later.
+		const bool bContainsSelectionEvent = OutSelectionResult.SelectionEvents.ContainsByPredicate(
+			[&LatestState = OutSelectionResult.SelectedStates.Last()](const FSelectionEventWithID& Other)
+			{
+				return Other.State == LatestState;
+			});
+		ensureMsgf(!bContainsSelectionEvent, TEXT("The event should be remove at the end of the scope."));
 
-	for (const FStateTreeSharedEvent* StateSelectionEvent : StateSelectionEvents)
+		bool bRemoveSelectionEvents = false;
+		if (StateSelectionEvent.IsValid())
+		{
+			OutSelectionResult.SelectionEvents.Emplace(OutSelectionResult.SelectedStates.Last(), StateSelectionEvent);
+			bRemoveSelectionEvents = true;
+	}
+
+		auto RemoveStateSelectionEvent = [bRemoveSelectionEvents](FSelectStateResult& SelectionResult)
 	{
-		if (StateSelectionEvent)
+				if (bRemoveSelectionEvents)
 		{
-			CurrentlyProcessedStateSelectionEvents->Events[NextState.Depth] = *StateSelectionEvent;
+					check(SelectionResult.SelectionEvents.Num() > 0);
+					ensureMsgf(SelectionResult.SelectionEvents.Last().State == SelectionResult.SelectedStates.Last(), TEXT("We should remove from the same element it was added."));
+					SelectionResult.SelectionEvents.Pop();
 		}
+			};
 		
 		if (bShouldPrerequisitesBeChecked)
 		{
 			// Check that the state can be entered
 			UE_STATETREE_DEBUG_ENTER_PHASE(this, EStateTreeUpdatePhase::EnterConditions);
-			const bool bEnterConditionsPassed = TestAllConditionsWithValidation(CurrentParentFrame, CurrentFrame, NextStateHandle, NextState.EnterConditionsBegin, NextState.EnterConditionsNum);
+			const bool bEnterConditionsPassed = TestAllConditionsWithValidation(NextParentFrame, NextFrame, NextStateHandle, NextState.EnterConditionsBegin, NextState.EnterConditionsNum);
 			UE_STATETREE_DEBUG_EXIT_PHASE(this, EStateTreeUpdatePhase::EnterConditions);
 
 			if (!bEnterConditionsPassed)
 			{
+				RemoveStateSelectionEvent(OutSelectionResult);
 				continue;
 			}
 		}
 		
-		if (!bIsDestinationState)
+		FSelectStateInternalArguments NewInternalArgs = FSelectStateInternalArguments{
+			.MissingActiveStates = bShouldCreateNewState ? TArrayView<const FActiveState>() : InternalArgs.MissingActiveStates.Mid(1),
+			.MissingSourceFrameID = FActiveFrameID(),
+			.MissingSourceStates = TArrayView<const FActiveState>(),
+			.MissingStatesToReachTarget = bIsTargetState ? TArrayView<const FStateHandleContext>() : InternalArgs.MissingStatesToReachTarget.Mid(1)
+		};
+		if (NextState.Type == EStateTreeStateType::Linked)
+		{
+			// MissingStatesToReachTarget can include a linked state and the frame needs to be constructed (if needed).
+			bSucceededToSelectState = SelectStateInternal_Linked(Args, NewInternalArgs, OutSelectionResult, NextFrame.StateTree, NextState, bShouldCreateNewState);
+		}
+		else if (NextState.Type == EStateTreeStateType::LinkedAsset)
+		{
+			// MissingStatesToReachTarget can include a linked asset state and the frame needs to be constructed (if needed).
+			bSucceededToSelectState = SelectStateInternal_LinkedAsset(Args, NewInternalArgs, OutSelectionResult, NextFrame.StateTree, NextState, NextLinkedStateAsset, bShouldCreateNewState);
+		}
+		else if (!bIsTargetState)
 		{
-			// Next child state is already known. Passing TransitionEvent further so state selected directly by transition can use it.
-			if (SelectStateInternal(CurrentParentFrame, CurrentFrame, CurrentFrameInActiveFrames, PathToNextState.Mid(1), OutSelectionResult, TransitionEvent))
+			check(NewInternalArgs.MissingStatesToReachTarget.Num() > 0);
+			// Next child state is already known. Passing TransitionEvent further, states selected directly by transition can use it.
+			bSucceededToSelectState = SelectStateInternal(Args, NewInternalArgs, OutSelectionResult);
+		}
+		else if (Args.Behavior != ESelectStateBehavior::Forced)
+		{
+			switch (NextState.SelectionBehavior)
 			{
+			case EStateTreeStateSelectionBehavior::TryEnterState:
+				UE_STATETREE_DEBUG_STATE_EVENT(this, NextStateHandle, EStateTreeTraceEventType::OnStateSelected);
 				bSucceededToSelectState = true;
 				break;
+			case EStateTreeStateSelectionBehavior::TrySelectChildrenInOrder:
+				bSucceededToSelectState = SelectStateInternal_TrySelectChildrenInOrder(Args, NewInternalArgs, OutSelectionResult, NextFrame.StateTree, NextState, NextStateHandle);
+				break;
+			case EStateTreeStateSelectionBehavior::TrySelectChildrenAtRandom:
+				bSucceededToSelectState = SelectStateInternal_TrySelectChildrenAtRandom(Args, NewInternalArgs, OutSelectionResult, NextFrame.StateTree, NextState, NextStateHandle);
+				break;
+			case EStateTreeStateSelectionBehavior::TrySelectChildrenWithHighestUtility:
+				bSucceededToSelectState = SelectStateInternal_TrySelectChildrenWithHighestUtility(Args, NewInternalArgs, OutSelectionResult, NextFrame.StateTree, NextState, NextStateHandle);
+				break;
+			case EStateTreeStateSelectionBehavior::TrySelectChildrenAtRandomWeightedByUtility:
+				bSucceededToSelectState = SelectStateInternal_TrySelectChildrenAtRandomWeightedByUtility(Args, NewInternalArgs, OutSelectionResult, NextFrame.StateTree, NextState, NextStateHandle);
+				break;
+			case EStateTreeStateSelectionBehavior::TryFollowTransitions:
+				bSucceededToSelectState = SelectStateInternal_TryFollowTransitions(Args, NewInternalArgs, OutSelectionResult, NextFrame.StateTree, NextState, NextStateHandle);
+				break;
 			}
 		}
-		else if (NextState.Type == EStateTreeStateType::Linked)
+		else
 		{
-			if (NextState.LinkedState.IsValid())
+			if (NewInternalArgs.MissingStatesToReachTarget.Num() == 0)
 			{
-				if (OutSelectionResult.IsFull())
+				bSucceededToSelectState = true;
+			}
+			else
 				{
-					STATETREE_LOG(Error, TEXT("%hs: Reached max execution depth when trying to select state %s from '%s'.  '%s' using StateTree '%s'."),
-						__FUNCTION__, *GetSafeStateName(CurrentFrame, NextStateHandle), *GetStateStatusString(Exec), *GetNameSafe(&Owner), *GetFullNameSafe(CurrentFrame.StateTree));
-					break;
+				// Continue the force selection. Next child state is already known.
+				bSucceededToSelectState = SelectStateInternal(Args, NewInternalArgs, OutSelectionResult);
+			}
 				}
 
-				FStateTreeExecutionFrame NewFrame;
-				NewFrame.StateTree = CurrentFrame.StateTree;
-				NewFrame.RootState = NextState.LinkedState;
-				NewFrame.ExternalDataBaseIndex = CurrentFrame.ExternalDataBaseIndex;
 
-				// Check and prevent recursion.
-				const bool bNewFrameAlreadySelected = OutSelectionResult.GetSelectedFrames().ContainsByPredicate([&NewFrame](const FStateTreeExecutionFrame& Frame) {
-					return Frame.IsSameFrame(NewFrame);
-				});
-				
-				if (bNewFrameAlreadySelected)
+		if (bSucceededToSelectState)
 				{
-					STATETREE_LOG(Error, TEXT("%hs: Trying to recursively enter subtree '%s' from '%s'.  '%s' using StateTree '%s'."),
-						__FUNCTION__, *GetSafeStateName(NewFrame, NewFrame.RootState), *GetStateStatusString(Exec), *GetNameSafe(&Owner), *GetFullNameSafe(CurrentFrame.StateTree));
 					break;
 				}
-
-				// If the Frame already exists, copy instance indices so that conditions that rely on active states work correctly.
-				const FStateTreeExecutionFrame* ExistingFrame = Exec.ActiveFrames.FindByPredicate(
-					[StateTree = NewFrame.StateTree, RootState = NewFrame.RootState](const FStateTreeExecutionFrame& Frame)
+		else
 					{
-						return Frame.StateTree == StateTree && Frame.RootState == RootState;
-					});
-				if (ExistingFrame)
-				{
-					NewFrame.FrameID = ExistingFrame->FrameID;
-					NewFrame.ActiveTasksStatus = ExistingFrame->ActiveTasksStatus;
-					NewFrame.ActiveTaskNodeIndex = ExistingFrame->ActiveTaskNodeIndex;
-					NewFrame.ActiveInstanceIndexBase = ExistingFrame->ActiveInstanceIndexBase;
-					NewFrame.GlobalInstanceIndexBase = ExistingFrame->GlobalInstanceIndexBase;
-					NewFrame.StateParameterDataHandle = ExistingFrame->StateParameterDataHandle;
-					NewFrame.GlobalParameterDataHandle = ExistingFrame->GlobalParameterDataHandle;
+			RemoveStateSelectionEvent(OutSelectionResult);
+		}
 				}
-				else
-				{
-					NewFrame.FrameID = UE::StateTree::FActiveFrameID(Storage.GenerateUniqueId());
-
-					const FCompactStateTreeFrame* FrameInfo = NewFrame.StateTree->GetFrameFromHandle(NewFrame.RootState);
-					ensureAlwaysMsgf(FrameInfo, TEXT("The compiled data is invalid. It should contains the information for the new root frame."));
-					NewFrame.ActiveTasksStatus = FrameInfo ? FStateTreeTasksCompletionStatus(*FrameInfo) : FStateTreeTasksCompletionStatus();
 
-					// Since the StateTree is the same, we can access the global tasks of CurrentFrame, if they are initialized.
-					NewFrame.GlobalParameterDataHandle = CurrentFrame.GlobalParameterDataHandle;
-					NewFrame.GlobalInstanceIndexBase = CurrentFrame.GlobalInstanceIndexBase;
-					NewFrame.StateParameterDataHandle = NextState.ParameterDataHandle; // Temporary allocated earlier if did not exists.
+	return bSucceededToSelectState;
 				}
 
-				OutSelectionResult.PushFrame(NewFrame);
+TOptional<bool> FStateTreeExecutionContext::SelectStateFromSourceInternal(
+	const FSelectStateArguments& Args,
+	const FSelectStateInternalArguments& InternalArgs,
+	FSelectStateResult& OutSelectionResult,
+	const FStateTreeExecutionFrame& NextFrame,
+	const FCompactStateTreeState& NextState,
+	const FStateTreeStateHandle NextStateHandle,
+	const bool bNewFrameCreated)
+{
+	using namespace UE::StateTree;
+	using namespace UE::StateTree::ExecutionContext;
+	using namespace UE::StateTree::ExecutionContext::Private;
 
-				// If State is linked, proceed to the linked state.
-				if (SelectStateInternal(&CurrentFrame, OutSelectionResult.GetSelectedFrames().Last(), ExistingFrame, {NewFrame.RootState}, OutSelectionResult))
+	if (InternalArgs.MissingSourceStates.Num() > 0
+		&& InternalArgs.MissingStatesToReachTarget.Num() > 1
+		&& !bNewFrameCreated)
 				{
-					bSucceededToSelectState = true;
-					break;
+		check(InternalArgs.MissingActiveStates.Num() > 0);
+		check(InternalArgs.MissingSourceStates[0] == InternalArgs.MissingActiveStates[0]);
+		const FActiveState& MissingActiveState = InternalArgs.MissingSourceStates[0];
+		bool bContinueWithStateSelection = NextStateHandle != MissingActiveState.GetStateHandle();
+		bool bCompletedState = false;
+		if (!bContinueWithStateSelection)
+		{
+			if (InternalArgs.MissingActiveStates.Num() == 1)
+			{
+				// MissingActiveState is the transition source.
+				if (EnumHasAllFlags(NextFrame.StateTree->StateSelectionRules, EStateTreeStateSelectionRules::CompletedTransitionStatesCreateNewStates))
+				{
+					bContinueWithStateSelection = NextFrame.ActiveTasksStatus.GetStatus(NextState).IsCompleted();
 				}
-				
-				OutSelectionResult.PopFrame();
 			}
 			else
 			{
-				STATETREE_LOG(Warning, TEXT("%hs: Trying to enter invalid linked subtree from '%s'.  '%s' using StateTree '%s'."),
-					__FUNCTION__, *GetStateStatusString(Exec), *GetNameSafe(&Owner), *GetFullNameSafe(CurrentFrame.StateTree));
+				if (EnumHasAllFlags(NextFrame.StateTree->StateSelectionRules, EStateTreeStateSelectionRules::CompletedTransitionStatesCreateNewStates))
+				{
+					bCompletedState = NextFrame.ActiveTasksStatus.GetStatus(NextState).IsCompleted();
 			}
 		}
-		else if (NextState.Type == EStateTreeStateType::LinkedAsset)
-		{
-			if (NextLinkedStateAsset == nullptr || NextLinkedStateAsset->States.Num() == 0)
-			{
-				break;
 			}
 
-			if (OutSelectionResult.IsFull())
+		if (!bContinueWithStateSelection)
 			{
-				STATETREE_LOG(Error, TEXT("%hs: Reached max execution depth when trying to select state %s from '%s'.  '%s' using StateTree '%s'."),
-					__FUNCTION__, *GetSafeStateName(CurrentFrame, NextStateHandle), *GetStateStatusString(Exec), *GetNameSafe(&Owner), *GetFullNameSafe(CurrentFrame.StateTree));
-				break;
-			}
+			OutSelectionResult.SelectedStates.Add(InternalArgs.MissingSourceStates[0]);
 
-			// The linked state tree should have compatible context requirements.
-			if (!NextLinkedStateAsset->HasCompatibleContextData(RootStateTree)
-				|| NextLinkedStateAsset->GetSchema()->GetClass() != RootStateTree.GetSchema()->GetClass())
+			FSelectStateInternalArguments NewInternalArgs = FSelectStateInternalArguments{
+				.MissingActiveStates = InternalArgs.MissingActiveStates.Mid(1),
+				.MissingSourceFrameID = FActiveFrameID(),
+				.MissingSourceStates = InternalArgs.MissingSourceStates.Mid(1),
+				.MissingStatesToReachTarget = InternalArgs.MissingStatesToReachTarget.Mid(1)
+			};
+			const bool bSelectedFromSource = SelectStateInternal(Args, NewInternalArgs, OutSelectionResult);
+			if (bSelectedFromSource && bCompletedState)
 			{
-				STATETREE_LOG(Error, TEXT("%hs: The linked State Tree '%s' does not have compatible schema, trying to select state %s from '%s'.  '%s' using StateTree '%s'."),
-					__FUNCTION__, *GetFullNameSafe(NextLinkedStateAsset), *GetSafeStateName(CurrentFrame, NextStateHandle), *GetStateStatusString(Exec), *GetNameSafe(&Owner), *GetFullNameSafe(CurrentFrame.StateTree));
-				break;
+				check(EnumHasAllFlags(NextFrame.StateTree->StateSelectionRules, EStateTreeStateSelectionRules::CompletedStateBeforeTransitionSourceFailsTransition));
+				// Before the source, there's a completed state.
+				//The completed state should fail the transition(because of EStateTreeStateSelectionRules::CompletedStateBeforeTargetFailsTransition).
+				//The transition can contain a state with "follow transition" that would remove this state from the transition.
+				if (OutSelectionResult.SelectedStates.Contains(MissingActiveState))
+			{
+					STATETREE_LOG(VeryVerbose, TEXT("%hs: Selection fails because state '%s' is completed.  '%s' using StateTree '%s'."),
+						__FUNCTION__,
+						*NextState.Name.ToString(),
+						*GetNameSafe(&Owner),
+						*GetFullNameSafe(NextFrame.StateTree)
+					);
+					return false;
+				}
+			}
+			return bSelectedFromSource;
+		}
+	}
+	return {};
 			}
 				
-			FStateTreeExecutionFrame NewFrame;
-			NewFrame.StateTree = NextLinkedStateAsset;
-			NewFrame.RootState = FStateTreeStateHandle::Root;
-			NewFrame.bIsGlobalFrame = true;
-
+namespace UE::StateTree::ExecutionContext::Private
+{
+	bool PreventRecursionCheck(const FExecutionFrameKey& LinkStateFrameKey,
+		TArrayView<const FActiveFrameID> SelectedFrames,
+		TArrayView<const FStateTreeExecutionFrame> ActiveFrames,
+		TArrayView<const FStateTreeExecutionFrame> TemporaryFrames)
+	{
 			// Check and prevent recursion.
-			const bool bNewFrameAlreadySelected = OutSelectionResult.GetSelectedFrames().ContainsByPredicate([&NewFrame](const FStateTreeExecutionFrame& Frame) {
-				return Frame.IsSameFrame(NewFrame);
+		return SelectedFrames.ContainsByPredicate(
+			[&LinkStateFrameKey, &ActiveFrames, &TemporaryFrames](const FActiveFrameID& FrameID)
+			{
+				const FStateTreeExecutionFrame* Frame = FindExecutionFrame(FrameID, ActiveFrames, TemporaryFrames);
+				return Frame && Frame->HasRoot(LinkStateFrameKey);
 			});
+	}
 				
-			if (bNewFrameAlreadySelected)
-			{
-				STATETREE_LOG(Error, TEXT("%hs: Trying to recursively enter subtree '%s' from '%s'.  '%s' using StateTree '%s'."),
-					__FUNCTION__, *GetSafeStateName(NewFrame, NewFrame.RootState), *GetStateStatusString(Exec), *GetNameSafe(&Owner), *GetFullNameSafe(CurrentFrame.StateTree));
-				break;
-			}
-
-			// If the Frame already exists, copy instance indices so that conditions that rely on active states work correctly.
-			const FStateTreeExecutionFrame* ExistingFrame = Exec.ActiveFrames.FindByPredicate(
-				[StateTree = NewFrame.StateTree, RootState = NewFrame.RootState](const FStateTreeExecutionFrame& Frame)
+	FStateTreeExecutionFrame* SelectedFrameLinkedFrame(FStateTreeExecutionState& Exec,
+		bool bShouldCreateNewState,
+		TArrayView<const UE::StateTree::FActiveState> MatchingActiveStates,
+		const FExecutionFrameKey& LinkStateFrameKey
+	)
+	{
+		FStateTreeExecutionFrame* Result = nullptr;
+		if (!bShouldCreateNewState)
+			{
+			// Get the next frame ID
+			if (ensure(MatchingActiveStates.Num() >= 1))
+			{
+				const FActiveFrameID ExistingFrameID = MatchingActiveStates[0].GetFrameID();
+				Result = Exec.FindActiveFrame(ExistingFrameID);
+				check(Result);
+				if (Result->HasRoot(LinkStateFrameKey))
 				{
-					return Frame.StateTree == StateTree && Frame.RootState == RootState;
-				});
-
-			bool bStartedTemporaryEvaluatorsAndGlobalTasks = false;
-			if (ExistingFrame)
+					if (EnumHasAllFlags(Result->StateTree->GetStateSelectionRules(), EStateTreeStateSelectionRules::CompletedTransitionStatesCreateNewStates)
+						&& Result->ActiveTasksStatus.GetStatus(Result->StateTree).IsCompleted()
+						)
 			{
-				NewFrame.FrameID = ExistingFrame->FrameID;
-				NewFrame.ActiveTasksStatus = ExistingFrame->ActiveTasksStatus;
-				NewFrame.ActiveTaskNodeIndex = ExistingFrame->ActiveTaskNodeIndex;
-				NewFrame.ActiveInstanceIndexBase = ExistingFrame->ActiveInstanceIndexBase;
-				NewFrame.GlobalInstanceIndexBase = ExistingFrame->GlobalInstanceIndexBase;
-				NewFrame.StateParameterDataHandle = ExistingFrame->StateParameterDataHandle;
-				NewFrame.GlobalParameterDataHandle = ExistingFrame->GlobalParameterDataHandle;
-				NewFrame.ExternalDataBaseIndex = ExistingFrame->ExternalDataBaseIndex;
+						return nullptr;
+					}
+				}
 			}
-			else
+			}
+		return Result;
+	}
+}
+
+bool FStateTreeExecutionContext::SelectStateInternal_Linked(
+	const FSelectStateArguments& Args,
+	const FSelectStateInternalArguments& InternalArgs,
+	FSelectStateResult& OutSelectionResult,
+	TNotNull<const UStateTree*> NextStateTree,
+	const FCompactStateTreeState& TargetState,
+	bool bShouldCreateNewState)
 			{
-				NewFrame.FrameID = UE::StateTree::FActiveFrameID(Storage.GenerateUniqueId());
-				const FCompactStateTreeFrame* FrameInfo = NewFrame.StateTree->GetFrameFromHandle(NewFrame.RootState);
-				ensureMsgf(FrameInfo, TEXT("The compiled data is invalid. It should contains the information for the root frame."));
-				NewFrame.ActiveTasksStatus = FrameInfo ? FStateTreeTasksCompletionStatus(*FrameInfo) : FStateTreeTasksCompletionStatus();
+	using namespace UE::StateTree;
+	using namespace UE::StateTree::ExecutionContext;
+	using namespace UE::StateTree::ExecutionContext::Private;
 
-				// Pass the linked state's parameters as global parameters to the linked asset.
-				NewFrame.GlobalParameterDataHandle = NextState.ParameterDataHandle;
+	FStateTreeExecutionState& Exec = GetExecState();
 
-				// Collect external data if needed
-				NewFrame.ExternalDataBaseIndex = CollectExternalData(NewFrame.StateTree);
-				if (!NewFrame.ExternalDataBaseIndex.IsValid())
+	if (!TargetState.LinkedState.IsValid())
 				{
-					STATETREE_LOG(VeryVerbose, TEXT("%hs: Cannot select state '%s' because failed to collect external data for nested tree '%s'.  '%s' using StateTree '%s'."),
-						__FUNCTION__, *GetSafeStateName(CurrentFrame, NextStateHandle), *GetFullNameSafe(NewFrame.StateTree), *GetNameSafe(&Owner), *GetFullNameSafe(CurrentFrame.StateTree));
-					break;
+		STATETREE_LOG(Warning, TEXT("%hs: Trying to enter invalid linked subtree from '%s'. '%s' using StateTree '%s'."),
+			__FUNCTION__,
+			*GetStateStatusString(Exec),
+			*Owner.GetName(),
+			*NextStateTree->GetFullName()
+		);
+		return false;
 				}
 					
-				// The state parameters will be from the root state.
-				const FCompactStateTreeState& RootState = NewFrame.StateTree->States[NewFrame.RootState.Index];
-				NewFrame.StateParameterDataHandle = RootState.ParameterDataHandle;
+	const FExecutionFrameKey LinkStateFrameKey = FExecutionFrameKey(NextStateTree, TargetState.LinkedState);
 
-				// Start global tasks and evaluators temporarily, so that their data is available already during select.
-				const EStateTreeRunStatus StartResult = StartTemporaryEvaluatorsAndGlobalTasks(&CurrentFrame, NewFrame);
-				if (StartResult != EStateTreeRunStatus::Running)
+	const bool bHasMissingState = InternalArgs.MissingStatesToReachTarget.Num() > 0 && Args.Behavior == ESelectStateBehavior::Forced;
+	if (bHasMissingState)
+	{
+		// In a force transition, the root could be different from what is expected.
+		//Ex: a previous transition go to root, then another transition go to another top level state (new root)
+		if (InternalArgs.MissingStatesToReachTarget[0].StateTree != LinkStateFrameKey.StateTree
+			|| InternalArgs.MissingStatesToReachTarget[0].StateHandle != LinkStateFrameKey.RootState)
 				{
-					STATETREE_LOG(VeryVerbose, TEXT("%hs: Cannot select state '%s' because cannot start nested tree's '%s' global tasks and evaluators.  '%s' using StateTree '%s'."),
-						__FUNCTION__, *GetSafeStateName(CurrentFrame, NextStateHandle), *GetFullNameSafe(NewFrame.StateTree), *GetNameSafe(&Owner), *GetFullNameSafe(CurrentFrame.StateTree));
-					
-					StopTemporaryEvaluatorsAndGlobalTasks(nullptr, NewFrame, StartResult);
-					GetExecState().DelegateActiveListeners.RemoveAll(NewFrame.FrameID);
-
-					break;
+			STATETREE_LOG(Error, TEXT("%hs: The missing state is not from the same state tree. '%s' using StateTree '%s'."),
+				__FUNCTION__,
+				*Owner.GetName(),
+				*NextStateTree->GetFullName()
+			);
+			return false;
 				}
-
-				bStartedTemporaryEvaluatorsAndGlobalTasks = true;
 			}
 				
-			OutSelectionResult.PushFrame(NewFrame);
-
-			// If State is linked, proceed to the linked state.
-			if (SelectStateInternal(&CurrentFrame, OutSelectionResult.GetSelectedFrames().Last(), ExistingFrame, {NewFrame.RootState}, OutSelectionResult))
-			{
-				bSucceededToSelectState = true;
-				break;
+	if (PreventRecursionCheck(LinkStateFrameKey, OutSelectionResult.SelectedFrames, MakeConstArrayView(Exec.ActiveFrames), MakeConstArrayView(OutSelectionResult.TemporaryFrames)))
+			{
+		STATETREE_LOG(Error, TEXT("%hs: Trying to recursively enter subtree '%s' from '%s'. '%s' using StateTree '%s'."),
+			__FUNCTION__,
+			*GetSafeStateName(LinkStateFrameKey.StateTree, LinkStateFrameKey.RootState),
+			*GetStateStatusString(Exec),
+			*Owner.GetName(),
+			*NextStateTree->GetFullName()
+		);
+		return false;
 			}
 
-			if (bStartedTemporaryEvaluatorsAndGlobalTasks)
+	const FCompactStateTreeFrame* LinkStateTreeFrame = FindStateTreeFrame(LinkStateFrameKey);
+	if (LinkStateTreeFrame == nullptr)
 			{
-				StopTemporaryEvaluatorsAndGlobalTasks(&CurrentFrame, NewFrame, EStateTreeRunStatus::Running);
-				GetExecState().DelegateActiveListeners.RemoveAll(NewFrame.FrameID);
+		STATETREE_LOG(Error, TEXT("%hs: The frame '%s' from '%s' does not exist. '%s' using StateTree '%s'."),
+			__FUNCTION__,
+			*GetSafeStateName(LinkStateFrameKey.StateTree, LinkStateFrameKey.RootState),
+			*GetStateStatusString(Exec),
+			*Owner.GetName(),
+			*NextStateTree->GetFullName()
+		);
+		return false;
 			}
 				
-			OutSelectionResult.PopFrame();
-		}
-		else if (NextState.SelectionBehavior == EStateTreeStateSelectionBehavior::TryEnterState)
+	// Do we have an existing frame.
+	FStateTreeExecutionFrame* SelectedFrame = SelectedFrameLinkedFrame(Exec, bShouldCreateNewState, InternalArgs.MissingActiveStates, LinkStateFrameKey);
+	if (SelectedFrame && !SelectedFrame->HasRoot(LinkStateFrameKey))
+	{
+		if (Args.Behavior == ESelectStateBehavior::Forced)
 		{
-			// Select this state.
-			UE_STATETREE_DEBUG_STATE_EVENT(this, NextStateHandle, EStateTreeTraceEventType::OnStateSelected);
-			bSucceededToSelectState = true;
-			break;
+			SelectedFrame = nullptr;
 		}
-		else if (NextState.SelectionBehavior == EStateTreeStateSelectionBehavior::TryFollowTransitions)
+		else
 		{
-			UE_STATETREE_DEBUG_SCOPED_STATE_PHASE(this, NextStateHandle, EStateTreeUpdatePhase::TrySelectBehavior);
-
-			EStateTreeTransitionPriority CurrentPriority = EStateTreeTransitionPriority::None;
+			STATETREE_LOG(Error, TEXT("%hs: The frame '%s' from '%s' does not have the same root as the active frame. '%s' using StateTree '%s'."),
+				__FUNCTION__,
+				*GetSafeStateName(LinkStateFrameKey.StateTree, LinkStateFrameKey.RootState),
+				*GetStateStatusString(Exec),
+				*Owner.GetName(),
+				*NextStateTree->GetFullName()
+			);
+			return false;
+		}
+		}
 
-			for (uint8 i = 0; i < NextState.TransitionsNum; i++)
+	const bool bIsNewFrame = SelectedFrame == nullptr;
+	if (bIsNewFrame)
 			{
-				const int16 TransitionIndex = NextState.TransitionsBegin + i;
-				const FCompactStateTransition& Transition = CurrentStateTree->Transitions[TransitionIndex];
-
-				// Skip disabled transitions
-				if (Transition.bTransitionEnabled == false)
+		// Note. Adding to TemporaryFrame can invalidate TargetFrame.
+		FStateTreeIndex16 ExternalDataBaseIndex;
+		FStateTreeDataHandle GlobalParameterDataHandle;
+		FStateTreeIndex16 GlobalInstanceIndexBase;
 				{
-					continue;
+			const FStateTreeExecutionFrame* NextFrame = FindExecutionFrame(OutSelectionResult.SelectedFrames.Last(), MakeConstArrayView(Exec.ActiveFrames), MakeConstArrayView(OutSelectionResult.TemporaryFrames));
+			check(NextFrame);
+			ExternalDataBaseIndex = NextFrame->ExternalDataBaseIndex;
+			GlobalParameterDataHandle = NextFrame->GlobalParameterDataHandle;
+			GlobalInstanceIndexBase = NextFrame->GlobalInstanceIndexBase;
 				}
 
-				// No need to test the transition if same or higher priority transition has already been processed.
-				if (Transition.Priority <= CurrentPriority)
-				{
-					continue;
-				}
+		constexpr bool bIsGlobalFrame = false;
+		FStateTreeExecutionFrame& NewFrame = OutSelectionResult.MakeAndAddTemporaryFrame(FActiveFrameID(Storage.GenerateUniqueId()), LinkStateFrameKey, bIsGlobalFrame);
+		NewFrame.ExternalDataBaseIndex = ExternalDataBaseIndex;
+		NewFrame.GlobalParameterDataHandle = GlobalParameterDataHandle;
+		NewFrame.GlobalInstanceIndexBase = GlobalInstanceIndexBase;
+		NewFrame.StateParameterDataHandle = TargetState.ParameterDataHandle; // Temporary allocated earlier if did not exists.
 
-				// Skip completion transitions
-				if (EnumHasAnyFlags(Transition.Trigger, EStateTreeTransitionTrigger::OnStateCompleted))
-				{
-					continue;
+		SelectedFrame = &NewFrame;
 				}
 
-				// Cannot follow transitions with delay.
-				if (Transition.HasDelay())
+	OutSelectionResult.SelectedFrames.Add(SelectedFrame->FrameID);
+
+	// Select the root state of the new frame.
+	FStateHandleContext RootState = FStateHandleContext(LinkStateFrameKey.StateTree, LinkStateFrameKey.RootState);
+	FSelectStateInternalArguments NewInternalArgs = FSelectStateInternalArguments{
+		.MissingActiveStates = bIsNewFrame ? TArrayView<const FActiveState>() : InternalArgs.MissingActiveStates,
+		.MissingSourceFrameID = SelectedFrame->FrameID,
+		.MissingSourceStates = TArrayView<const FActiveState>(),
+		.MissingStatesToReachTarget = bHasMissingState ? InternalArgs.MissingStatesToReachTarget : MakeArrayView(&RootState, 1)
+	};
+	const bool bSucceededToSelectState = SelectStateInternal(Args, NewInternalArgs, OutSelectionResult);
+
+	if (!bSucceededToSelectState)
 				{
-					continue;
+		OutSelectionResult.SelectedFrames.Pop();
 				}
 
-				// Try to prevent (infinite) loops in the selection.
-				if (CurrentFrame.ActiveStates.Contains(Transition.State))
-				{
-					STATETREE_LOG(Error, TEXT("%hs: Loop detected when trying to select state %s from '%s'. Prior states: %s.  '%s' using StateTree '%s'.")
-						, __FUNCTION__
-						, *GetSafeStateName(CurrentFrame, NextStateHandle)
-						, *GetStateStatusString(Exec)
-						, *DebugGetStatePath(OutSelectionResult.GetSelectedFrames(), &CurrentFrame)
-						, *GetNameSafe(&Owner)
-						, *GetFullNameSafe(CurrentFrame.StateTree));
-					continue;
+	return bSucceededToSelectState;
 				}
 
-				TArray<const FStateTreeSharedEvent*, TInlineAllocator<FStateTreeEventQueue::MaxActiveEvents>> SelectedStateTransitionEvents;
-				if (Transition.Trigger == EStateTreeTransitionTrigger::OnEvent)
+bool FStateTreeExecutionContext::SelectStateInternal_LinkedAsset(
+	const FSelectStateArguments& Args,
+	const FSelectStateInternalArguments& InternalArgs,
+	FSelectStateResult& OutSelectionResult,
+	TNotNull<const UStateTree*> NextStateTree,
+	const FCompactStateTreeState& NextState,
+	const UStateTree* NextLinkedStateAsset,
+	bool bShouldCreateNewState)
 				{
-					check(Transition.RequiredEvent.IsValid());
+	using namespace UE::StateTree;
+	using namespace UE::StateTree::ExecutionContext;
+	using namespace UE::StateTree::ExecutionContext::Private;
 
-					if (StateSelectionEvent)
+	FStateTreeExecutionState& Exec = GetExecState();
+
+	if (NextLinkedStateAsset == nullptr || NextLinkedStateAsset->States.Num() == 0)
 					{
-						SelectedStateTransitionEvents.Emplace(StateSelectionEvent);
+		STATETREE_LOG(Error, TEXT("%hs: The linked State Tree is invalid. '%s' using StateTree '%s'."),
+			__FUNCTION__,
+			*Owner.GetName(),
+			*NextStateTree->GetFullName()
+		);
+		return false;
 					}
-					else
+
+	FStateTreeStateHandle NextLinkedStateRoot = FStateTreeStateHandle::Root;
+	const bool bHasMissingState = InternalArgs.MissingStatesToReachTarget.Num() > 0 && Args.Behavior == ESelectStateBehavior::Forced;
+	if (bHasMissingState)
 					{
-						TArrayView<FStateTreeSharedEvent> EventsQueue = GetMutableEventsToProcessView();
-						for (FStateTreeSharedEvent& Event : EventsQueue)
+		// In a force transition, the root could be different from what is expected.
+		//Ex: a previous transition go to root, then another transition go to another top level state (new root)
+		if (InternalArgs.MissingStatesToReachTarget[0].StateTree != NextLinkedStateAsset)
 						{
-							check(Event.IsValid());
-							if (Transition.RequiredEvent.DoesEventMatchDesc(*Event))
-							{
-								SelectedStateTransitionEvents.Emplace(&Event);
+			STATETREE_LOG(Error, TEXT("%hs: The missing state is not from the same state tree. '%s' using StateTree '%s'."),
+				__FUNCTION__,
+				*Owner.GetName(),
+				*NextStateTree->GetFullName()
+			);
+			return false;
 							}
+		NextLinkedStateRoot = InternalArgs.MissingStatesToReachTarget[0].StateHandle;
 						}
+
+	const FExecutionFrameKey LinkStateFrameKey = FExecutionFrameKey(NextLinkedStateAsset, NextLinkedStateRoot);
+
+	// The linked state tree should have compatible context requirements.
+	if (!NextLinkedStateAsset->HasCompatibleContextData(RootStateTree)
+		|| NextLinkedStateAsset->GetSchema()->GetClass() != NextStateTree->GetSchema()->GetClass())
+	{
+		STATETREE_LOG(Error, TEXT("%hs: The linked State Tree '%s' does not have compatible schema, trying to select state %s from '%s'. '%s' using StateTree '%s'."),
+			__FUNCTION__,
+			*GetFullNameSafe(NextLinkedStateAsset),
+			*GetSafeStateName(LinkStateFrameKey.StateTree, LinkStateFrameKey.RootState),
+			*GetStateStatusString(Exec),
+			*Owner.GetName(),
+			*NextStateTree->GetFullName()
+		);
+		return false;
 					}
+
+	if (PreventRecursionCheck(LinkStateFrameKey, OutSelectionResult.SelectedFrames, MakeConstArrayView(Exec.ActiveFrames), MakeConstArrayView(OutSelectionResult.TemporaryFrames)))
+	{
+		STATETREE_LOG(Error, TEXT("%hs: Trying to recursively enter subtree '%s' from '%s'. '%s' using StateTree '%s'."),
+			__FUNCTION__,
+			*GetSafeStateName(LinkStateFrameKey.StateTree, LinkStateFrameKey.RootState),
+			*GetStateStatusString(Exec),
+			*Owner.GetName(),
+			*NextStateTree->GetFullName()
+		);
+		return false;
 				}
-				else if (EnumHasAnyFlags(Transition.Trigger, EStateTreeTransitionTrigger::OnTick))
+
+	const FCompactStateTreeFrame* LinkStateTreeFrame = FindStateTreeFrame(LinkStateFrameKey);
+	if (LinkStateTreeFrame == nullptr)
 				{
-					SelectedStateTransitionEvents.Emplace();
+		STATETREE_LOG(Error, TEXT("%hs: The frame '%s' from '%s' does not exist. '%s' using StateTree '%s'."),
+			__FUNCTION__,
+			*GetSafeStateName(LinkStateFrameKey.StateTree, LinkStateFrameKey.RootState),
+			*GetStateStatusString(Exec),
+			*Owner.GetName(),
+			*NextStateTree->GetFullName()
+		);
+		return false;
 				}
-				else if (EnumHasAnyFlags(Transition.Trigger, EStateTreeTransitionTrigger::OnDelegate))
+
+	// Do we have an existing frame.
+	FStateTreeExecutionFrame* SelectedFrame = SelectedFrameLinkedFrame(Exec, bShouldCreateNewState, InternalArgs.MissingActiveStates, LinkStateFrameKey);
+	if (SelectedFrame && !SelectedFrame->HasRoot(LinkStateFrameKey))
 				{
-					if (Storage.IsDelegateBroadcasted(Transition.RequiredDelegateDispatcher))
+		if (Args.Behavior == ESelectStateBehavior::Forced)
 					{
-						SelectedStateTransitionEvents.Emplace();
-					}
+			SelectedFrame = nullptr;
 				}
 				else
 				{
-					ensureMsgf(false, TEXT("Missing a transition trigger type."));
+			STATETREE_LOG(Error, TEXT("%hs: The frame '%s' from '%s' does not have the same root as the active frame. '%s' using StateTree '%s'."),
+				__FUNCTION__,
+				*GetSafeStateName(LinkStateFrameKey.StateTree, LinkStateFrameKey.RootState),
+				*GetStateStatusString(Exec),
+				*Owner.GetName(),
+				*NextStateTree->GetFullName()
+			);
+			return false;
+		}
 				}
 
-				for (const FStateTreeSharedEvent* SelectedStateTransitionEvent : SelectedStateTransitionEvents)
+	const bool bIsNewFrame = SelectedFrame == nullptr;
+	if (bIsNewFrame)
 				{
-					bool bTransitionConditionsPassed = false;
+		// Collect external data if needed
+		const FStateTreeIndex16 ExternalDataBaseIndex = CollectExternalData(LinkStateFrameKey.StateTree);
+		if (!ExternalDataBaseIndex.IsValid())
 					{
-						FCurrentlyProcessedTransitionEventScope TransitionEventScope(*this, SelectedStateTransitionEvent ? SelectedStateTransitionEvent->Get() : nullptr);
+			STATETREE_LOG(VeryVerbose, TEXT("%hs: Cannot select state '%s' because failed to collect external data for nested tree '%s' from '%s'. '%s' using StateTree '%s'."),
+				__FUNCTION__,
+				*GetSafeStateName(LinkStateFrameKey.StateTree, LinkStateFrameKey.RootState),
+				*GetFullNameSafe(NextLinkedStateAsset),
+				*GetStateStatusString(Exec),
+				*Owner.GetName(),
+				*NextStateTree->GetFullName()
+			);
+			return false;
+		}
 
-						UE_STATETREE_DEBUG_TRANSITION_EVENT(this, FStateTreeTransitionSource(CurrentFrame.StateTree, FStateTreeIndex16(TransitionIndex), Transition.State, Transition.Priority), EStateTreeTraceEventType::OnEvaluating);
+		FStateTreeExecutionFrame& NewFrame = OutSelectionResult.MakeAndAddTemporaryFrame(FActiveFrameID(Storage.GenerateUniqueId()), LinkStateFrameKey, true);
+		NewFrame.ExternalDataBaseIndex = ExternalDataBaseIndex;
+		// Pass the linked state's parameters as global parameters to the linked asset.
+		NewFrame.GlobalParameterDataHandle = NextState.ParameterDataHandle;
+		// The state parameters will be from the root state.
+		const FCompactStateTreeState& NewFrameRootState = LinkStateFrameKey.StateTree->States[NewFrame.RootState.Index];
+		NewFrame.StateParameterDataHandle = NewFrameRootState.ParameterDataHandle;
 
-						UE_STATETREE_DEBUG_SCOPED_PHASE(this, EStateTreeUpdatePhase::TransitionConditions);
+		SelectedFrame = &NewFrame;
+	}
 
-						bTransitionConditionsPassed = TestAllConditionsWithValidation(CurrentParentFrame, CurrentFrame, NextStateHandle, Transition.ConditionsBegin, Transition.ConditionsNum);
-					}
+	OutSelectionResult.SelectedFrames.Add(SelectedFrame->FrameID);
 
-					if (bTransitionConditionsPassed)
+	if (bIsNewFrame)
 					{
-						// Using SelectState() instead of SelectStateInternal to treat the transitions the same way as regular transitions,
-						// e.g. it may jump to a completely different branch.
-						FStateSelectionResult StateSelectionResult;
-						if (SelectState(CurrentFrame, Transition.State, StateSelectionResult, SelectedStateTransitionEvent, Transition.Fallback))
+		// Start global tasks and evaluators temporarily, so that their data is available already during select.
+		const FStateTreeExecutionFrame* NextParentFrame = FindExecutionFrame(OutSelectionResult.SelectedFrames.Last(1), MakeConstArrayView(Exec.ActiveFrames), MakeConstArrayView(OutSelectionResult.TemporaryFrames));
+		FStateTreeExecutionFrame* NextFrame = FindExecutionFrame(OutSelectionResult.SelectedFrames.Last(), MakeArrayView(Exec.ActiveFrames), MakeArrayView(OutSelectionResult.TemporaryFrames));
+		check(NextParentFrame);
+		check(SelectedFrame == NextFrame);
+
+		const EStateTreeRunStatus StartResult = StartTemporaryEvaluatorsAndGlobalTasks(NextParentFrame, *NextFrame);
+		if (StartResult != EStateTreeRunStatus::Running)
 						{
-							// Selection succeeded.
-							// Cannot break yet because higher priority transitions may override the selection. 
-							OutSelectionResult = StateSelectionResult;
-							CurrentPriority = Transition.Priority;
-							break;
+			STATETREE_LOG(VeryVerbose, TEXT("%hs: Cannot select state '%s' because cannot start nested tree's '%s' global tasks and evaluators. '%s' using StateTree '%s'."),
+				__FUNCTION__,
+				*GetSafeStateName(LinkStateFrameKey.StateTree, LinkStateFrameKey.RootState),
+				*GetFullNameSafe(NextLinkedStateAsset),
+				*Owner.GetName(),
+				*NextStateTree->GetFullName()
+			);
+
+			StopTemporaryEvaluatorsAndGlobalTasks(NextParentFrame, *NextFrame, StartResult);
+			CleanFrame(Exec, SelectedFrame->FrameID);
+
+			OutSelectionResult.SelectedFrames.Pop();
+			return false;
 						}
 					}
+
+	// Select the root state of the new frame.
+	FStateHandleContext RootState = FStateHandleContext(LinkStateFrameKey.StateTree, LinkStateFrameKey.RootState);
+	FSelectStateInternalArguments NewInternalArgs = FSelectStateInternalArguments{
+		.MissingActiveStates = bIsNewFrame ? TArrayView<const FActiveState>() : InternalArgs.MissingActiveStates,
+		.MissingSourceFrameID = SelectedFrame->FrameID,
+		.MissingSourceStates = TArrayView<const FActiveState>(),
+		.MissingStatesToReachTarget = bHasMissingState ? InternalArgs.MissingStatesToReachTarget : MakeArrayView(&RootState, 1)
+	};
+	const bool bSucceededToSelectState = SelectStateInternal(Args, NewInternalArgs, OutSelectionResult);
+
+	if (!bSucceededToSelectState)
+	{
+		if (bIsNewFrame)
+		{
+			const FStateTreeExecutionFrame* NextParentFrame = FindExecutionFrame(OutSelectionResult.SelectedFrames.Last(1), MakeConstArrayView(Exec.ActiveFrames), MakeConstArrayView(OutSelectionResult.TemporaryFrames));
+			FStateTreeExecutionFrame* NextFrame = FindExecutionFrame(OutSelectionResult.SelectedFrames.Last(), MakeArrayView(Exec.ActiveFrames), MakeArrayView(OutSelectionResult.TemporaryFrames));
+			check(NextParentFrame);
+			check(NextFrame);
+
+			constexpr EStateTreeRunStatus CompletionStatus = EStateTreeRunStatus::Stopped;
+			StopTemporaryEvaluatorsAndGlobalTasks(NextParentFrame, *NextFrame, CompletionStatus);
+			CleanFrame(Exec, SelectedFrame->FrameID);
 				}
+		OutSelectionResult.SelectedFrames.Pop();
 			}
 
-			if (CurrentPriority != EStateTreeTransitionPriority::None)
-			{
-				bSucceededToSelectState = true;
-				break;
+	return bSucceededToSelectState;
 			}
-		}
-		else if (NextState.SelectionBehavior == EStateTreeStateSelectionBehavior::TrySelectChildrenInOrder)
+
+bool FStateTreeExecutionContext::SelectStateInternal_TrySelectChildrenInOrder(
+	const FSelectStateArguments& Args,
+	const FSelectStateInternalArguments& InternalArgs,
+	FSelectStateResult& OutSelectionResult,
+	TNotNull<const UStateTree*> NextStateTree,
+	const FCompactStateTreeState& NextState,
+	const FStateTreeStateHandle NextStateHandle)
 		{
-			if (NextState.HasChildren())
+	using namespace UE::StateTree::ExecutionContext;
+
+	if (!NextState.HasChildren())
 			{
+		// Select this state (For backwards compatibility)
+		UE_STATETREE_DEBUG_STATE_EVENT(this, NextStateHandle, EStateTreeTraceEventType::OnStateSelected);
+		return true;
+	}
+
 				UE_STATETREE_DEBUG_SCOPED_STATE_PHASE(this, NextStateHandle, EStateTreeUpdatePhase::TrySelectBehavior);
 
 				// If the state has children, proceed to select children.
-				for (uint16 ChildState = NextState.ChildrenBegin; ChildState < NextState.ChildrenEnd; ChildState = CurrentStateTree->States[ChildState].GetNextSibling())
+	bool bSucceededToSelectState = false;
+	for (uint16 ChildStateIndex = NextState.ChildrenBegin; ChildStateIndex < NextState.ChildrenEnd; ChildStateIndex = NextStateTree->States[ChildStateIndex].GetNextSibling())
 				{
-					if (SelectStateInternal(CurrentParentFrame, CurrentFrame, CurrentFrameInActiveFrames, {FStateTreeStateHandle(ChildState)}, OutSelectionResult))
+		FStateHandleContext ChildState = FStateHandleContext(NextStateTree, FStateTreeStateHandle(ChildStateIndex));
+		FSelectStateInternalArguments NewInternalArgs = InternalArgs;
+		NewInternalArgs.MissingStatesToReachTarget = MakeArrayView(&ChildState, 1);
+		if (SelectStateInternal(Args, NewInternalArgs, OutSelectionResult))
 					{
 						// Selection succeeded
 						bSucceededToSelectState = true;
@@ -5910,142 +6901,172 @@ bool FStateTreeExecutionContext::SelectStateInternal(
 					}
 				}
 
-				if (bSucceededToSelectState)
+	return bSucceededToSelectState;
+}
+
+bool FStateTreeExecutionContext::SelectStateInternal_TrySelectChildrenAtRandom(
+	const FSelectStateArguments& Args,
+	const FSelectStateInternalArguments& InternalArgs,
+	FSelectStateResult& OutSelectionResult,
+	TNotNull<const UStateTree*> NextStateTree,
+	const FCompactStateTreeState& NextState,
+	const FStateTreeStateHandle NextStateHandle)
 				{
-					break;
-				}
-			}
-			else
+	using namespace UE::StateTree::ExecutionContext;
+
+	if (!NextState.HasChildren())
 			{
 				// Select this state (For backwards compatibility)
 				UE_STATETREE_DEBUG_STATE_EVENT(this, NextStateHandle, EStateTreeTraceEventType::OnStateSelected);
-				bSucceededToSelectState = true;
-				break;
+		return true;
 			}
-		}
-		else if (NextState.SelectionBehavior == EStateTreeStateSelectionBehavior::TrySelectChildrenAtRandom)
-		{
-			if (NextState.HasChildren())
-			{
+
 				UE_STATETREE_DEBUG_SCOPED_STATE_PHASE(this, NextStateHandle, EStateTreeUpdatePhase::TrySelectBehavior);
 				
-				TArray<uint16, TInlineAllocator<8>> NextLevelChildStates;
-				for (uint16 ChildState = NextState.ChildrenBegin; ChildState < NextState.ChildrenEnd; ChildState = CurrentStateTree->States[ChildState].GetNextSibling())
+	TArray<uint16, TInlineAllocator<16>> NextLevelChildStates;
+	for (uint16 ChildStateIndex = NextState.ChildrenBegin; ChildStateIndex < NextState.ChildrenEnd; ChildStateIndex = NextStateTree->States[ChildStateIndex].GetNextSibling())
 				{
-					NextLevelChildStates.Push(ChildState);
+		NextLevelChildStates.Push(ChildStateIndex);
 				}
 
+	const FStateTreeExecutionState& Exec = GetExecState();
 				while (!NextLevelChildStates.IsEmpty())
 				{
 					const int32 ChildStateIndex = Exec.RandomStream.RandRange(0, NextLevelChildStates.Num() - 1);
-					if (SelectStateInternal(CurrentParentFrame, CurrentFrame, CurrentFrameInActiveFrames,
-					                        {FStateTreeStateHandle(NextLevelChildStates[ChildStateIndex])}, OutSelectionResult))
+		FStateHandleContext ChildState = FStateHandleContext(NextStateTree, FStateTreeStateHandle(NextLevelChildStates[ChildStateIndex]));
+		FSelectStateInternalArguments NewInternalArgs = InternalArgs;
+		NewInternalArgs.MissingStatesToReachTarget = MakeArrayView(&ChildState, 1);
+		if (SelectStateInternal(Args, NewInternalArgs, OutSelectionResult))
 					{
 						// Selection succeeded
-						bSucceededToSelectState = true;
-						break;
+			return true;
 					}
 
 					constexpr EAllowShrinking AllowShrinking = EAllowShrinking::No;
 					NextLevelChildStates.RemoveAtSwap(ChildStateIndex, AllowShrinking);
 				}
 
-				if (bSucceededToSelectState)
+	return false;
+}
+
+bool FStateTreeExecutionContext::SelectStateInternal_TrySelectChildrenWithHighestUtility(
+	const FSelectStateArguments& Args,
+	const FSelectStateInternalArguments& InternalArgs,
+	FSelectStateResult& OutSelectionResult,
+	TNotNull<const UStateTree*> NextStateTree,
+	const FCompactStateTreeState& NextState,
+	const FStateTreeStateHandle NextStateHandle)
 				{
-					break;
-				}
-			}
-			else
+	using namespace UE::StateTree::ExecutionContext;
+
+	if (!NextState.HasChildren())
 			{
 				// Select this state (For backwards compatibility)
 				UE_STATETREE_DEBUG_STATE_EVENT(this, NextStateHandle, EStateTreeTraceEventType::OnStateSelected);
-				bSucceededToSelectState = true;
-				break;
+		return true;
 			}
-		}
-		else if (NextState.SelectionBehavior == EStateTreeStateSelectionBehavior::TrySelectChildrenWithHighestUtility)
+
+	UE_STATETREE_DEBUG_SCOPED_STATE_PHASE(this, NextStateHandle, EStateTreeUpdatePhase::TrySelectBehavior);
+
+	TArray<TPair<uint16, float>, TInlineAllocator<16>> NextLevelChildStates;
 		{
-			if (NextState.HasChildren())
-			{
-				UE_STATETREE_DEBUG_SCOPED_STATE_PHASE(this, NextStateHandle, EStateTreeUpdatePhase::TrySelectBehavior);
+		using namespace UE::StateTree::ExecutionContext::Private;
+		const FStateTreeExecutionFrame* NextParentFrame = OutSelectionResult.SelectedFrames.Num() > 1
+			? FindExecutionFrame(OutSelectionResult.SelectedFrames.Last(1), MakeConstArrayView(GetExecState().ActiveFrames), MakeConstArrayView(OutSelectionResult.TemporaryFrames))
+			: nullptr;
+		FStateTreeExecutionFrame* NextFrame = FindExecutionFrame(OutSelectionResult.SelectedFrames.Last(), MakeArrayView(GetExecState().ActiveFrames), MakeArrayView(OutSelectionResult.TemporaryFrames));
+		check(NextFrame);
 
-				TArray<uint16, TInlineAllocator<8>> NextLevelChildStates;
-				for (uint16 ChildState = NextState.ChildrenBegin; ChildState < NextState.ChildrenEnd; ChildState = CurrentStateTree->States[ChildState].GetNextSibling())
+		for (uint16 ChildState = NextState.ChildrenBegin; ChildState < NextState.ChildrenEnd; ChildState = NextStateTree->States[ChildState].GetNextSibling())
 				{
-					NextLevelChildStates.Push(ChildState);
+			const FCompactStateTreeState& CurrentState = NextStateTree->States[ChildState];
+			const float Score = EvaluateUtilityWithValidation(NextParentFrame, *NextFrame, FStateTreeStateHandle(ChildState), CurrentState.UtilityConsiderationsBegin, CurrentState.UtilityConsiderationsNum, CurrentState.Weight);
+			NextLevelChildStates.Emplace(ChildState, Score);
+		}
 				}
 
 				while (!NextLevelChildStates.IsEmpty())
 				{
 					//Find one with highest score in the remaining candidates
-					float HighestScore = -std::numeric_limits<float>::infinity();;
-					uint16 StateIndexWithHighestScore = FStateTreeStateHandle::InvalidIndex;
+		float HighestScore = -std::numeric_limits<float>::infinity();
 					int32 ArrayIndexWithHighestScore = INDEX_NONE;
 					for (int32 Index = 0; Index < NextLevelChildStates.Num(); ++Index)
 					{
-						const uint16 CurrentStateIndex = NextLevelChildStates[Index];
-						const FCompactStateTreeState& CurrentState = CurrentStateTree->States[CurrentStateIndex];
-						const FStateTreeStateHandle CurrentStateHandle = FStateTreeStateHandle(CurrentStateIndex);
-						const float Score = EvaluateUtilityWithValidation(CurrentParentFrame, CurrentFrame, CurrentStateHandle, CurrentState.UtilityConsiderationsBegin, CurrentState.UtilityConsiderationsNum, CurrentState.Weight);
+			const float Score = NextLevelChildStates[Index].Get<1>();
 						if (Score > HighestScore)
 						{
 							HighestScore = Score;
-							StateIndexWithHighestScore = CurrentStateIndex;
 							ArrayIndexWithHighestScore = Index;
 						}
 					}
 
-					if (FStateTreeStateHandle::IsValidIndex(StateIndexWithHighestScore))
+		if (!NextLevelChildStates.IsValidIndex(ArrayIndexWithHighestScore))
 					{
-						if (SelectStateInternal(CurrentParentFrame, CurrentFrame, CurrentFrameInActiveFrames, { FStateTreeStateHandle(StateIndexWithHighestScore) }, OutSelectionResult))
+			return false;
+		}
+
+		FStateHandleContext ChildState = FStateHandleContext(NextStateTree, FStateTreeStateHandle(NextLevelChildStates[ArrayIndexWithHighestScore].Get<0>()));
+		FSelectStateInternalArguments NewInternalArgs = InternalArgs;
+		NewInternalArgs.MissingStatesToReachTarget = MakeArrayView(&ChildState, 1);
+		if (SelectStateInternal(Args, NewInternalArgs, OutSelectionResult))
 						{
-							// Selection succeeded
-							bSucceededToSelectState = true;
-							break;
+			return true;
 						}
 						
 						// Disqualify the state we failed to enter
-						NextLevelChildStates.RemoveAtSwap(ArrayIndexWithHighestScore, EAllowShrinking::No);
-					}
-					else
-					{
-						// No states in array were valid
-						break;
+		constexpr EAllowShrinking AllowShrinking = EAllowShrinking::No;
+		NextLevelChildStates.RemoveAtSwap(ArrayIndexWithHighestScore, AllowShrinking);
 					}
+
+	return false;
 				}
 
-				if (bSucceededToSelectState)
+bool FStateTreeExecutionContext::SelectStateInternal_TrySelectChildrenAtRandomWeightedByUtility(
+	const FSelectStateArguments& Args,
+	const FSelectStateInternalArguments& InternalArgs,
+	FSelectStateResult& OutSelectionResult,
+	TNotNull<const UStateTree*> NextStateTree,
+	const FCompactStateTreeState& NextState,
+	const FStateTreeStateHandle NextStateHandle)
 				{
-					break;
-				}
-			}
-			else
+	using namespace UE::StateTree::ExecutionContext;
+
+	if (!NextState.HasChildren())
 			{
 				// Select this state (For backwards compatibility)
 				UE_STATETREE_DEBUG_STATE_EVENT(this, NextStateHandle, EStateTreeTraceEventType::OnStateSelected);
-				bSucceededToSelectState = true;
-				break;
+		return true;
 			}
-		}
-		else if (NextState.SelectionBehavior == EStateTreeStateSelectionBehavior::TrySelectChildrenAtRandomWeightedByUtility)
+
+	UE_STATETREE_DEBUG_SCOPED_STATE_PHASE(this, NextStateHandle, EStateTreeUpdatePhase::TrySelectBehavior);
+
+	TArray<TTuple<uint16, float>, TInlineAllocator<16>> NextLevelChildStates;
+	float TotalScore = 0.0f;
 		{
-			if (NextState.HasChildren())
+		using namespace UE::StateTree::ExecutionContext::Private;
+		const FStateTreeExecutionFrame* NextParentFrame = OutSelectionResult.SelectedFrames.Num() > 1
+			? FindExecutionFrame(OutSelectionResult.SelectedFrames.Last(1), MakeConstArrayView(GetExecState().ActiveFrames), MakeConstArrayView(OutSelectionResult.TemporaryFrames))
+			: nullptr;
+		FStateTreeExecutionFrame* NextFrame = FindExecutionFrame(OutSelectionResult.SelectedFrames.Last(), MakeArrayView(GetExecState().ActiveFrames), MakeArrayView(OutSelectionResult.TemporaryFrames));
+		check(NextFrame);
+
+		for (uint16 ChildState = NextState.ChildrenBegin; ChildState < NextState.ChildrenEnd; ChildState = NextStateTree->States[ChildState].GetNextSibling())
 			{
-				TArray<TTuple<uint16, float>, TInlineAllocator<8>> NextLevelChildStates;
-				float TotalScore = .0f;
-				for (uint16 CurrentStateIndex = NextState.ChildrenBegin; CurrentStateIndex < NextState.ChildrenEnd; CurrentStateIndex = CurrentStateTree->States[CurrentStateIndex].GetNextSibling())
+			const FCompactStateTreeState& CurrentState = NextStateTree->States[ChildState];
+			const float Score = EvaluateUtilityWithValidation(NextParentFrame, *NextFrame, FStateTreeStateHandle(ChildState), CurrentState.UtilityConsiderationsBegin, CurrentState.UtilityConsiderationsNum, CurrentState.Weight);
+			if (Score > 0.0f)
 				{
-					const FCompactStateTreeState& CurrentState = CurrentStateTree->States[CurrentStateIndex];
-					const FStateTreeStateHandle CurrentStateHandle = FStateTreeStateHandle(CurrentStateIndex);
-					const float CurrentStateScore = EvaluateUtilityWithValidation(CurrentParentFrame, CurrentFrame, CurrentStateHandle, CurrentState.UtilityConsiderationsBegin, CurrentState.UtilityConsiderationsNum, CurrentState.Weight);
-					NextLevelChildStates.Emplace(CurrentStateIndex, CurrentStateScore);
-					TotalScore += CurrentStateScore;
+				NextLevelChildStates.Emplace(ChildState, Score);
+				TotalScore += Score;
+			}
+		}
 				}
 
+	const FStateTreeExecutionState& Exec = GetExecState();
 				while (!NextLevelChildStates.IsEmpty())
 				{
 					const float RandomScore = Exec.RandomStream.FRand() * TotalScore;
-					float AccumulatedScore = .0f;
+		float AccumulatedScore = 0.0f;
 					for (int32 Index = 0; Index < NextLevelChildStates.Num(); ++Index)
 					{
 						const TTuple<uint16, float>& StateScorePair = NextLevelChildStates[Index];
@@ -6053,57 +7074,158 @@ bool FStateTreeExecutionContext::SelectStateInternal(
 						const float StateScore = StateScorePair.Value;
 						AccumulatedScore += StateScore;
 
-						if (RandomScore < AccumulatedScore || (Index == (NextLevelChildStates.Num() - 1)))
+			if (RandomScore <= AccumulatedScore || (Index == (NextLevelChildStates.Num() - 1)))
 						{
-							// States with zero possibility won't be selected
-							if (StateScore != 0.f && SelectStateInternal(CurrentParentFrame, CurrentFrame, CurrentFrameInActiveFrames, { FStateTreeStateHandle(StateIndex) }, OutSelectionResult))
+				FStateHandleContext ChildState = FStateHandleContext(NextStateTree, FStateTreeStateHandle(StateIndex));
+				FSelectStateInternalArguments NewInternalArgs = InternalArgs;
+				NewInternalArgs.MissingStatesToReachTarget = MakeArrayView(&ChildState, 1);
+				if (SelectStateInternal(Args, NewInternalArgs, OutSelectionResult))
 							{
-								// Selection succeeded
-								bSucceededToSelectState = true;
-								break;
+					return true;
 							}
 
-							//Disqualify the state we failed to enter, and restart the loop
-							TotalScore -= StateScore;
+				// Disqualify the state we failed to enter
 							constexpr EAllowShrinking AllowShrinking = EAllowShrinking::No;
 							NextLevelChildStates.RemoveAtSwap(Index, AllowShrinking);
+				break;
+			}
+		}
+	}
 
-							break;
+	return false;
+}
+
+bool FStateTreeExecutionContext::SelectStateInternal_TryFollowTransitions(
+	const FSelectStateArguments& Args,
+	const FSelectStateInternalArguments& InternalArgs,
+	FSelectStateResult& OutSelectionResult,
+	TNotNull<const UStateTree*> NextStateTree,
+	const FCompactStateTreeState& NextState,
+	const FStateTreeStateHandle NextStateHandle)
+{
+	using namespace UE::StateTree;
+	using namespace UE::StateTree::ExecutionContext;
+	using namespace UE::StateTree::ExecutionContext::Private;
+
+	UE_STATETREE_DEBUG_SCOPED_STATE_PHASE(this, NextStateHandle, EStateTreeUpdatePhase::TrySelectBehavior);
+
+	bool bSucceededToSelectState = false;
+	EStateTreeTransitionPriority CurrentPriority = EStateTreeTransitionPriority::None;
+
+	for (uint8 Index = 0; Index < NextState.TransitionsNum; ++Index)
+	{
+		const int32 TransitionIndex = NextState.TransitionsBegin + Index;
+		const FCompactStateTransition& Transition = NextStateTree->Transitions[TransitionIndex];
+
+		// Skip disabled transitions
+		if (Transition.bTransitionEnabled == false)
+		{
+			continue;
+		}
+
+		// No need to test the transition if same or higher priority transition has already been processed.
+		if (Transition.Priority <= CurrentPriority)
+		{
+			continue;
 						}
+
+		// Skip completion transitions
+		if (EnumHasAnyFlags(Transition.Trigger, EStateTreeTransitionTrigger::OnStateCompleted))
+		{
+			continue;
 					}
 
-					if (bSucceededToSelectState)
+		// Skip invalid state or completion state
+		if (!Transition.State.IsValid() || Transition.State.IsCompletionState())
 					{
-						break;
+			continue;
 					}
+
+		// Cannot follow transitions with delay.
+		if (Transition.HasDelay())
+		{
+			continue;
 				}
 
-				if (bSucceededToSelectState)
+		// Try to prevent (infinite) loops in the selection.
+		const bool bSelectionContainsState = OutSelectionResult.SelectedStates.ContainsByPredicate(
+			[State = Transition.State](const FActiveState& Other)
+			{
+				return State == Other.GetStateHandle();
+			});
+		if (bSelectionContainsState)
 				{
-					break;
+			STATETREE_LOG(Warning, TEXT("%hs: Loop detected when trying to select state %s from '%s'. Prior states: %s.  '%s' using StateTree '%s'.")
+				, __FUNCTION__
+				, *GetSafeStateName(NextStateTree, Transition.State)
+				, *GetStateStatusString(GetExecState())
+				, *GetStatePathAsString(&RootStateTree, OutSelectionResult.SelectedStates)
+				, *Owner.GetName()
+				, *NextStateTree->GetFullName());
+			continue;
 				}
+
+		FSharedEventInlineArray TransitionEvents;
+		GetTriggerTransitionEvent(Transition, Storage, Args.TransitionEvent, GetEventsToProcessView(), TransitionEvents);
+
+		for (const FStateTreeSharedEvent& SelectedStateTransitionEvent : TransitionEvents)
+		{
+			bool bTransitionConditionsPassed = false;
+			{
+				using namespace UE::StateTree::ExecutionContext::Private;
+				const FStateTreeExecutionFrame* NextParentFrame = OutSelectionResult.SelectedFrames.Num() > 1
+					? FindExecutionFrame(OutSelectionResult.SelectedFrames.Last(1), MakeConstArrayView(GetExecState().ActiveFrames), MakeConstArrayView(OutSelectionResult.TemporaryFrames))
+					: nullptr;
+				FStateTreeExecutionFrame* NextFrame = FindExecutionFrame(OutSelectionResult.SelectedFrames.Last(), MakeArrayView(GetExecState().ActiveFrames), MakeArrayView(OutSelectionResult.TemporaryFrames));
+				check(NextFrame);
+
+				FCurrentlyProcessedTransitionEventScope TransitionEventScope(*this, SelectedStateTransitionEvent.IsValid() ? SelectedStateTransitionEvent.Get() : nullptr);
+				UE_STATETREE_DEBUG_TRANSITION_EVENT(this, FStateTreeTransitionSource(NextStateTree, FStateTreeIndex16(TransitionIndex), Transition.State, Transition.Priority), EStateTreeTraceEventType::OnEvaluating);
+				UE_STATETREE_DEBUG_SCOPED_PHASE(this, EStateTreeUpdatePhase::TransitionConditions);
+				bTransitionConditionsPassed = TestAllConditionsWithValidation(NextParentFrame, *NextFrame, NextStateHandle, Transition.ConditionsBegin, Transition.ConditionsNum);
 			}
-			else
+
+			if (bTransitionConditionsPassed)
 			{
-				// Select this state (For backwards compatibility)
-				UE_STATETREE_DEBUG_STATE_EVENT(this, NextStateHandle, EStateTreeTraceEventType::OnStateSelected);
+				// Using SelectState() instead of SelectStateInternal to treat the transitions the same way as regular transitions,
+				// e.g. it may jump to a completely different branch.
+
+				FSelectStateArguments SelectStateArgs;
+				SelectStateArgs.SourceState = Args.SourceState;
+				SelectStateArgs.TargetState = FStateHandleContext(NextStateTree, Transition.State);
+				SelectStateArgs.TransitionEvent = SelectedStateTransitionEvent;
+				SelectStateArgs.Fallback = Transition.Fallback;
+				FSelectStateResult StateSelectionResult;
+				StateSelectionResult.TemporaryFrames = OutSelectionResult.TemporaryFrames;
+				if (SelectState(SelectStateArgs, StateSelectionResult))
+			{
+					OutSelectionResult = MoveTemp(StateSelectionResult);
+					CurrentPriority = Transition.Priority;
 				bSucceededToSelectState = true;
+
+					//@todo sort the transition by priority at compile time. This will solve having to loop back once we found a valid transition.
+					// Consume the transition event
+					if (SelectedStateTransitionEvent.IsValid() && Transition.bConsumeEventOnSelect)
+					{
+						ConsumeEvent(SelectedStateTransitionEvent);
+					}
+
+					// Cannot return because higher priority transitions may override the selection. 
 				break;
 			}
 		}
 	}
-
-	if (!bSucceededToSelectState)
-	{
-		// State could not be selected, restore.
-		CurrentFrame.NumCurrentlyActiveStates = PrevNumCurrentlyActiveStates;
-		CurrentFrame.ActiveStates.Pop();
 	}
 
 	return bSucceededToSelectState;
 }
 
 FString FStateTreeExecutionContext::GetSafeStateName(const FStateTreeExecutionFrame& CurrentFrame, const FStateTreeStateHandle State) const
+{
+	return GetSafeStateName(CurrentFrame.StateTree, State);
+}
+
+FString FStateTreeExecutionContext::GetSafeStateName(const UStateTree* StateTree, const FStateTreeStateHandle State) const
 {
 	if (State == FStateTreeStateHandle::Invalid)
 	{
@@ -6117,9 +7239,9 @@ FString FStateTreeExecutionContext::GetSafeStateName(const FStateTreeExecutionFr
 	{
 		return TEXT("(State Failed)");
 	}
-	else if (CurrentFrame.StateTree && CurrentFrame.StateTree->States.IsValidIndex(State.Index))
+	else if (StateTree && StateTree->States.IsValidIndex(State.Index))
 	{
-		return *CurrentFrame.StateTree->States[State.Index].Name.ToString();
+		return *StateTree->States[State.Index].Name.ToString();
 	}
 	return TEXT("(Unknown)");
 }
@@ -6165,7 +7287,16 @@ FString FStateTreeExecutionContext::GetStateStatusString(const FStateTreeExecuti
 	{
 		return TEXT("--:") + UEnum::GetDisplayValueAsText(ExecState.LastTickStatus).ToString();
 	}
-	return GetSafeStateName(ExecState.ActiveFrames.Last(), ExecState.ActiveFrames.Last().ActiveStates.Last()) + TEXT(":") + UEnum::GetDisplayValueAsText(ExecState.LastTickStatus).ToString();
+
+	if (ExecState.ActiveFrames.Num())
+	{
+		const FStateTreeExecutionFrame& LastFrame = ExecState.ActiveFrames.Last();
+		if (LastFrame.ActiveStates.Num() > 0)
+		{
+			return GetSafeStateName(LastFrame, LastFrame.ActiveStates.Last()) + TEXT(":") + UEnum::GetDisplayValueAsText(ExecState.LastTickStatus).ToString();
+		}
+	}
+	return FString();
 }
 
 // Deprecated
