[StateTree] Transition with ID. Use the state and frame ID to identify the state during a transition request. Add new rules for transitions for backward compatibility with the previous expected behavi

diff --git a/Engine/Plugins/Runtime/StateTree/Source/StateTreeModule/Private/StateTreeExecutionTypes.cpp b/Engine/Plugins/Runtime/StateTree/Source/StateTreeModule/Private/StateTreeExecutionTypes.cpp
--- a/Engine/Plugins/Runtime/StateTree/Source/StateTreeModule/Private/StateTreeExecutionTypes.cpp
+++ b/Engine/Plugins/Runtime/StateTree/Source/StateTreeModule/Private/StateTreeExecutionTypes.cpp
@@ -369,7 +369,7 @@ void FStateTreeDelegateActiveListeners::BroadcastDelegate(const FStateTreeDelega
 					else
 					{
 						const int32 FoundStateIndex = ExectionFrame->ActiveStates.IndexOfReverse(ActiveListener.StateID);
-						if (FoundStateIndex != INDEX_NONE && FoundStateIndex < ExectionFrame->NumCurrentlyActiveStates)
+						if (FoundStateIndex != INDEX_NONE)
 						{
 							ActiveListener.Delegate.ExecuteIfBound();
 						}
diff --git a/Engine/Plugins/Runtime/StateTree/Source/StateTreeModule/Private/StateTreeInstanceData.cpp b/Engine/Plugins/Runtime/StateTree/Source/StateTreeModule/Private/StateTreeInstanceData.cpp
--- a/Engine/Plugins/Runtime/StateTree/Source/StateTreeModule/Private/StateTreeInstanceData.cpp
+++ b/Engine/Plugins/Runtime/StateTree/Source/StateTreeModule/Private/StateTreeInstanceData.cpp
@@ -103,7 +103,7 @@ namespace UE::StateTree::InstanceData::Private
 	bool IsActiveInstanceHandleSourceValid(const FStateTreeInstanceStorage& Storage, const FStateTreeExecutionFrame& CurrentFrame, const FStateTreeDataHandle& Handle)
 	{
 		return CurrentFrame.ActiveInstanceIndexBase.IsValid()
-			&& CurrentFrame.ActiveStates.Contains(Handle.GetState(), CurrentFrame.NumCurrentlyActiveStates)
+			&& CurrentFrame.ActiveStates.Contains(Handle.GetState())
 			&& Storage.IsValidIndex(CurrentFrame.ActiveInstanceIndexBase.Get() + Handle.GetIndex());
 	}
 
diff --git a/Engine/Plugins/Runtime/StateTree/Source/StateTreeModule/Private/StateTreeSchema.cpp b/Engine/Plugins/Runtime/StateTree/Source/StateTreeModule/Private/StateTreeSchema.cpp
--- a/Engine/Plugins/Runtime/StateTree/Source/StateTreeModule/Private/StateTreeSchema.cpp
+++ b/Engine/Plugins/Runtime/StateTree/Source/StateTreeModule/Private/StateTreeSchema.cpp
@@ -6,7 +6,25 @@
 
 #include UE_INLINE_GENERATED_CPP_BY_NAME(StateTreeSchema)
 
-bool UStateTreeSchema::IsChildOfBlueprintBase(const UClass* InClass) const
+namespace UE::StateTree::Private
+{
+	bool bCompletedTransitionStatesCreateNewStates = true;
+	FAutoConsoleVariableRef CVarCompletedTransitionStatesCreateNewStates(
+		TEXT("StateTree.SelectState.CompletedTransitionStatesCreateNewStates"),
+		bCompletedTransitionStatesCreateNewStates,
+		TEXT("Activate the EStateTreeStateSelectionRules::CompletedTransitionStatesCreateNewStates rule.")
+	);
+
+	bool bCompletedStateBeforeTransitionSourceFailsTransition = true;
+	FAutoConsoleVariableRef CVarCompletedStateBeforeTransitionSourceFailsTransition(
+		TEXT("StateTree.SelectState.CompletedStateBeforeTransitionSourceFailsTransition"),
+		bCompletedStateBeforeTransitionSourceFailsTransition,
+		TEXT("Activate the EStateTreeStateSelectionRules::CompletedStateBeforeTransitionSourceFailsTransition rule.")
+	);
+}
+
+
+bool UStateTreeSchema::IsChildOfBlueprintBase(const UClass* InClass)
 {
 	return InClass->IsChildOf<UStateTreeNodeBlueprintBase>();
 }
@@ -14,4 +32,20 @@ bool UStateTreeSchema::IsChildOfBlueprintBase(const UClass* InClass) const
 EStateTreeParameterDataType UStateTreeSchema::GetGlobalParameterDataType() const
 {
 	return EStateTreeParameterDataType::GlobalParameterData;
-}
\ No newline at end of file
+}
+
+EStateTreeStateSelectionRules UStateTreeSchema::GetStateSelectionRules() const
+{
+	EStateTreeStateSelectionRules Result = EStateTreeStateSelectionRules::None;
+	if (UE::StateTree::Private::bCompletedTransitionStatesCreateNewStates)
+	{
+		Result |= EStateTreeStateSelectionRules::CompletedTransitionStatesCreateNewStates;
+	}
+
+	if (UE::StateTree::Private::bCompletedStateBeforeTransitionSourceFailsTransition)
+	{
+		Result |= EStateTreeStateSelectionRules::CompletedStateBeforeTransitionSourceFailsTransition;
+	}
+
+	return Result;
+}
diff --git a/Engine/Plugins/Runtime/StateTree/Source/StateTreeModule/Private/Tasks/StateTreeRunParallelStateTreeTask.cpp b/Engine/Plugins/Runtime/StateTree/Source/StateTreeModule/Private/Tasks/StateTreeRunParallelStateTreeTask.cpp
--- a/Engine/Plugins/Runtime/StateTree/Source/StateTreeModule/Private/Tasks/StateTreeRunParallelStateTreeTask.cpp
+++ b/Engine/Plugins/Runtime/StateTree/Source/StateTreeModule/Private/Tasks/StateTreeRunParallelStateTreeTask.cpp
@@ -32,8 +32,12 @@ EStateTreeRunStatus FStateTreeRunParallelStateTreeTask::EnterState(FStateTreeExe
 	}
 
 	// Find if it's a recursive call. The detection is not perfect. For example: StateTrees with a parallel task that links to each other cannot be detected.
-	const bool bInParentContext = Context.GetActiveFrames().ContainsByPredicate([NewTree = StateTreeToRun.GetStateTree()](const FStateTreeExecutionFrame& Frame){ return Frame.StateTree == NewTree; });
-	if (bInParentContext)
+	const bool bInParentContext = Context.GetActiveFrames().ContainsByPredicate([NewTree = StateTreeToRun.GetStateTree()](const FStateTreeExecutionFrame& Frame)
+		{
+			return Frame.StateTree == NewTree;
+		});
+	const bool bFromParentProcessedFrame = Context.GetCurrentlyProcessedFrame() != nullptr ? Context.GetCurrentlyProcessedFrame()->StateTree == StateTreeToRun.GetStateTree() : false;
+	if (bInParentContext || bFromParentProcessedFrame)
 	{
 		UE_LOG(LogStateTree, Warning, TEXT("Trying to start a new parallel tree from the same tree '%s'"), *StateTreeToRun.GetStateTree()->GetName());
 		return EStateTreeRunStatus::Failed;
diff --git a/Engine/Plugins/Runtime/StateTree/Source/StateTreeModule/Public/StateTree.h b/Engine/Plugins/Runtime/StateTree/Source/StateTreeModule/Public/StateTree.h
--- a/Engine/Plugins/Runtime/StateTree/Source/StateTreeModule/Public/StateTree.h
+++ b/Engine/Plugins/Runtime/StateTree/Source/StateTreeModule/Public/StateTree.h
@@ -243,6 +243,12 @@ public:
 		return bScheduledTickAllowed;
 	}
 
+	/** @return the rules used by the execution context for selecting states. */
+	EStateTreeStateSelectionRules GetStateSelectionRules() const
+	{
+		return StateSelectionRules;
+	}
+
 #if WITH_EDITOR
 	/** Resets the compiled data to empty. */
 	UE_API void ResetCompiled();
@@ -421,6 +427,9 @@ private:
 	UPROPERTY()
 	uint16 GlobalTasksNum = 0;
 
+	/** The cached value of UStateTreeSchema::GetStateSelectionRules */
+	EStateTreeStateSelectionRules StateSelectionRules = EStateTreeStateSelectionRules::Default;
+
 	/** How the global tasks control the completion of the frame. */
 	UPROPERTY()
 	EStateTreeTaskCompletionType CompletionGlobalTasksControl = EStateTreeTaskCompletionType::Any;
diff --git a/Engine/Plugins/Runtime/StateTree/Source/StateTreeModule/Public/StateTreeExecutionContext.h b/Engine/Plugins/Runtime/StateTree/Source/StateTreeModule/Public/StateTreeExecutionContext.h
--- a/Engine/Plugins/Runtime/StateTree/Source/StateTreeModule/Public/StateTreeExecutionContext.h
+++ b/Engine/Plugins/Runtime/StateTree/Source/StateTreeModule/Public/StateTreeExecutionContext.h
@@ -38,6 +38,19 @@ namespace UE::StateTree::ExecutionContext
 	UE_API EStateTreeRunStatus CastToRunStatus(EStateTreeFinishTaskType FinishTask);
 	UE_API UE::StateTree::ETaskCompletionStatus CastToTaskStatus(EStateTreeRunStatus InStatus);
 	UE_API EStateTreeRunStatus CastToRunStatus(UE::StateTree::ETaskCompletionStatus InStatus);
+
+	/** Helper that combines the state handle and its tree asset. */
+	struct FStateHandleContext
+	{
+		FStateHandleContext() = default;
+		FStateHandleContext(TNotNull<const UStateTree*> InStateTree, FStateTreeStateHandle InStateHandle)
+			: StateTree(InStateTree)
+			, StateHandle(InStateHandle)
+		{
+		}
+		const UStateTree* StateTree = nullptr;
+		FStateTreeStateHandle StateHandle;
+	};
 }
 
 /**
@@ -709,6 +722,7 @@ PRAGMA_ENABLE_DEPRECATION_WARNINGS
 		return UE::StateTree::InstanceData::GetDataView(InstanceDataStorage, CurrentlyProcessedSharedInstanceStorage, ParentFrame, CurrentFrame, Handle);
 	}
 
+	//@todo deprecate in favor of a find with frameID
 	/** Looks for a frame in provided list of frames. */
 	static UE_API const FStateTreeExecutionFrame* FindFrame(const UStateTree* StateTree, FStateTreeStateHandle RootState, TConstArrayView<FStateTreeExecutionFrame> Frames, const FStateTreeExecutionFrame*& OutParentFrame);
 
@@ -727,13 +741,55 @@ PRAGMA_ENABLE_DEPRECATION_WARNINGS
 	UE_API EStateTreeRunStatus ForceTransition(const FRecordedStateTreeTransitionResult& Transition);
 
 	/** Returns the recorded transitions for this context. */
-	TConstArrayView<FRecordedStateTreeTransitionResult> GetRecordedTransitions() const { return RecordedTransitions; }
+	TConstArrayView<FRecordedStateTreeTransitionResult> GetRecordedTransitions() const
+	{
+		return RecordedTransitions;
+	}
 
 protected:
-	/**
-	 * Describes a result of States Selection.
-	 */
-	struct FStateSelectionResult
+
+	/** Event used during state selection. Identified by the ID of the active state that consumed it. */
+	struct FSelectionEventWithID
+	{
+		UE::StateTree::FActiveState State;
+		FStateTreeSharedEvent Event;
+	};
+
+	/** Describes a result of States Selection. */
+	struct FSelectStateResult
+	{
+		FSelectStateResult() = default;
+		FSelectStateResult(const FSelectStateResult&) = delete;
+		FSelectStateResult(FSelectStateResult&&) = default;
+		FSelectStateResult& operator=(const FSelectStateResult&) = delete;
+		FSelectStateResult& operator=(FSelectStateResult&&) = default;
+
+		/** The active states selected. They are in order from the root to the leaf. */
+		TArray<UE::StateTree::FActiveState> SelectedStates;
+
+		/** The selected frame ID. The frame can be in the current active list or in the TemporaryFrames list. */
+		TArray<UE::StateTree::FActiveFrameID, TInlineAllocator<4>> SelectedFrames;
+
+		/**
+		 * New execution frame created during the state selection.
+		 * The active state list is empty and will be filled during EnterState()
+		 */
+		TArray<FStateTreeExecutionFrame, TInlineAllocator<2>> TemporaryFrames;
+
+		/** Events used during the state selection. */
+		TArray<FSelectionEventWithID, TInlineAllocator<2>> SelectionEvents;
+
+		/** The requested target of the selection. */
+		UE::StateTree::FActiveState TargetState;
+
+		/** @return a new temporary frame */
+		UE_API FStateTreeExecutionFrame& MakeAndAddTemporaryFrame(UE::StateTree::FActiveFrameID FrameID, const UE::StateTree::FExecutionFrameKey& FrameKey, bool bIsGlobalFrame);
+		UE_API FStateTreeExecutionFrame& MakeAndAddTemporaryFrameWithNewRoot(UE::StateTree::FActiveFrameID FrameID, const UE::StateTree::FExecutionFrameKey& FrameKey, FStateTreeExecutionFrame& OtherFrame);
+	};
+
+	struct
+	UE_DEPRECATED(5.7, "The selection result changed to include the state ID and frame ID")
+	FStateSelectionResult
 	{
 		/** Max number of execution frames handled during state selection. */
 		static constexpr int32 MaxExecutionFrames = 8;
@@ -792,6 +848,14 @@ protected:
 		TArray<FStateTreeFrameStateSelectionEvents, TFixedAllocator<MaxExecutionFrames>> FramesStateSelectionEvents;
 	};
 
+	/** The result of the RequestTransition function. */
+	struct FRequestTransitionResult
+	{
+		FSelectStateResult Selection;
+		FStateTreeTransitionResult Transition;
+		FStateTreeTransitionSource Source;
+	};
+
 	/** @return Prefix that will be used by STATETREE_LOG and STATETREE_CLOG, Owner name by default. */
 	UE_DEPRECATED(5.6, "Use FStateTreeExecutionExtension::GetInstanceDescription instead")
 	UE_API virtual FString GetInstanceDescription() const final;
@@ -799,22 +863,32 @@ protected:
 	/** Callback when delayed transition is triggered. Contexts that are event based can use this to trigger a future event. */
 	virtual void BeginDelayedTransition(const FStateTreeTransitionDelayedState& DelayedState) {};
 
+	UE_DEPRECATED(5.7, "Use the UpdateInstanceData with the FSelectStateResult argument.")
 	UE_API void UpdateInstanceData(TConstArrayView<FStateTreeExecutionFrame> CurrentActiveFrames, TArrayView<FStateTreeExecutionFrame> NextActiveFrames);
 
+	/** Allocate the instance data. Fixup the TemporarySelectedFrame with the correct data. */
+	UE_API void UpdateInstanceData(FSelectStateResult& SelectStateResult);
+
+	UE_DEPRECATED(5.7, "Use the EnterState with the FSelectStateResult argument.")
+	UE_API EStateTreeRunStatus EnterState(FStateTreeTransitionResult& Transition);
+
 	/**
 	 * Handles logic for entering State. EnterState is called on new active Evaluators and Tasks that are part of the re-planned tree.
 	 * Re-planned tree is from the transition target up to the leaf state. States that are parent to the transition target state
 	 * and still active after the transition will remain intact.
-	 * @return Run status returned by the tasks.
+	 * @return run status returned by the tasks.
 	 */
-	UE_API EStateTreeRunStatus EnterState(FStateTreeTransitionResult& Transition);
+	UE_API EStateTreeRunStatus EnterState(FSelectStateResult& SelectStateResult, const FStateTreeTransitionResult& Transition);
+
+	UE_DEPRECATED(5.7, "Use the ExitState with the FSelectStateResult argument.")
+	UE_API void ExitState(const FStateTreeTransitionResult& Transition);
 
 	/**
 	 * Handles logic for exiting State. ExitState is called on current active Evaluators and Tasks that are part of the re-planned tree.
 	 * Re-planned tree is from the transition target up to the leaf state. States that are parent to the transition target state
 	 * and still active after the transition will remain intact.
 	 */
-	UE_API void ExitState(const FStateTreeTransitionResult& Transition);
+	UE_API void ExitState(const FSelectStateResult& SelectStateResult, const FStateTreeTransitionResult& Transition);
 
 	/**
 	 * Removes all delegate listeners.
@@ -1002,9 +1076,7 @@ protected:
 	 */
 	UE_API void EvaluatePropertyFunctionsWithValidation(const FStateTreeExecutionFrame* CurrentParentFrame, const FStateTreeExecutionFrame& CurrentFrame, FStateTreeIndex16 FuncsBegin, uint16 FuncsNum);
 
-	/**
-	 * Requests transition to a specified state with specified priority.
-	 */
+	UE_DEPRECATED(5.7, "Use the RequestTransitionInternal with the active state ID")
 	UE_API bool RequestTransition(
 		const FStateTreeExecutionFrame& CurrentFrame,
 		const FStateTreeStateHandle NextState,
@@ -1012,22 +1084,39 @@ protected:
 		const FStateTreeSharedEvent* TransitionEvent = nullptr,
 		const EStateTreeSelectionFallback Fallback = EStateTreeSelectionFallback::None);
 
-	/**
-	 * Sets up NextTransition based on the provided parameters and the current execution status. 
-	 */
+	/** Requests transition to a specified state with specified priority. */
+	UE_API bool RequestTransitionInternal(
+		const FStateTreeExecutionFrame& SourceFrame,
+		UE::StateTree::FActiveStateID SourceStateID,
+		UE::StateTree::ExecutionContext::FStateHandleContext TargetState,
+		EStateTreeTransitionPriority Priority = EStateTreeTransitionPriority::Normal,
+		FStateTreeSharedEvent TransitionEvent = FStateTreeSharedEvent(),
+		EStateTreeSelectionFallback Fallback = EStateTreeSelectionFallback::None);
+
+	UE_DEPRECATED(5.7, "Use the SetupNextTransition with the active state ID")
 	UE_API void SetupNextTransition(const FStateTreeExecutionFrame& CurrentFrame, const FStateTreeStateHandle NextState, const EStateTreeTransitionPriority Priority);
 
-	/**
-	 * Triggers transitions based on current run status. CurrentStatus is used to select which transitions events are triggered.
-	 * If CurrentStatus is "Running", "Conditional" transitions pass, "Completed/Failed" will trigger "OnCompleted/OnSucceeded/OnFailed" transitions.
-	 * Transition target state can point to a selector state. For that reason the result contains both the target state, as well ass
-	 * the actual next state returned by the selector.
-	 * @return Transition result describing the source state, state transitioned to, and next selected state.
-	 */
+	/** Sets up NextTransition based on the provided parameters and the current execution status. */
+	UE_API void SetupNextTransition(
+		const FStateTreeExecutionFrame& SourceFrame,
+		UE::StateTree::FActiveStateID SourceStateID,
+		UE::StateTree::ExecutionContext::FStateHandleContext TargetState,
+		EStateTreeTransitionPriority Priority);
+
+	/** Sets up NextTransition based on the provided parameters and the current execution status. */
+	UE_API void SetupNextTransition(
+		const FStateTreeExecutionFrame& SourceFrame,
+		UE::StateTree::FActiveStateID SourceStateID,
+		UE::StateTree::ExecutionContext::FStateHandleContext TargetState,
+		EStateTreeTransitionPriority Priority,
+		FStateTreeTransitionResult& OutTransitionResult);
+
+	/** Tick the states transitions. */
 	UE_API bool TriggerTransitions();
 
 	/** Create a new transition result from a recorded transition result. It will fail if the recorded transition is malformed. */
 	UE_API TOptional<FStateTreeTransitionResult> MakeTransitionResult(const FRecordedStateTreeTransitionResult& Transition) const;
+
 	/** Create a new recorded transition from a transition result. */
 	UE_API FRecordedStateTreeTransitionResult MakeRecordedTransitionResult(const FStateTreeTransitionResult& Transition) const;
 
@@ -1058,36 +1147,136 @@ PRAGMA_ENABLE_DEPRECATION_WARNINGS
 	/** In the list of completed states, returns the run status when it's completed by a global task. */
 	UE_API EStateTreeRunStatus GetGlobalTasksCompletedStatesStatus() const;
 
+	UE_DEPRECATED(5.7, "Use the SelectState with the FSelectStateArguments argument.")
+	UE_API bool SelectState(const FStateTreeExecutionFrame& CurrentFrame, const FStateTreeStateHandle NextState, FSelectStateResult& OutSelectionResult, const FStateTreeSharedEvent* TransitionEvent = nullptr, const EStateTreeSelectionFallback Fallback = EStateTreeSelectionFallback::None);
+
+	/** The different type of selections. */
+	enum class ESelectStateBehavior : uint8
+	{
+		/** From a state transition. Normal rules apply. */
+		StateTransition,
+		/**
+		 * From a force transition.
+		 * The state's selection behavior and condition will not be respected.
+		 * The target is the final state.
+		 */
+		Forced,
+	};
+	struct FSelectStateArguments
+	{
+		/** The current list of active states. */
+		TArrayView<const UE::StateTree::FActiveState> ActiveStates;
+		/** The state that requested the transition. */
+		UE::StateTree::FActiveState SourceState;
+		/** The state that we want to select. Depending on the state's selection behavior, another state can be selected. */
+		UE::StateTree::ExecutionContext::FStateHandleContext TargetState;
+		/** Transition event used by transition to trigger the state selection. */
+		FStateTreeSharedEvent TransitionEvent;
+		/** Fallback selection behavior to execute if it fails to select the desired state. */
+		EStateTreeSelectionFallback Fallback = EStateTreeSelectionFallback::None;
+		/** The type of selection. */
+		ESelectStateBehavior Behavior = ESelectStateBehavior::StateTransition;
+	};
+
 	/**
-	 * Runs state selection logic starting at the specified state, walking towards the leaf states.
-	 * If a state cannot be selected, false is returned. 
-	 * If NextState is a selector state, SelectStateInternal is called recursively (depth-first) to all child states (where NextState will be one of child states).
-	 * If NextState is a leaf state, the active states leading from root to the leaf are returned.
-	 * @param CurrentFrame The frame where the NextState is valid. 
-	 * @param NextState The state which we try to select next.
-	 * @param OutSelectionResult The result of selection.
-	 * @param TransitionEvent Event used by transition to trigger the state selection.
-	 * @param Fallback selection behavior to execute if it fails to select the desired state
+	 * Starting at the specified state, walking towards the leaf states.
+	 * @param SelectStateArgs the arguments for SelectState.
+	 * @param OutSelectionResult the result of the selection.
 	 * @return True if succeeded to select new active states.
 	 */
-	UE_API bool SelectState(
-		const FStateTreeExecutionFrame& CurrentFrame,
-		const FStateTreeStateHandle NextState,
-		FStateSelectionResult& OutSelectionResult,
-		const FStateTreeSharedEvent* TransitionEvent = nullptr,
-		const EStateTreeSelectionFallback Fallback = EStateTreeSelectionFallback::None);
+	UE_API bool SelectState(const FSelectStateArguments& SelectStateArgs, FSelectStateResult& OutSelectionResult);
 
-	/**
-	 * Used internally to do the recursive part of the SelectState().
-	 */
+	UE_DEPRECATED(5.7, "Use the SelectStateInternal with FSelectStateArguments argument.")
 	UE_API bool SelectStateInternal(
 		const FStateTreeExecutionFrame* CurrentParentFrame,
 		FStateTreeExecutionFrame& CurrentFrame,
 		const FStateTreeExecutionFrame* CurrentFrameInActiveFrames,
 		TConstArrayView<FStateTreeStateHandle> PathToNextState,
-		FStateSelectionResult& OutSelectionResult,
+		FSelectStateResult& OutSelectionResult,
 		const FStateTreeSharedEvent* TransitionEvent = nullptr);
 
+	struct FSelectStateInternalArguments
+	{
+		TArrayView<const UE::StateTree::FActiveState> MissingActiveStates;
+		UE::StateTree::FActiveFrameID MissingSourceFrameID;
+		TArrayView<const UE::StateTree::FActiveState> MissingSourceStates;
+		TArrayView<const UE::StateTree::ExecutionContext::FStateHandleContext> MissingStatesToReachTarget;
+	};
+
+	/** Used internally to do the recursive part of the SelectState(). */
+	UE_API bool SelectStateInternal(
+		const FSelectStateArguments& Args,
+		const FSelectStateInternalArguments& InternalArgs,
+		FSelectStateResult& OutSelectionResult);
+
+	/** Used internally to select state until we reach the transition source or the transition target. */
+	UE_API TOptional<bool> SelectStateFromSourceInternal(
+		const FSelectStateArguments& Args,
+		const FSelectStateInternalArguments& InternalArgs,
+		FSelectStateResult& OutSelectionResult,
+		const FStateTreeExecutionFrame& NextFrame,
+		const FCompactStateTreeState& NextState,
+		FStateTreeStateHandle NextStateHandle,
+		bool bNewFrameCreated);
+
+	UE_API bool SelectStateInternal_Linked(
+		const FSelectStateArguments& Args,
+		const FSelectStateInternalArguments& InternalArgs,
+		FSelectStateResult& OutSelectionResult,
+		TNotNull<const UStateTree*> NextStateTree,
+		const FCompactStateTreeState& NextState,
+		bool bShouldCreateNewState
+		);
+
+	UE_API bool SelectStateInternal_LinkedAsset(
+		const FSelectStateArguments& Args,
+		const FSelectStateInternalArguments& InternalArgs,
+		FSelectStateResult& OutSelectionResult,
+		TNotNull<const UStateTree*> NextStateTree,
+		const FCompactStateTreeState& NextState,
+		const UStateTree* NextLinkedStateAsset,
+		bool bShouldCreateNewState);
+
+	UE_API bool SelectStateInternal_TrySelectChildrenInOrder(
+		const FSelectStateArguments& Args,
+		const FSelectStateInternalArguments& InternalArgs,
+		FSelectStateResult& OutSelectionResult,
+		TNotNull<const UStateTree*> NextStateTree,
+		const FCompactStateTreeState& TargetState,
+		const FStateTreeStateHandle TargetStateHandle);
+		
+	UE_API bool SelectStateInternal_TrySelectChildrenAtRandom(
+		const FSelectStateArguments& Args,
+		const FSelectStateInternalArguments& InternalArgs,
+		FSelectStateResult& OutSelectionResult,
+		TNotNull<const UStateTree*> NextStateTree,
+		const FCompactStateTreeState& TargetState,
+		const FStateTreeStateHandle TargetStateHandle);
+
+	UE_API bool SelectStateInternal_TrySelectChildrenWithHighestUtility(
+		const FSelectStateArguments& Args,
+		const FSelectStateInternalArguments& InternalArgs,
+		FSelectStateResult& OutSelectionResult,
+		TNotNull<const UStateTree*> NextStateTree,
+		const FCompactStateTreeState& TargetState,
+		const FStateTreeStateHandle TargetStateHandle);
+
+	UE_API bool SelectStateInternal_TrySelectChildrenAtRandomWeightedByUtility(
+		const FSelectStateArguments& Args,
+		const FSelectStateInternalArguments& InternalArgs,
+		FSelectStateResult& OutSelectionResult,
+		TNotNull<const UStateTree*> NextStateTree,
+		const FCompactStateTreeState& TargetState,
+		const FStateTreeStateHandle TargetStateHandle);
+
+	UE_API bool SelectStateInternal_TryFollowTransitions(
+		const FSelectStateArguments& Args,
+		const FSelectStateInternalArguments& InternalArgs,
+		FSelectStateResult& OutSelectionResult,
+		TNotNull<const UStateTree*> NextStateTree,
+		const FCompactStateTreeState& TargetState,
+		const FStateTreeStateHandle TargetStateHandle);
+
 	/** @return StateTree execution state from the instance storage. */
 	FStateTreeExecutionState& GetExecState()
 	{
@@ -1109,6 +1298,9 @@ PRAGMA_ENABLE_DEPRECATION_WARNINGS
 	/** @return String describing state name for logging and debug. */
 	UE_API FString GetSafeStateName(const FStateTreeExecutionFrame& CurrentFrame, const FStateTreeStateHandle State) const;
 
+	/** @return String describing state name for logging and debug. */
+	UE_API FString GetSafeStateName(const UStateTree* StateTree, const FStateTreeStateHandle State) const;
+
 	/** @return String describing full path of an activate state for logging and debug. */
 	UE_API FString DebugGetStatePath(TConstArrayView<FStateTreeExecutionFrame> ActiveFrames, const FStateTreeExecutionFrame* CurrentFrame = nullptr, const int32 ActiveStateIndex = INDEX_NONE) const;
 
@@ -1180,10 +1372,13 @@ PRAGMA_ENABLE_DEPRECATION_WARNINGS
 	 */
 	UE_API bool SetGlobalParameters(const FInstancedPropertyBag& Parameters);
 
+	UE_DEPRECATED(5.7, "Use the CaptureNewStateEvents with the FSelectStateResult argument.")
+	void CaptureNewStateEvents(TConstArrayView<FStateTreeExecutionFrame> PrevFrames, TConstArrayView<FStateTreeExecutionFrame> NewFrames, TArrayView<FStateTreeFrameStateSelectionEvents> FramesStateSelectionEvents);
+
 	/**
 	 * Captures StateTree events used during the state selection.
 	 */
-	UE_API void CaptureNewStateEvents(TConstArrayView<FStateTreeExecutionFrame> PrevFrames, TConstArrayView<FStateTreeExecutionFrame> NewFrames, TArrayView<FStateTreeFrameStateSelectionEvents> FramesStateSelectionEvents);
+	UE_API void CaptureNewStateEvents(FSelectStateResult& Args);
 
 	/** @return a weak reference for a task that can be stored for later uses. */
 PRAGMA_DISABLE_DEPRECATION_WARNINGS
@@ -1224,16 +1419,21 @@ PRAGMA_ENABLE_DEPRECATION_WARNINGS
 	bool bActiveExternalDataCollected = false;
 
 #if WITH_EDITORONLY_DATA
-	/** Run status set in FinishTask and that used after the task ticks. */
 	UE_DEPRECATED(5.6, "CompletedStateRunStatus is not used.")
+	/** Run status set in FinishTask and that used after the task ticks. */
 	TOptional<EStateTreeRunStatus> CompletedStateRunStatus;
-#endif
 
+	UE_DEPRECATED(5.7, "Use RequestedTransition instead.")
 	/** Next transition, used by RequestTransition(). */
 	FStateTreeTransitionResult NextTransition;
 
+	UE_DEPRECATED(5.7, "Use RequestedTransition instead.")
 	/** Structure describing the origin of the state transition that caused the state change. */
 	FStateTreeTransitionSource NextTransitionSource;
+#endif 
+
+	/** Current selected transition. */
+	TUniquePtr<FRequestTransitionResult> RequestedTransition;
 
 	/** When set, start the transitions loop from TriggerTransitionsFromFrameIndex. */
 	TOptional<int32> TriggerTransitionsFromFrameIndex;
@@ -1261,8 +1461,13 @@ PRAGMA_ENABLE_DEPRECATION_WARNINGS
 		const FStateTreeExecutionFrame* SavedParentFrame = nullptr;
 	};
 
+#if WITH_EDITORONLY_DATA
+PRAGMA_DISABLE_DEPRECATION_WARNINGS
+	UE_DEPRECATED(5.7, "The selection is now identified by the state ID.")
 	/** Current state selection result when performing recursive state selection, or nullptr if not applicable. */
 	const FStateSelectionResult* CurrentSelectionResult = nullptr;
+PRAGMA_ENABLE_DEPRECATION_WARNINGS
+#endif
 
 	/** Current state we're processing, or invalid if not applicable. */
 	FStateTreeStateHandle CurrentlyProcessedState;
@@ -1286,7 +1491,6 @@ PRAGMA_ENABLE_DEPRECATION_WARNINGS
 	private:
 		FStateTreeExecutionContext& Context;
 		FStateTreeStateHandle SavedState = FStateTreeStateHandle::Invalid;
-		bool bChanged = false;
 	};
 
 	/** Current event we're processing in transition, or invalid if not applicable. */
@@ -1311,11 +1515,17 @@ PRAGMA_ENABLE_DEPRECATION_WARNINGS
 		FStateTreeExecutionContext& Context;
 	};
 
-	/** Current frames events we're processing during the state selection, or invalid if not applicable. */
+	/** Current select state result we're processing during the state selection, or invalid if not applicable. */
+	FSelectStateResult* CurrentlyProcessedStateSelectionResult = nullptr;
+
+#if WITH_EDITORONLY_DATA
+PRAGMA_DISABLE_DEPRECATION_WARNINGS
+	UE_DEPRECATED(5.7, "FStateTreeFrameStateSelectionEvents is now unused. Use CurrentlyProcessedStateSelectionResult")
 	FStateTreeFrameStateSelectionEvents* CurrentlyProcessedStateSelectionEvents = nullptr;
 
-	/** Helper struct to track currently processed state selection events. */
-	struct FCurrentFrameStateSelectionEventsScope
+	struct
+	UE_DEPRECATED(5.7, "FStateTreeFrameStateSelectionEvents is now unused. Use CurrentlyProcessedStateSelectionResult")
+	FCurrentFrameStateSelectionEventsScope
 	{
 		FCurrentFrameStateSelectionEventsScope(FStateTreeExecutionContext& InContext, FStateTreeFrameStateSelectionEvents& InCurrentlyProcessedStateSelectionEvents)
 			: Context(InContext)
@@ -1333,6 +1543,8 @@ PRAGMA_ENABLE_DEPRECATION_WARNINGS
 		FStateTreeExecutionContext& Context;
 		FStateTreeFrameStateSelectionEvents* SavedStateSelectionEvents = nullptr;
 	};
+PRAGMA_ENABLE_DEPRECATION_WARNINGS
+#endif
 
 	/** True if transitions are allowed to be requested directly instead of buffering. */
 	bool bAllowDirectTransitions = false;
diff --git a/Engine/Plugins/Runtime/StateTree/Source/StateTreeModule/Public/StateTreeExecutionTypes.h b/Engine/Plugins/Runtime/StateTree/Source/StateTreeModule/Public/StateTreeExecutionTypes.h
--- a/Engine/Plugins/Runtime/StateTree/Source/StateTreeModule/Public/StateTreeExecutionTypes.h
+++ b/Engine/Plugins/Runtime/StateTree/Source/StateTreeModule/Public/StateTreeExecutionTypes.h
@@ -441,33 +441,6 @@ struct FStateTreeActiveStates
 		return UE::StateTree::FActiveStateID::Invalid;
 	}
 
-	/** Returns true of the array contains specified state. */
-	bool Contains(const FStateTreeStateHandle StateHandle) const
-	{
-		for (const FStateTreeStateHandle& Handle : *this)
-		{
-			if (Handle == StateHandle)
-			{
-				return true;
-			}
-		}
-		return false;
-	}
-
-	/** Returns true of the array contains specified state within MaxNumStatesToCheck states. */
-	bool Contains(const FStateTreeStateHandle StateHandle, const uint8 MaxNumStatesToCheck) const
-	{
-		const int32 Num = (int32)FMath::Min(NumStates, MaxNumStatesToCheck);
-		for (int32 Index = 0; Index < Num; Index++)
-		{
-			if (States[Index] == StateHandle)
-			{
-				return true;
-			}
-		}
-		return false;
-	}
-
 	/** Returns index of a state, searching in reverse order. */
 	int32 IndexOfReverse(const FStateTreeStateHandle StateHandle) const
 	{
@@ -494,6 +467,33 @@ struct FStateTreeActiveStates
 		return INDEX_NONE;
 	}
 	
+	/** Returns true of the array contains specified state. */
+	bool Contains(const FStateTreeStateHandle StateHandle) const
+	{
+		for (int32 Index = 0; Index < NumStates; ++Index)
+		{
+			if (States[Index] == StateHandle)
+			{
+				return true;
+			}
+		}
+		return false;
+	}
+
+	/** Returns true of the array contains specified state within MaxNumStatesToCheck states. */
+	bool Contains(const FStateTreeStateHandle StateHandle, const uint8 MaxNumStatesToCheck) const
+	{
+		const int32 Num = (int32)FMath::Min(NumStates, MaxNumStatesToCheck);
+		for (int32 Index = 0; Index < Num; ++Index)
+		{
+			if (States[Index] == StateHandle)
+			{
+				return true;
+			}
+		}
+		return false;
+	}
+
 	/** Returns true if the state id is inside the container. */
 	bool Contains(const UE::StateTree::FActiveStateID StateId) const
 	{
@@ -903,6 +903,20 @@ private:
 	uint32 bContainsUnboundListeners : 1 = false;
 };
 
+namespace UE::StateTree
+{
+	/** Helper that identifies an execution frame with the root state handle and its tree asset. */
+	struct FExecutionFrameKey
+	{
+		FExecutionFrameKey(TNotNull<const UStateTree*> InStateTree, FStateTreeStateHandle InRootState)
+			: StateTree(InStateTree)
+			, RootState(InRootState)
+		{
+		}
+		TNotNull<const UStateTree*> StateTree;
+		FStateTreeStateHandle RootState;
+	};
+}
 
 /** Describes an active branch of a State Tree. */
 USTRUCT(BlueprintType)
@@ -914,13 +928,31 @@ struct FStateTreeExecutionFrame
 	UE_DEPRECATED(5.6, "The recorded frame doesn't have all the needed information to properly form an FStateTreeExecutionFrame.")
 	UE_API FStateTreeExecutionFrame(const FRecordedStateTreeExecutionFrame& RecordedExecutionFrame);
 
-	//@todo Refactor IsSameFrame to use FrameID
+	UE_DEPRECATED(5.7, "Use the other version of IsSameRoot.")
 	bool IsSameFrame(const FStateTreeExecutionFrame& OtherFrame) const
 	{
 		return StateTree == OtherFrame.StateTree && RootState == OtherFrame.RootState;
 	}
+
+	/** @return wherever the frame uses the state tree and root state. */
+	bool HasSameRoot(const FStateTreeExecutionFrame& OtherFrame) const
+	{
+		return StateTree == OtherFrame.StateTree && RootState == OtherFrame.RootState;
+	}
+
+	/** @return wherever the frame uses the state tree and root state. */
+	bool HasRoot(const UE::StateTree::FExecutionFrameKey& FrameKey) const
+	{
+		return StateTree == FrameKey.StateTree && RootState == FrameKey.RootState;
+	}
+
+	/** @return wherever the frame uses the state tree and root state. */
+	bool HasRoot(TNotNull<const UStateTree*> InStateTree, FStateTreeStateHandle InRootState) const
+	{
+		return StateTree == InStateTree && RootState == InRootState;
+	}
 	
-	/** The State Tree used for ticking this frame. */
+	/** The State Tree used by the frame. */
 	UPROPERTY()
 	TObjectPtr<const UStateTree> StateTree = nullptr;
 
@@ -965,13 +997,23 @@ struct FStateTreeExecutionFrame
 	/** Handle to the global parameter data, exists in ParentFrame. */
 	UPROPERTY()
 	FStateTreeDataHandle GlobalParameterDataHandle = FStateTreeDataHandle::Invalid;
-
-	/** Number of states in ActiveStates which have instance data. Used during state selection to decide which active state data is safe to access. */
-	uint8 NumCurrentlyActiveStates = 0;
 	
 	/** If true, the global tasks of the State Tree should be handle in this frame. */
 	UPROPERTY()
 	uint8 bIsGlobalFrame : 1 = false;
+
+	/**
+	 * If true, the global tasks/evaluator received the "EnterState".
+	 * Can be sustained or added via a linked state.
+	 * Only call StateEnter when the state didn't previously receive a state enter.
+	 */
+	UPROPERTY()
+	uint8 bHaveEntered : 1 = false;
+
+#if WITH_EDITORDATA_ONLY
+	UE_DEPRECATED(5.7, "States in ActiveStates are always valid.")
+	uint8 NumCurrentlyActiveStates = 0;
+#endif
 };
 
 /** Describes the execution state of the current State Tree instance. */
@@ -1147,7 +1189,7 @@ PRAGMA_ENABLE_DEPRECATION_WARNINGS
 #endif //WITH_EDITORONLY_DATA
 };
 
-
+//@todo deprecate once ForceTransition is fixed
 /** Contains StateTree events used during State Selection for a single execution frame. */
 struct FStateTreeFrameStateSelectionEvents
 {
@@ -1180,10 +1222,12 @@ PRAGMA_ENABLE_DEPRECATION_WARNINGS
 		*this = FStateTreeTransitionResult();
 	}
 
+	//@todo deprecate once ForceTransition is fixed
 	/** States selected as result of the transition. */
 	UPROPERTY(EditDefaultsOnly, Category = "Default", BlueprintReadOnly)
 	TArray<FStateTreeExecutionFrame> NextActiveFrames;
 
+	//@todo deprecate once ForceTransition is fixed
 	/** Events used in state selection. */
 	TArray<FStateTreeFrameStateSelectionEvents> NextActiveFrameEvents;
 
diff --git a/Engine/Plugins/Runtime/StateTree/Source/StateTreeModule/Public/StateTreeSchema.h b/Engine/Plugins/Runtime/StateTree/Source/StateTreeModule/Public/StateTreeSchema.h
--- a/Engine/Plugins/Runtime/StateTree/Source/StateTreeModule/Public/StateTreeSchema.h
+++ b/Engine/Plugins/Runtime/StateTree/Source/StateTreeModule/Public/StateTreeSchema.h
@@ -12,6 +12,7 @@ struct FStateTreeLinker;
 enum class EStateTreeDataSourceType : uint8;
 enum class EStateTreeParameterDataType : uint8;
 enum class EStateTreeStateSelectionBehavior : uint8;
+enum class EStateTreeStateSelectionRules : uint32;
 
 
 /**
@@ -65,7 +66,7 @@ public:
 	 * Can be used to quickly accept all of those classes in IsClassAllowed().
 	 * @return True if the class is a StateTree item Blueprint base class.
 	 */
-	UE_API bool IsChildOfBlueprintBase(const UClass* InClass) const;
+	UE_API static bool IsChildOfBlueprintBase(const UClass* InClass);
 
 	/** @return List of context objects (UObjects or UScriptStructs) enforced by the schema. They must be provided at runtime through the execution context. */
 	virtual TConstArrayView<FStateTreeExternalDataDesc> GetContextDataDescs() const
@@ -76,6 +77,9 @@ public:
 	/** @return the global parameter type used by the schema. */
 	UE_API virtual EStateTreeParameterDataType GetGlobalParameterDataType() const;
 
+	/** @return the selection rules used by the schema. */
+	UE_API virtual EStateTreeStateSelectionRules GetStateSelectionRules() const;
+
 	/** Resolves schema references to other StateTree data. */
 	virtual bool Link(FStateTreeLinker& Linker)
 	{
diff --git a/Engine/Plugins/Runtime/StateTree/Source/StateTreeModule/Public/StateTreeTypes.h b/Engine/Plugins/Runtime/StateTree/Source/StateTreeModule/Public/StateTreeTypes.h
--- a/Engine/Plugins/Runtime/StateTree/Source/StateTreeModule/Public/StateTreeTypes.h
+++ b/Engine/Plugins/Runtime/StateTree/Source/StateTreeModule/Public/StateTreeTypes.h
@@ -90,6 +90,28 @@ enum class EStateTreeTransitionType : uint8
 	NotSet UE_DEPRECATED(all, "Use None instead."),
 };
 
+/** The rules used by the execution context for selecting states. */
+UENUM()
+enum class EStateTreeStateSelectionRules : uint32
+{
+	/** Previous (UE 5.6) rules */
+	None = 0,
+	/**
+	 * If the transition source state or the transition target state (or any child states of the source or target) is completed, then the completed states are new (recreated if needed).
+	 * When a new state is created, all the following states in the selection path are also new (recreated if needed).
+	 * Without this rule, the state will remain active even if it's completed. The state might never trigger a complete transition.
+	 */
+	CompletedTransitionStatesCreateNewStates = 1 << 0,
+	/**
+	 * The transition will fail if there is a completed state before the transition source state.
+	 * Without this rule, a transition that triggers on completion might take an extra tick to execute.
+	 */
+	CompletedStateBeforeTransitionSourceFailsTransition = 1 << 1,
+
+	Default = CompletedTransitionStatesCreateNewStates | CompletedStateBeforeTransitionSourceFailsTransition,
+};
+ENUM_CLASS_FLAGS(EStateTreeStateSelectionRules);
+
 /** Operand in an expression */
 UENUM()
 enum class EStateTreeExpressionOperand : uint8
diff --git a/Engine/Plugins/Runtime/StateTree/Source/StateTreeTestSuite/Private/StateTreeLinkedStateTest.cpp b/Engine/Plugins/Runtime/StateTree/Source/StateTreeTestSuite/Private/StateTreeLinkedStateTest.cpp
--- a/Engine/Plugins/Runtime/StateTree/Source/StateTreeTestSuite/Private/StateTreeLinkedStateTest.cpp
+++ b/Engine/Plugins/Runtime/StateTree/Source/StateTreeTestSuite/Private/StateTreeLinkedStateTest.cpp
@@ -314,19 +314,17 @@ struct FStateTreeTest_EnterAndExitLinkedAsset2 : FStateTreeTestBase
 				EStateTreeRunStatus Status = Exec.Tick(0.1f);
 				AITEST_EQUAL("Tick should complete with Running", Status, EStateTreeRunStatus::Running);
 				AITEST_TRUE(TEXT("StateTree should exit the root state but not the global"),
-					Exec.Expect(GlobalTask2.GetName(), TickStr)
+					Exec.Expect(Task0.GetName(), TickStr)
+					.Then(GlobalTask2.GetName(), TickStr)
 					.Then(Task2.GetName(), TickStr)
 					.Then(Task2.GetName(), StateCompletedStateStr)
-					//@TODO the state completed. It should be a new state and not sustained.
-					//.Then(Task2.GetName(), ChangedExitStateStr)
-					//@TODO call enter/exit of globals.
-					//.Then(GlobalTask2.GetName(), SustainedExitStateStr)
+					.Then(Task0.GetName(), StateCompletedStateStr)
+					.Then(Task2.GetName(), ChangedExitStateStr)
+					.Then(GlobalTask2.GetName(), SustainedExitStateStr)
 					.Then(Task0.GetName(), SustainedExitStateStr)
 					.Then(Task0.GetName(), SustainedEnterStateStr)
-					//@TODO call enter/exit of globals.
-					//.Then(GlobalTask2.GetName(), SustainedEnterStateStr)
-					//@TODO the state completed. It should be a new state and not sustained.
-					//.Then(Task2.GetName(), ChangedEnterStateStr)
+					.Then(GlobalTask2.GetName(), SustainedEnterStateStr)
+					.Then(Task2.GetName(), ChangedEnterStateStr)
 				);
 				AITEST_FALSE("StateTree should not tick Task1", Exec.Expect(Task1.GetName(), TickStr));
 				AITEST_FALSE("StateTree should not exit Task1", Exec.Expect(Task1.GetName(), ExitStateStr));
@@ -354,17 +352,13 @@ struct FStateTreeTest_EnterAndExitLinkedAsset2 : FStateTreeTestBase
 				AITEST_EQUAL("Tick should complete with Running", Status, EStateTreeRunStatus::Running);
 				AITEST_TRUE(TEXT("StateTree should exit the root state and the global"),
 					Exec.Expect(GlobalTask2.GetName(), TickStr)
-					//@TODO The tree completed. It is a new state.
-					//.Then(Task2.GetName(), ChangedExitStateStr)
-					//@TODO call enter/exit of globals.
-					//.Then(GlobalTask2.GetName(), ChangedExitStateStr)
+					.Then(Task2.GetName(), ChangedExitStateStr)
+					.Then(GlobalTask2.GetName(), ChangedExitStateStr)
 					.Then(Task0.GetName(), SustainedExitStateStr)
 					.Then(Task0.GetName(), SustainedEnterStateStr)
-					//@TODO call enter/exit of globals.
-					//.Then(GlobalTask2.GetName(), ChangedEnterStateStr)
-					//@TODO The tree completed. It is a new state.
-					//.Then(Task2.GetName(), EnterStateStr)
+					.Then(Task2.GetName(), EnterStateStr)
 				);
+				AITEST_TRUE(TEXT("StateTree should enter a new global task2"), Exec.Expect(GlobalTask2.GetName(), ChangedEnterStateStr));
 				AITEST_FALSE("StateTree should not tick Tasks2", Exec.Expect(Task2.GetName(), TickStr));
 				AITEST_FALSE("StateTree should not tick Tasks1", Exec.Expect(Task1.GetName(), TickStr));
 				AITEST_FALSE("StateTree should not exit Task1", Exec.Expect(Task1.GetName(), EnterStateStr));
@@ -490,10 +484,10 @@ struct FStateTreeTest_MultipleSameLinkedAsset : FStateTreeTestBase
 				C3.Parameters.Parameters.SetValueInt32("Value", 333);
 			}
 			{
-				UStateTreeState& StateD = Root.AddChildState("Tree1StateC", EStateTreeStateType::State);
-				TStateTreeEditorNode<FTestTask_PrintValue>& Task = StateD.AddTask<FTestTask_PrintValue>("Tree1StateBTaskA");
+				UStateTreeState& StateC = Root.AddChildState("Tree1StateC", EStateTreeStateType::State);
+				TStateTreeEditorNode<FTestTask_PrintValue>& Task = StateC.AddTask<FTestTask_PrintValue>("Tree1StateCTaskA");
 				Task.GetInstanceData().Value = 3;
-				FStateTreeTransition& Transition = StateD.AddTransition(EStateTreeTransitionTrigger::OnTick, EStateTreeTransitionType::GotoState, &Root);
+				FStateTreeTransition& Transition = StateC.AddTransition(EStateTreeTransitionTrigger::OnTick, EStateTreeTransitionType::GotoState, &Root);
 				Transition.bDelayTransition = true;
 				Transition.DelayDuration = 1.0;
 			}
@@ -518,6 +512,7 @@ struct FStateTreeTest_MultipleSameLinkedAsset : FStateTreeTestBase
 			{
 				EStateTreeRunStatus Status = Exec.Start();
 				AITEST_EQUAL(TEXT("Start should complete with Running"), Status, EStateTreeRunStatus::Running);
+				AITEST_TRUE(TEXT("Start should be in the correct states"), Exec.ExpectInActiveStates("Tree1StateRoot", "Tree1StateA"));
 				FTestStateTreeExecutionContext::FLogOrder LogOrder = Exec.Expect("Tree1GlobalTaskA", TEXT("EnterState99"));
 				AITEST_TRUE(TEXT("Start should enter Tree1GlobalTaskA"), LogOrder);
 				LogOrder = LogOrder.Then("Tree1StateRootTaskA", TEXT("EnterState88"));
@@ -529,6 +524,7 @@ struct FStateTreeTest_MultipleSameLinkedAsset : FStateTreeTestBase
 			{
 				EStateTreeRunStatus Status = Exec.Tick(1.5f); // over tick, should trigger
 				AITEST_EQUAL(TEXT("1st Tick should complete with Running"), Status, EStateTreeRunStatus::Running);
+				AITEST_TRUE(TEXT("1st Tick should be in the correct states"), Exec.ExpectInActiveStates("Tree1StateRoot", "Tree1StateA"));
 				FTestStateTreeExecutionContext::FLogOrder LogOrder = Exec.Expect("Tree1GlobalTaskA", TEXT("Tick99"));
 				AITEST_TRUE(TEXT("1st should tick tasks Tree1GlobalTaskA"), LogOrder);
 				LogOrder = LogOrder.Then("Tree1StateRootTaskA", TEXT("Tick88"));
@@ -540,6 +536,7 @@ struct FStateTreeTest_MultipleSameLinkedAsset : FStateTreeTestBase
 			{
 				EStateTreeRunStatus Status = Exec.Tick(1.0f);
 				AITEST_EQUAL(TEXT("2nd Tick should complete with Running"), Status, EStateTreeRunStatus::Running);
+				AITEST_TRUE(TEXT("2nd Tick should be in the correct states"), Exec.ExpectInActiveStates("Tree1StateRoot", "Tree1StateB"));
 				FTestStateTreeExecutionContext::FLogOrder LogOrder = Exec.Expect("Tree1GlobalTaskA", TEXT("Tick99"));
 				AITEST_TRUE(TEXT("2nd Tick should tick Tree1GlobalTaskA"), LogOrder);
 				LogOrder = LogOrder.Then("Tree1StateRootTaskA", TEXT("Tick88"));
@@ -555,6 +552,7 @@ struct FStateTreeTest_MultipleSameLinkedAsset : FStateTreeTestBase
 			{
 				EStateTreeRunStatus Status = Exec.Tick(1.0f);
 				AITEST_EQUAL(TEXT("3rd Tick should complete with Running"), Status, EStateTreeRunStatus::Running);
+				AITEST_TRUE(TEXT("3rd Tick should be in the correct states"), Exec.ExpectInActiveStates("Tree1StateRoot", "Tree1StateB"));
 				FTestStateTreeExecutionContext::FLogOrder LogOrder = Exec.Expect("Tree1GlobalTaskA", TEXT("Tick99"));
 				AITEST_TRUE(TEXT("3rd should tick Tree1GlobalTaskA"), LogOrder);
 				LogOrder = LogOrder.Then("Tree1StateRootTaskA", TEXT("Tick88"));
@@ -566,6 +564,7 @@ struct FStateTreeTest_MultipleSameLinkedAsset : FStateTreeTestBase
 			{
 				EStateTreeRunStatus Status = Exec.Tick(1.0f);
 				AITEST_EQUAL(TEXT("4th Tick should complete with Running"), Status, EStateTreeRunStatus::Running);
+				AITEST_TRUE(TEXT("4th Tick should be in the correct states"), Exec.ExpectInActiveStates("Tree1StateRoot", "Tree1StateLinkedTreeA", "Tree2StateRoot", "Tree2StateA"));
 				FTestStateTreeExecutionContext::FLogOrder LogOrder = Exec.Expect("Tree1GlobalTaskA", TEXT("Tick99"));
 				AITEST_TRUE(TEXT("4th Tick should tick Tree1GlobalTaskA"), LogOrder);
 				LogOrder = LogOrder.Then("Tree1StateRootTaskA", TEXT("Tick88"));
@@ -574,9 +573,6 @@ struct FStateTreeTest_MultipleSameLinkedAsset : FStateTreeTestBase
 				AITEST_TRUE(TEXT("4th Tick should tick Tree1StateBTaskA"), LogOrder);
 				LogOrder = LogOrder.Then("Tree1StateBTaskA", TEXT("ExitState2"));
 				AITEST_TRUE(TEXT("4th Tick should exit Tree1StateBTaskA"), LogOrder);
-				//@todo: bug: order is wrong
-				//LogOrder = LogOrder.Then("Tree2GlobalTaskA", TEXT("EnterState111"))
-				//AITEST_TRUE(TEXT("4th Tick should enter Tree2GlobalTaskA"), LogOrder);
 				LogOrder = LogOrder.Then("Tree2StateRootTaskA", TEXT("EnterState111"));
 				AITEST_TRUE(TEXT("4th Tick should enter Tree2StateRootTaskA"), LogOrder);
 				LogOrder = LogOrder.Then("Tree2StateATaskA", TEXT("EnterState111"));
@@ -587,6 +583,7 @@ struct FStateTreeTest_MultipleSameLinkedAsset : FStateTreeTestBase
 			{
 				EStateTreeRunStatus Status = Exec.Tick(0.001f);
 				AITEST_EQUAL(TEXT("5th Tick should complete with Running"), Status, EStateTreeRunStatus::Running);
+				AITEST_TRUE(TEXT("5th Tick should be in the correct states"), Exec.ExpectInActiveStates("Tree1StateRoot", "Tree1StateLinkedTreeA", "Tree2StateRoot", "Tree2StateA"));
 				if (bTickGlobalNodesWithHierarchy)
 				{
 					FTestStateTreeExecutionContext::FLogOrder LogOrder = Exec.Expect("Tree1GlobalTaskA", TEXT("Tick99"));
@@ -618,6 +615,7 @@ struct FStateTreeTest_MultipleSameLinkedAsset : FStateTreeTestBase
 			{
 				EStateTreeRunStatus Status = Exec.Tick(1.0f);
 				AITEST_EQUAL(TEXT("6th Tick should complete with Running"), Status, EStateTreeRunStatus::Running);
+				AITEST_TRUE(TEXT("6th Tick should be in the correct states"), Exec.ExpectInActiveStates("Tree1StateRoot", "Tree1StateLinkedTreeB", "Tree2StateRoot", "Tree2StateA"));
 				if (bTickGlobalNodesWithHierarchy)
 				{
 					FTestStateTreeExecutionContext::FLogOrder LogOrder = Exec.Expect("Tree1GlobalTaskA", TEXT("Tick99"));
@@ -634,16 +632,13 @@ struct FStateTreeTest_MultipleSameLinkedAsset : FStateTreeTestBase
 					AITEST_TRUE(TEXT("6th Tick should exit Tree2StateATaskA"), LogOrder);
 					LogOrder = LogOrder.Then("Tree2StateRootTaskA", TEXT("ExitState111"));
 					AITEST_TRUE(TEXT("6th Tick should exit Tree2StateRootTaskA"), LogOrder);
-					//@todo: bug: The exit state is not called on the global, it's a new frame.
-					//LogOrder = LogOrder.Then("Tree2GlobalTaskA", TEXT("ExitState111"));
-					//AITEST_TRUE(TEXT("6th Tick should exit Tree2GlobalTaskA"), LogOrder);
-					//@todo: bug: The enter state is not called on the global, it's a new frame.
-					//LogOrder = LogOrder.Then("Tree2GlobalTaskA", TEXT("EnterState222"));
-					//AITEST_TRUE(TEXT("6th Tick should enter Tree2GlobalTaskA"), LogOrder);
+					LogOrder = LogOrder.Then("Tree2GlobalTaskA", TEXT("ExitState111"));
+					AITEST_TRUE(TEXT("6th Tick should exit Tree2GlobalTaskA"), LogOrder);
 					LogOrder = LogOrder.Then("Tree2StateRootTaskA", TEXT("EnterState222"));
 					AITEST_TRUE(TEXT("6th Tick should enter Tree2StateRootTaskA"), LogOrder);
 					LogOrder = LogOrder.Then("Tree2StateATaskA", TEXT("EnterState222"));
 					AITEST_TRUE(TEXT("6th Tick should enter Tree2StateATaskA"), LogOrder);
+					AITEST_TRUE(TEXT("6th Tick should enter Tree2GlobalTaskA"), Exec.Expect("Tree2GlobalTaskA", TEXT("EnterState222")));
 				}
 				else
 				{
@@ -661,16 +656,13 @@ struct FStateTreeTest_MultipleSameLinkedAsset : FStateTreeTestBase
 					AITEST_TRUE(TEXT("6th Tick should exit Tree2StateATaskA"), LogOrder);
 					LogOrder = LogOrder.Then("Tree2StateRootTaskA", TEXT("ExitState111"));
 					AITEST_TRUE(TEXT("6th Tick should exit Tree2StateRootTaskA"), LogOrder);
-					//@todo: bug: The exit state is not called on the global, it's a new frame.
-					//LogOrder = LogOrder.Then("Tree2GlobalTaskA", TEXT("ExitState111"));
-					//AITEST_TRUE(TEXT("6th Tick should exit Tree2GlobalTaskA"), LogOrder);
-					//@todo: bug: The enter state is not called on the global, it's a new frame.
-					//LogOrder = LogOrder.Then("Tree2GlobalTaskA", TEXT("EnterState222"));
-					//AITEST_TRUE(TEXT("6th Tick should enter Tree2GlobalTaskA"), LogOrder);
+					LogOrder = LogOrder.Then("Tree2GlobalTaskA", TEXT("ExitState111"));
+					AITEST_TRUE(TEXT("6th Tick should exit Tree2GlobalTaskA"), LogOrder);
 					LogOrder = LogOrder.Then("Tree2StateRootTaskA", TEXT("EnterState222"));
 					AITEST_TRUE(TEXT("6th Tick should enter Tree2StateRootTaskA"), LogOrder);
 					LogOrder = LogOrder.Then("Tree2StateATaskA", TEXT("EnterState222"));
 					AITEST_TRUE(TEXT("6th Tick should enter Tree2StateATaskA"), LogOrder);
+					AITEST_TRUE(TEXT("6th Tick should enter Tree2GlobalTaskA"), Exec.Expect("Tree2GlobalTaskA", TEXT("EnterState222")));
 				}
 				Exec.LogClear();
 			}
diff --git a/Engine/Plugins/Runtime/StateTree/Source/StateTreeTestSuite/Private/StateTreeStatePathTest.cpp b/Engine/Plugins/Runtime/StateTree/Source/StateTreeTestSuite/Private/StateTreeStatePathTest.cpp
--- a/Engine/Plugins/Runtime/StateTree/Source/StateTreeTestSuite/Private/StateTreeStatePathTest.cpp
+++ b/Engine/Plugins/Runtime/StateTree/Source/StateTreeTestSuite/Private/StateTreeStatePathTest.cpp
@@ -349,8 +349,7 @@ struct FStateTreeTest_StatePath_LinkStates : FStateTreeTestBase
 			const UE::StateTree::FActiveStatePath ExecPath = InstanceData.GetExecutionState()->GetActiveStatePath();
 			AITEST_TRUE(TEXT("Has the correct number of path elements"), ExecPath.Num() == 9);
 			AITEST_TRUE(TEXT("Has the correct number of active states"), InstanceData.GetExecutionState()->ActiveFrames.Num() == 4);
-			//@todo Temporary reverted the frameID fixed until TransitionID
-			//AITEST_TRUE(TEXT("The last frame changed"), TestFrameId(3));
+			AITEST_TRUE(TEXT("The last frame changed"), TestFrameId(3));
 			Exec.LogClear();
 		}
 		{
@@ -386,8 +385,7 @@ struct FStateTreeTest_StatePath_LinkStates : FStateTreeTestBase
 			UE::StateTree::FActiveStatePath ExecPath = InstanceData.GetExecutionState()->GetActiveStatePath();
 			AITEST_TRUE(TEXT("Has the correct number of path elements"), ExecPath.Num() == 9);
 			AITEST_TRUE(TEXT("Has the correct number of active states"), InstanceData.GetExecutionState()->ActiveFrames.Num() == 4);
-			//@todo Temporary reverted the frameID fixed until TransitionID
-			//AITEST_TRUE(TEXT("All frames are the same"), TestFrameId(2));
+			AITEST_TRUE(TEXT("All frames are the same"), TestFrameId(2));
 			Exec.LogClear();
 		}
 		{
@@ -402,8 +400,7 @@ struct FStateTreeTest_StatePath_LinkStates : FStateTreeTestBase
 			Status = Exec.Tick(1.0f);
 			AITEST_EQUAL(TEXT("Start should complete with Running"), Status, EStateTreeRunStatus::Running);
 			AITEST_TRUE(TEXT("Should be in the correct state"), Exec.ExpectInActiveStates("Tree1Root", "Tree1RootState1", "Tree1State1StateSub1A", "Tree1Sub1", "Tree1Sub1StateLinkTree2B", "Tree2StateRoot", "Tree2RootStateSub1B", "Tree2Sub1", "Tree2Sub1State1"));
-			//@todo Temporary reverted the frameID fixed until TransitionID
-			//AITEST_TRUE(TEXT("All frames are the same"), TestFrameId(3));
+			AITEST_TRUE(TEXT("All frames are the same"), TestFrameId(3));
 			Exec.LogClear();
 
 			PreTestFrameId();
diff --git a/Engine/Plugins/Runtime/StateTree/Source/StateTreeTestSuite/Private/StateTreeTest.cpp b/Engine/Plugins/Runtime/StateTree/Source/StateTreeTestSuite/Private/StateTreeTest.cpp
--- a/Engine/Plugins/Runtime/StateTree/Source/StateTreeTestSuite/Private/StateTreeTest.cpp
+++ b/Engine/Plugins/Runtime/StateTree/Source/StateTreeTestSuite/Private/StateTreeTest.cpp
@@ -1785,15 +1785,14 @@ struct FStateTreeTest_FinishTasks : FStateTreeTestBase
 				AITEST_TRUE(TEXT("StateTree should be running"), Status == EStateTreeRunStatus::Running);
 				Exec.LogClear();
 
-				// One FinishTask, it should loop to StateA
+				// On FinishTask, go to StateB but the condition will fail. Reselect a new StateA.
 				Status = Exec.Tick(0.1f);
 				AITEST_TRUE(TEXT("StateTree Active States should be in RootA/StateA"),
 					Exec.ExpectInActiveStates(RootA.Name, StateA.Name));
 				AITEST_TRUE(TEXT("StateTree should be running"), Status == EStateTreeRunStatus::Running);
 				AITEST_TRUE(TEXT("Expect the output tasks"), Exec.Expect("StateATaskA", TEXT("Tick101"))
 					.Then("StateATaskA", TEXT("Exitstate102"))
-					//@TODO: instance data should be reset
-					//.Then("StateATaskA", TEXT("EnterState101"))
+					.Then("StateATaskA", TEXT("EnterState101"))
 					);
 				Exec.LogClear();
 
@@ -1831,15 +1830,14 @@ struct FStateTreeTest_FinishTasks : FStateTreeTestBase
 				AITEST_TRUE(TEXT("StateTree should be running"), Status == EStateTreeRunStatus::Running);
 				Exec.LogClear();
 
-				// One FinishTask, it should loop to StateA
+				// On FinishTask, go to StateB but the condition will fail. Reselect a new StateA.
 				Status = Exec.Tick(0.1f);
 				AITEST_TRUE(TEXT("StateTree Active States should be in RootA/StateA"),
 					Exec.ExpectInActiveStates(RootA.Name, StateA.Name));
 				AITEST_TRUE(TEXT("StateTree should be running"), Status == EStateTreeRunStatus::Running);
 				AITEST_TRUE(TEXT("Expect the output tasks"), Exec.Expect("StateATaskA", TEXT("Tick101"))
 					.Then("StateATaskA", TEXT("Exitstate102"))
-					//@TODO: instance data should be reset
-					//.Then("StateATaskA", TEXT("EnterState101"))
+					.Then("StateATaskA", TEXT("EnterState101"))
 				);
 				Exec.LogClear();
 
@@ -1889,8 +1887,7 @@ struct FStateTreeTest_FinishTasks : FStateTreeTestBase
 				AITEST_TRUE(TEXT("StateTree should be running"), Status == EStateTreeRunStatus::Running);
 				AITEST_TRUE(TEXT("Expect the output tasks"), Exec.Expect("StateATaskA", TEXT("Tick101"))
 					.Then("StateATaskA", TEXT("ExitState102"))
-					//@TODO: instance data should be reset
-					//.Then("StateATaskA", TEXT("EnterState101"))
+					.Then("StateATaskA", TEXT("EnterState101"))
 				);
 				AITEST_FALSE(TEXT("Expect the output tasks"), Exec.Expect("StateATaskA", TEXT("ExitState103")));
 				Exec.LogClear();
@@ -1904,7 +1901,6 @@ struct FStateTreeTest_FinishTasks : FStateTreeTestBase
 };
 IMPLEMENT_STATE_TREE_INSTANT_TEST(FStateTreeTest_FinishTasks, "System.StateTree.FinishTask");
 
-
 // Test nested tree overrides
 struct FStateTreeTest_NestedOverride : FStateTreeTestBase
 {
@@ -2388,15 +2384,13 @@ struct FStateTreeTest_Reentrant : FStateTreeTestBase
 				AITEST_TRUE("Expect the output tasks", Exec.Expect("Tree1State1Task1", TEXT("Tick101"))
 					.Then("Tree1State2Task1", TEXT("Tick201"))
 					.Then("Tree1State2Task1", TEXT("ExitState201"))
-					//@TODO Should select a new state because it's completed.
-					//.Then("Tree1State2Task1", TEXT("ExitState=Changed"))
+					.Then("Tree1State2Task1", TEXT("ExitState=Changed"))
 					.Then("Tree1State1Task1", TEXT("ExitState101"))
 					.Then("Tree1State1Task1", TEXT("ExitState=Sustained"))
 					.Then("Tree1State1Task1", TEXT("EnterState101"))
 					.Then("Tree1State1Task1", TEXT("EnterState=Sustained"))
 					.Then("Tree1State2Task1", TEXT("EnterState201"))
-					//@TODO Should select a new state because it's completed.
-					//.Then("Tree1State2Task1", TEXT("EnterState=Changed"))
+					.Then("Tree1State2Task1", TEXT("EnterState=Changed"))
 				);
 
 
@@ -2405,10 +2399,8 @@ struct FStateTreeTest_Reentrant : FStateTreeTestBase
 				AITEST_TRUE(TEXT("Same length"), AfterStatePath.Num() == BeforeStatePath.Num());
 				AITEST_TRUE(TEXT("Length is 3"), AfterStatePath.Num() == 3);
 				AITEST_TRUE(TEXT("Element 1 is the same"), AfterStatePath.GetView()[0] == BeforeStatePath.GetView()[0]);
-				//@TODO Should select a new state because it's completed.
-				//AITEST_TRUE(TEXT("Element 2 is the same"), AfterStatePath.GetView()[1] == BeforeStatePath.GetView()[1]);
-				//@TODO Should select a new state because it's completed.
-				//AITEST_TRUE(TEXT("Element 3 is different"), AfterStatePath.GetView()[2] != BeforeStatePath.GetView()[2]);
+				AITEST_TRUE(TEXT("Element 2 is the same"), AfterStatePath.GetView()[1] == BeforeStatePath.GetView()[1]);
+				AITEST_TRUE(TEXT("Element 3 is different"), AfterStatePath.GetView()[2] != BeforeStatePath.GetView()[2]);
 
 				Exec.LogClear();
 			}
diff --git a/Engine/Plugins/Runtime/StateTree/Source/StateTreeTestSuite/Private/StateTreeTest.h b/Engine/Plugins/Runtime/StateTree/Source/StateTreeTestSuite/Private/StateTreeTest.h
--- a/Engine/Plugins/Runtime/StateTree/Source/StateTreeTestSuite/Private/StateTreeTest.h
+++ b/Engine/Plugins/Runtime/StateTree/Source/StateTreeTestSuite/Private/StateTreeTest.h
@@ -3,6 +3,7 @@
 #pragma once
 
 #include "StateTreeSchema.h"
+#include "StateTreeTypes.h"
 #include "StateTreeTest.generated.h"
 
 UCLASS(HideDropdown)
@@ -10,6 +11,7 @@ class UStateTreeTestSchema : public UStateTreeSchema
 {
 	GENERATED_BODY()
 
+public:
 	virtual bool IsStructAllowed(const UScriptStruct* InScriptStruct) const override
 	{
 		return true;
@@ -22,6 +24,18 @@ class UStateTreeTestSchema : public UStateTreeSchema
 	{
 		return true;
 	}
+	virtual EStateTreeStateSelectionRules GetStateSelectionRules() const
+	{
+		return DefaultRules;
+	}
+	void SetStateSelectionRules(EStateTreeStateSelectionRules Rules)
+	{
+		DefaultRules = Rules;
+	}
+
+private:
+	UPROPERTY()
+	EStateTreeStateSelectionRules DefaultRules = EStateTreeStateSelectionRules::Default;
 };
 
 UCLASS(HideDropdown)
diff --git a/Engine/Plugins/Runtime/StateTree/Source/StateTreeTestSuite/Private/StateTreeTestTypes.h b/Engine/Plugins/Runtime/StateTree/Source/StateTreeTestSuite/Private/StateTreeTestTypes.h
--- a/Engine/Plugins/Runtime/StateTree/Source/StateTreeTestSuite/Private/StateTreeTestTypes.h
+++ b/Engine/Plugins/Runtime/StateTree/Source/StateTreeTestSuite/Private/StateTreeTestTypes.h
@@ -42,9 +42,8 @@ struct FStateTreeTestLog
 
 struct FTestStateTreeExecutionContext : public FStateTreeExecutionContext
 {
-	
-	FTestStateTreeExecutionContext(UObject& InOwner, const UStateTree& InStateTree, FStateTreeInstanceData& InInstanceData)
-		: FStateTreeExecutionContext(InOwner, InStateTree, InInstanceData)
+	FTestStateTreeExecutionContext(UObject& InOwner, const UStateTree& InStateTree, FStateTreeInstanceData& InInstanceData, const FOnCollectStateTreeExternalData& CollectExternalDataCallback = {}, const EStateTreeRecordTransitions RecordTransitions = EStateTreeRecordTransitions::No)
+		: FStateTreeExecutionContext(InOwner, InStateTree, InInstanceData, CollectExternalDataCallback, RecordTransitions)
 	{
 		// Handle getting pointer to the test log.
 		FStateTreeDataView TestLogView(TBaseStructure<FStateTreeTestLog>::Get(), (uint8*)&Log);
