[StateTree] Transition with ID. Use the state and frame ID to identify the state during a transition request. Add new rules for transitions for backward compatibility with the previous expected behavi

diff --git a/Engine/Plugins/Runtime/StateTree/Source/StateTreeTestSuite/Private/StateTreeTransitionTest.cpp b/Engine/Plugins/Runtime/StateTree/Source/StateTreeTestSuite/Private/StateTreeTransitionTest.cpp
--- a/Engine/Plugins/Runtime/StateTree/Source/StateTreeTestSuite/Private/StateTreeTransitionTest.cpp
+++ b/Engine/Plugins/Runtime/StateTree/Source/StateTreeTestSuite/Private/StateTreeTransitionTest.cpp
@@ -776,7 +776,7 @@ struct FStateTreeTest_InfiniteLoop : FStateTreeTestBase
 		const TCHAR* EnterStateStr = TEXT("EnterState");
 		const TCHAR* ExitStateStr = TEXT("ExitState");
 
-		GetTestRunner().AddExpectedError(TEXT("Loop detected when trying to select state"), EAutomationExpectedErrorFlags::Contains, 1);
+		GetTestRunner().AddExpectedMessage(TEXT("Loop detected when trying to select state"), ELogVerbosity::Warning, EAutomationExpectedErrorFlags::Contains, 1);
 		GetTestRunner().AddExpectedError(TEXT("Failed to select initial state"), EAutomationExpectedErrorFlags::Contains, 1);
 		
 		Status = Exec.Start();
@@ -1134,394 +1134,558 @@ struct FStateTreeTest_RequestTransition : FStateTreeTestBase
 			AITEST_TRUE("StateTree should init", bInitSucceeded);
 		}
 
+		constexpr int32 MaxRules = 2;
+		auto MakeStateSelectionRule = [](int32 Index)
+			{
+				EStateTreeStateSelectionRules Rule = EStateTreeStateSelectionRules::None;
+				if ((Index % 2) == 1)
+				{
+					Rule |= EStateTreeStateSelectionRules::CompletedTransitionStatesCreateNewStates;
+					Rule |= EStateTreeStateSelectionRules::CompletedStateBeforeTransitionSourceFailsTransition;
+				}
+				return Rule;
+			};
+
 		auto ResetInstanceData = [this, &StateTree, &InstanceData]()
+			{
+				{
+					FTestStateTreeExecutionContext Exec(StateTree, StateTree, InstanceData);
+					EStateTreeRunStatus Status = Exec.Stop();
+					AITEST_EQUAL(TEXT("Should stop"), Status, EStateTreeRunStatus::Stopped);
+				}
+				{
+					FTestStateTreeExecutionContext Exec(StateTree, StateTree, InstanceData);
+					const EStateTreeRunStatus Status = Exec.Start();
+					AITEST_EQUAL(TEXT("State should complete with Running"), Status, EStateTreeRunStatus::Running);
+					AITEST_TRUE(TEXT("Start should"), Exec.ExpectInActiveStates("StateA", "StateB", "StateC", "StateD"));
+				}
+				return true;
+			};
+
+		const TCHAR* EnterStateChangedStr = TEXT("EnterState=Changed");
+		const TCHAR* EnterStateSustainedStr = TEXT("EnterState=Sustained");
+		const TCHAR* ExitStateChangedStr = TEXT("ExitState=Changed");;
+		const TCHAR* ExitStateSustainedStr = TEXT("ExitState=Sustained");
+
+		for (int32 Index = 0; Index < MaxRules; ++Index)
 		{
+			const EStateTreeStateSelectionRules StateSelectionRules = MakeStateSelectionRule(Index);
+			InstanceData = FStateTreeInstanceData();
+			if (StateTree.GetStateSelectionRules() != StateSelectionRules)
 			{
-				FTestStateTreeExecutionContext Exec(StateTree, StateTree, InstanceData);
-				EStateTreeRunStatus Status = Exec.Stop();
-				AITEST_EQUAL(TEXT("Should stop"), Status, EStateTreeRunStatus::Stopped);
+				StateTree.ResetCompiled();
+				UStateTreeEditorData* EditorData = CastChecked<UStateTreeEditorData>(StateTree.EditorData);
+				UStateTreeTestSchema* Schema = CastChecked<UStateTreeTestSchema>(EditorData->Schema);
+				Schema->SetStateSelectionRules(StateSelectionRules);
+
+				FStateTreeCompiler Compiler(Log);
+				const bool bResult = Compiler.Compile(StateTree);
+				AITEST_TRUE("StateTree should get compiled", bResult);
 			}
-			InstanceData = FStateTreeInstanceData();
+
+			// Normal Start and tick. Make sure we are in ABCD before testing the transitions
 			{
 				FTestStateTreeExecutionContext Exec(StateTree, StateTree, InstanceData);
 				const EStateTreeRunStatus Status = Exec.Start();
-				AITEST_EQUAL(TEXT("State should complete with Running"), Status, EStateTreeRunStatus::Running);
-				AITEST_TRUE(TEXT("Start should"), Exec.ExpectInActiveStates("StateA", "StateB", "StateC", "StateD"));
+				AITEST_EQUAL("Start should complete with Running", Status, EStateTreeRunStatus::Running);
+				AITEST_TRUE("Start should enter Global tasks", Exec.ExpectInActiveStates("StateA", "StateB", "StateC", "StateD"));
+				FTestStateTreeExecutionContext::FLogOrder LogOrder = Exec.Expect("StateATask", EnterStateChangedStr);
+				AITEST_TRUE(TEXT("Start StateA"), LogOrder);
+				LogOrder = LogOrder.Then("StateBTask", EnterStateChangedStr);
+				AITEST_TRUE(TEXT("Start StateB"), LogOrder);
+				LogOrder = LogOrder.Then("StateCTask", EnterStateChangedStr);
+				AITEST_TRUE(TEXT("Start StateC"), LogOrder);
+				LogOrder = LogOrder.Then("StateDTask", EnterStateChangedStr);
+				AITEST_TRUE(TEXT("Start StateD"), LogOrder);
+				AITEST_FALSE(TEXT("Tick StateA"), Exec.Expect("StateATask", TEXT("Tick0")));
+				AITEST_FALSE(TEXT("Tick StateB"), Exec.Expect("StateBTask", TEXT("Tick0")));
+				AITEST_FALSE(TEXT("Tick StateC"), Exec.Expect("StateCTask", TEXT("Tick0")));
+				AITEST_FALSE(TEXT("Tick StateD"), Exec.Expect("StateDTask", TEXT("Tick0")));
+				AITEST_FALSE(TEXT("Start StateA"), Exec.Expect("StateATask", TEXT("ExitState0")));
+				AITEST_FALSE(TEXT("Start StateB"), Exec.Expect("StateBTask", TEXT("ExitState0")));
+				AITEST_FALSE(TEXT("Start StateC"), Exec.Expect("StateCTask", TEXT("ExitState0")));
+				AITEST_FALSE(TEXT("Start StateD"), Exec.Expect("StateDTask", TEXT("ExitState0")));
+				Exec.LogClear();
 			}
-			return true;
-		};
-		{
-			FTestStateTreeExecutionContext Exec(StateTree, StateTree, InstanceData);
-			const EStateTreeRunStatus Status = Exec.Start();
-			AITEST_EQUAL("Start should complete with Running", Status, EStateTreeRunStatus::Running);
-			AITEST_TRUE("Start should enter Global tasks", Exec.ExpectInActiveStates("StateA", "StateB", "StateC", "StateD"));
-			FTestStateTreeExecutionContext::FLogOrder LogOrder = Exec.Expect("StateATask", TEXT("EnterState=Changed"));
-			AITEST_TRUE(TEXT("Start StateA"), LogOrder);
-			LogOrder = LogOrder.Then("StateBTask", TEXT("EnterState=Changed"));
-			AITEST_TRUE(TEXT("Start StateB"), LogOrder);
-			LogOrder = LogOrder.Then("StateCTask", TEXT("EnterState=Changed"));
-			AITEST_TRUE(TEXT("Start StateC"), LogOrder);
-			LogOrder = LogOrder.Then("StateDTask", TEXT("EnterState=Changed"));
-			AITEST_TRUE(TEXT("Start StateD"), LogOrder);
-			AITEST_FALSE(TEXT("Tick StateA"), Exec.Expect("StateATask", TEXT("Tick0")));
-			AITEST_FALSE(TEXT("Tick StateB"), Exec.Expect("StateBTask", TEXT("Tick0")));
-			AITEST_FALSE(TEXT("Tick StateC"), Exec.Expect("StateCTask", TEXT("Tick0")));
-			AITEST_FALSE(TEXT("Tick StateD"), Exec.Expect("StateDTask", TEXT("Tick0")));
-			AITEST_FALSE(TEXT("Start StateA"), Exec.Expect("StateATask", TEXT("ExitState0")));
-			AITEST_FALSE(TEXT("Start StateB"), Exec.Expect("StateBTask", TEXT("ExitState0")));
-			AITEST_FALSE(TEXT("Start StateC"), Exec.Expect("StateCTask", TEXT("ExitState0")));
-			AITEST_FALSE(TEXT("Start StateD"), Exec.Expect("StateDTask", TEXT("ExitState0")));
-			Exec.LogClear();
-		}
-		{
-			FTestStateTreeExecutionContext Exec(StateTree, StateTree, InstanceData);
-			const EStateTreeRunStatus Status = Exec.Tick(0.1f);
-			AITEST_EQUAL("1st Tick should complete with Running", Status, EStateTreeRunStatus::Running);
-			AITEST_TRUE("1st Tick no transition", Exec.ExpectInActiveStates("StateA", "StateB", "StateC", "StateD"));
-			FTestStateTreeExecutionContext::FLogOrder LogOrder = Exec.Expect("StateATask", TEXT("Tick0"));
-			AITEST_TRUE(TEXT("Tick StateA"), LogOrder);
-			LogOrder = LogOrder.Then("StateBTask", TEXT("Tick0"));
-			AITEST_TRUE(TEXT("Tick StateB"), LogOrder);
-			LogOrder = LogOrder.Then("StateCTask", TEXT("Tick0"));
-			AITEST_TRUE(TEXT("Tick StateC"), LogOrder);
-			LogOrder = LogOrder.Then("StateDTask", TEXT("Tick0"));
-			AITEST_TRUE(TEXT("Tick StateD"), LogOrder);
-			AITEST_FALSE(TEXT("Enter StateA"), Exec.Expect("StateATask", TEXT("EnterState0")));
-			AITEST_FALSE(TEXT("Enter StateB"), Exec.Expect("StateBTask", TEXT("EnterState0")));
-			AITEST_FALSE(TEXT("Enter StateC"), Exec.Expect("StateCTask", TEXT("EnterState0")));
-			AITEST_FALSE(TEXT("Enter StateD"), Exec.Expect("StateDTask", TEXT("EnterState0")));
-			AITEST_FALSE(TEXT("Exit StateA"), Exec.Expect("StateATask", TEXT("ExitState0")));
-			AITEST_FALSE(TEXT("Exit StateB"), Exec.Expect("StateBTask", TEXT("ExitState0")));
-			AITEST_FALSE(TEXT("Exit StateC"), Exec.Expect("StateCTask", TEXT("ExitState0")));
-			AITEST_FALSE(TEXT("Exit StateD"), Exec.Expect("StateDTask", TEXT("ExitState0")));
-			Exec.LogClear();
-		}
-		{
-			TransitionToExecute = ECustomFunctionToRun::TransitionD_To_E;
-
-			FTestStateTreeExecutionContext Exec(StateTree, StateTree, InstanceData);
-			const EStateTreeRunStatus Status = Exec.Tick(0.1f);
-			AITEST_EQUAL(TEXT("2nd Tick should complete with Running"), Status, EStateTreeRunStatus::Running);
-			AITEST_TRUE(TEXT("2nd tick sold be in new state"),  Exec.ExpectInActiveStates("StateA", "StateB", "StateC", "StateD", "StateE", "StateF"));
-			FTestStateTreeExecutionContext::FLogOrder LogOrder = Exec.Expect("StateATask", TEXT("Tick0"));
-			AITEST_TRUE(TEXT("Tick StateA"), LogOrder);
-			LogOrder = LogOrder.Then("StateBTask", TEXT("Tick0"));
-			AITEST_TRUE(TEXT("Tick StateB"), LogOrder);
-			LogOrder = LogOrder.Then("StateCTask", TEXT("Tick0"));
-			AITEST_TRUE(TEXT("Tick StateC"), LogOrder);
-			LogOrder = LogOrder.Then("StateDTask", TEXT("Tick0"));
-			AITEST_TRUE(TEXT("Tick StateD"), LogOrder);
-			LogOrder = LogOrder.Then("StateETask", TEXT("EnterState=Changed"));
-			AITEST_TRUE(TEXT("Enter StateE"), LogOrder);
-			LogOrder = LogOrder.Then("StateFTask", TEXT("EnterState=Changed"));
-			AITEST_TRUE(TEXT("Enter StateF"), LogOrder);
-			AITEST_FALSE(TEXT("Enter StateA"), Exec.Expect("StateATask", TEXT("EnterState0")));
-			AITEST_FALSE(TEXT("Enter StateB"), Exec.Expect("StateBTask", TEXT("EnterState0")));
-			AITEST_FALSE(TEXT("Enter StateC"), Exec.Expect("StateCTask", TEXT("EnterState0")));
-			AITEST_FALSE(TEXT("Enter StateD"), Exec.Expect("StateDTask", TEXT("EnterState0")));
-			AITEST_FALSE(TEXT("Exit StateA"), Exec.Expect("StateATask", TEXT("ExitState0")));
-			AITEST_FALSE(TEXT("Exit StateB"), Exec.Expect("StateBTask", TEXT("ExitState0")));
-			AITEST_FALSE(TEXT("Exit StateC"), Exec.Expect("StateCTask", TEXT("ExitState0")));
-			AITEST_FALSE(TEXT("Exit StateD"), Exec.Expect("StateDTask", TEXT("ExitState0")));
-			Exec.LogClear();
-			TransitionToExecute = ECustomFunctionToRun::None;
-		}
-		// Same transition but with a completed D
-		{
-			TransitionToExecute = ECustomFunctionToRun::TransitionD_To_E;
-			AllStateHandle.bFinishTasksD = true;
-			if (!ResetInstanceData())
 			{
-				return false;
+				FTestStateTreeExecutionContext Exec(StateTree, StateTree, InstanceData);
+				const EStateTreeRunStatus Status = Exec.Tick(0.1f);
+				AITEST_EQUAL("1st Tick should complete with Running", Status, EStateTreeRunStatus::Running);
+				AITEST_TRUE("1st Tick no transition", Exec.ExpectInActiveStates("StateA", "StateB", "StateC", "StateD"));
+				FTestStateTreeExecutionContext::FLogOrder LogOrder = Exec.Expect("StateATask", TEXT("Tick0"));
+				AITEST_TRUE(TEXT("Tick StateA"), LogOrder);
+				LogOrder = LogOrder.Then("StateBTask", TEXT("Tick0"));
+				AITEST_TRUE(TEXT("Tick StateB"), LogOrder);
+				LogOrder = LogOrder.Then("StateCTask", TEXT("Tick0"));
+				AITEST_TRUE(TEXT("Tick StateC"), LogOrder);
+				LogOrder = LogOrder.Then("StateDTask", TEXT("Tick0"));
+				AITEST_TRUE(TEXT("Tick StateD"), LogOrder);
+				AITEST_FALSE(TEXT("Enter StateA"), Exec.Expect("StateATask", TEXT("EnterState0")));
+				AITEST_FALSE(TEXT("Enter StateB"), Exec.Expect("StateBTask", TEXT("EnterState0")));
+				AITEST_FALSE(TEXT("Enter StateC"), Exec.Expect("StateCTask", TEXT("EnterState0")));
+				AITEST_FALSE(TEXT("Enter StateD"), Exec.Expect("StateDTask", TEXT("EnterState0")));
+				AITEST_FALSE(TEXT("Exit StateA"), Exec.Expect("StateATask", TEXT("ExitState0")));
+				AITEST_FALSE(TEXT("Exit StateB"), Exec.Expect("StateBTask", TEXT("ExitState0")));
+				AITEST_FALSE(TEXT("Exit StateC"), Exec.Expect("StateCTask", TEXT("ExitState0")));
+				AITEST_FALSE(TEXT("Exit StateD"), Exec.Expect("StateDTask", TEXT("ExitState0")));
+				Exec.LogClear();
 			}
-
-			FTestStateTreeExecutionContext Exec(StateTree, StateTree, InstanceData);
-			const EStateTreeRunStatus Status = Exec.Tick(0.1f);
-			AITEST_EQUAL(TEXT("Tick should complete with Running"), Status, EStateTreeRunStatus::Running);
-			AITEST_TRUE(TEXT("Start should"), Exec.ExpectInActiveStates("StateA", "StateB", "StateC", "StateD", "StateE", "StateF"));
-			FTestStateTreeExecutionContext::FLogOrder LogOrder = Exec.Expect("StateATask", TEXT("Tick0"));
-			AITEST_TRUE(TEXT("Tick StateA"), LogOrder);
-			LogOrder = LogOrder.Then("StateBTask", TEXT("Tick0"));
-			AITEST_TRUE(TEXT("Tick StateB"), LogOrder);
-			LogOrder = LogOrder.Then("StateCTask", TEXT("Tick0"));
-			AITEST_TRUE(TEXT("Tick StateC"), LogOrder);
-			LogOrder = LogOrder.Then("StateDTask", TEXT("Tick0"));
-			AITEST_TRUE(TEXT("Tick StateD"), LogOrder);
-			//@TODO State D completed. It should be a new state.
-			//LogOrder = LogOrder.Then("StateDTask", TEXT("ExitState=Changed"));
-			//AITEST_TRUE(TEXT("Exit StateD"), LogOrder);
-			//@TODO State D completed. It should be a new state.
-			//LogOrder = LogOrder.Then("StateDTask", TEXT("EnterState=Changed"));
-			//AITEST_TRUE(TEXT("Enter StateD"), LogOrder);
-			LogOrder = LogOrder.Then("StateETask", TEXT("EnterState=Changed"));
-			AITEST_TRUE(TEXT("Enter StateE"), LogOrder);
-			LogOrder = LogOrder.Then("StateFTask", TEXT("EnterState=Changed"));
-			AITEST_TRUE(TEXT("Enter StateF"), LogOrder);
-			AITEST_FALSE(TEXT("Enter StateA"), Exec.Expect("StateATask", TEXT("EnterState0")));
-			AITEST_FALSE(TEXT("Enter StateB"), Exec.Expect("StateBTask", TEXT("EnterState0")));
-			AITEST_FALSE(TEXT("Enter StateC"), Exec.Expect("StateCTask", TEXT("EnterState0")));
-			AITEST_FALSE(TEXT("Exit StateA"), Exec.Expect("StateATask", TEXT("ExitState0")));
-			AITEST_FALSE(TEXT("Exit StateB"), Exec.Expect("StateBTask", TEXT("ExitState0")));
-			AITEST_FALSE(TEXT("Exit StateC"), Exec.Expect("StateCTask", TEXT("ExitState0")));
-			Exec.LogClear();
-			TransitionToExecute = ECustomFunctionToRun::None;
-			AllStateHandle.bFinishTasksD = false;
-		}
-		{
-			TransitionToExecute = ECustomFunctionToRun::TransitionB_To_I;
-			if (!ResetInstanceData())
+			// Select new child
 			{
-				return false;
+				TransitionToExecute = ECustomFunctionToRun::TransitionD_To_E;
+
+				FTestStateTreeExecutionContext Exec(StateTree, StateTree, InstanceData);
+				const EStateTreeRunStatus Status = Exec.Tick(0.1f);
+				AITEST_EQUAL(TEXT("2nd Tick should complete with Running"), Status, EStateTreeRunStatus::Running);
+				AITEST_TRUE(TEXT("2nd tick should be in new state"), Exec.ExpectInActiveStates("StateA", "StateB", "StateC", "StateD", "StateE", "StateF"));
+				FTestStateTreeExecutionContext::FLogOrder LogOrder = Exec.Expect("StateATask", TEXT("Tick0"));
+				AITEST_TRUE(TEXT("Tick StateA"), LogOrder);
+				LogOrder = LogOrder.Then("StateBTask", TEXT("Tick0"));
+				AITEST_TRUE(TEXT("Tick StateB"), LogOrder);
+				LogOrder = LogOrder.Then("StateCTask", TEXT("Tick0"));
+				AITEST_TRUE(TEXT("Tick StateC"), LogOrder);
+				LogOrder = LogOrder.Then("StateDTask", TEXT("Tick0"));
+				AITEST_TRUE(TEXT("Tick StateD"), LogOrder);
+				LogOrder = LogOrder.Then("StateETask", EnterStateChangedStr);
+				AITEST_TRUE(TEXT("Enter StateE"), LogOrder);
+				LogOrder = LogOrder.Then("StateFTask", EnterStateChangedStr);
+				AITEST_TRUE(TEXT("Enter StateF"), LogOrder);
+				AITEST_FALSE(TEXT("Enter StateA"), Exec.Expect("StateATask", TEXT("EnterState0")));
+				AITEST_FALSE(TEXT("Enter StateB"), Exec.Expect("StateBTask", TEXT("EnterState0")));
+				AITEST_FALSE(TEXT("Enter StateC"), Exec.Expect("StateCTask", TEXT("EnterState0")));
+				AITEST_FALSE(TEXT("Enter StateD"), Exec.Expect("StateDTask", TEXT("EnterState0")));
+				AITEST_FALSE(TEXT("Exit StateA"), Exec.Expect("StateATask", TEXT("ExitState0")));
+				AITEST_FALSE(TEXT("Exit StateB"), Exec.Expect("StateBTask", TEXT("ExitState0")));
+				AITEST_FALSE(TEXT("Exit StateC"), Exec.Expect("StateCTask", TEXT("ExitState0")));
+				AITEST_FALSE(TEXT("Exit StateD"), Exec.Expect("StateDTask", TEXT("ExitState0")));
+				Exec.LogClear();
+				TransitionToExecute = ECustomFunctionToRun::None;
 			}
 
-			FTestStateTreeExecutionContext Exec(StateTree, StateTree, InstanceData);
-			const EStateTreeRunStatus Status = Exec.Tick(0.1f);
-			AITEST_EQUAL(TEXT("Tick should complete with Running"), Status, EStateTreeRunStatus::Running);
-			AITEST_TRUE(TEXT("Start should"), Exec.ExpectInActiveStates("StateA", "StateB", "StateI", "StateJ"));
-			FTestStateTreeExecutionContext::FLogOrder LogOrder = Exec.Expect("StateATask", TEXT("Tick0"));
-			AITEST_TRUE(TEXT("Tick StateA"), LogOrder);
-			LogOrder = LogOrder.Then("StateBTask", TEXT("Tick0"));
-			AITEST_TRUE(TEXT("Tick StateB"), LogOrder);
-			LogOrder = LogOrder.Then("StateCTask", TEXT("Tick0"));
-			AITEST_TRUE(TEXT("Tick StateC"), LogOrder);
-			LogOrder = LogOrder.Then("StateDTask", TEXT("Tick0"));
-			AITEST_TRUE(TEXT("Tick StateD"), LogOrder);
-			LogOrder = LogOrder.Then("StateDTask", TEXT("ExitState=Changed"));
-			AITEST_TRUE(TEXT("Exit StateD"), LogOrder);
-			LogOrder = LogOrder.Then("StateCTask", TEXT("ExitState=Changed"));
-			AITEST_TRUE(TEXT("Exit StateC"), LogOrder);
-			LogOrder = LogOrder.Then("StateITask", TEXT("EnterState=Changed"));
-			AITEST_TRUE(TEXT("Enter StateI"), LogOrder);
-			LogOrder = LogOrder.Then("StateJTask", TEXT("EnterState=Changed"));
-			AITEST_TRUE(TEXT("Enter StateJ"), LogOrder);
-			AITEST_FALSE(TEXT("Enter StateA"), Exec.Expect("StateATask", TEXT("EnterState0")));
-			AITEST_FALSE(TEXT("Enter StateB"), Exec.Expect("StateBTask", TEXT("EnterState0")));
-			AITEST_FALSE(TEXT("Exit StateA"), Exec.Expect("StateATask", TEXT("ExitState0")));
-			AITEST_FALSE(TEXT("Exit StateB"), Exec.Expect("StateBTask", TEXT("ExitState0")));
-			Exec.LogClear();
-			TransitionToExecute = ECustomFunctionToRun::None;
-		}
-		{
-			TransitionToExecute = ECustomFunctionToRun::TransitionB_To_C;
-			if (!ResetInstanceData())
+			// Select a new child but with a completed D
 			{
-				return false;
+				TransitionToExecute = ECustomFunctionToRun::TransitionD_To_E;
+				AllStateHandle.bFinishTasksD = true;
+				if (!ResetInstanceData())
+				{
+					return false;
+				}
+
+				const bool bUseCompletedRule = EnumHasAllFlags(StateSelectionRules, EStateTreeStateSelectionRules::CompletedTransitionStatesCreateNewStates);
+
+				FTestStateTreeExecutionContext Exec(StateTree, StateTree, InstanceData);
+				const EStateTreeRunStatus Status = Exec.Tick(0.1f);
+				AITEST_EQUAL(TEXT("Tick should complete with Running"), Status, EStateTreeRunStatus::Running);
+				AITEST_TRUE(TEXT("Start should"), Exec.ExpectInActiveStates("StateA", "StateB", "StateC", "StateD", "StateE", "StateF"));
+				FTestStateTreeExecutionContext::FLogOrder LogOrder = Exec.Expect("StateATask", TEXT("Tick0"));
+				AITEST_TRUE(TEXT("Tick StateA"), LogOrder);
+				LogOrder = LogOrder.Then("StateBTask", TEXT("Tick0"));
+				AITEST_TRUE(TEXT("Tick StateB"), LogOrder);
+				LogOrder = LogOrder.Then("StateCTask", TEXT("Tick0"));
+				AITEST_TRUE(TEXT("Tick StateC"), LogOrder);
+				LogOrder = LogOrder.Then("StateDTask", TEXT("Tick0"));
+				AITEST_TRUE(TEXT("Tick StateD"), LogOrder);
+				if (bUseCompletedRule)
+				{
+					LogOrder = LogOrder.Then("StateDTask", ExitStateChangedStr);
+					AITEST_TRUE(TEXT("Exit StateD"), LogOrder);
+					LogOrder = LogOrder.Then("StateDTask", EnterStateChangedStr);
+					AITEST_TRUE(TEXT("Enter StateD"), LogOrder);
+				}
+				LogOrder = LogOrder.Then("StateETask", EnterStateChangedStr);
+				AITEST_TRUE(TEXT("Enter StateE"), LogOrder);
+				LogOrder = LogOrder.Then("StateFTask", EnterStateChangedStr);
+				AITEST_TRUE(TEXT("Enter StateF"), LogOrder);
+				AITEST_FALSE(TEXT("Enter StateA"), Exec.Expect("StateATask", TEXT("EnterState0")));
+				AITEST_FALSE(TEXT("Enter StateB"), Exec.Expect("StateBTask", TEXT("EnterState0")));
+				AITEST_FALSE(TEXT("Enter StateC"), Exec.Expect("StateCTask", TEXT("EnterState0")));
+				AITEST_FALSE(TEXT("Exit StateA"), Exec.Expect("StateATask", TEXT("ExitState0")));
+				AITEST_FALSE(TEXT("Exit StateB"), Exec.Expect("StateBTask", TEXT("ExitState0")));
+				AITEST_FALSE(TEXT("Exit StateC"), Exec.Expect("StateCTask", TEXT("ExitState0")));
+				Exec.LogClear();
+				TransitionToExecute = ECustomFunctionToRun::None;
+				AllStateHandle.bFinishTasksD = false;
 			}
+			// Select a new child but B is completed (before the source/target
+			{
+				TransitionToExecute = ECustomFunctionToRun::TransitionD_To_E;
+				AllStateHandle.bFinishTasksB = true;
+				if (!ResetInstanceData())
+				{
+					return false;
+				}
 
-			FTestStateTreeExecutionContext Exec(StateTree, StateTree, InstanceData);
-			const EStateTreeRunStatus Status = Exec.Tick(0.1f);
-			AITEST_EQUAL(TEXT("Tick should complete with Running"), Status, EStateTreeRunStatus::Running);
-			AITEST_TRUE(TEXT("Start should"), Exec.ExpectInActiveStates("StateA", "StateB", "StateC", "StateD"));
-			FTestStateTreeExecutionContext::FLogOrder LogOrder = Exec.Expect("StateATask", TEXT("Tick0"));
-			AITEST_TRUE(TEXT("Tick StateA"), LogOrder);
-			LogOrder = LogOrder.Then("StateBTask", TEXT("Tick0"));
-			AITEST_TRUE(TEXT("Tick StateB"), LogOrder);
-			LogOrder = LogOrder.Then("StateCTask", TEXT("Tick0"));
-			AITEST_TRUE(TEXT("Tick StateC"), LogOrder);
-			LogOrder = LogOrder.Then("StateDTask", TEXT("Tick0"));
-			AITEST_TRUE(TEXT("Tick StateD"), LogOrder);
-			LogOrder = LogOrder.Then("StateDTask", TEXT("ExitState=Sustained"));
-			AITEST_TRUE(TEXT("Exit StateD"), LogOrder);
-			LogOrder = LogOrder.Then("StateCTask", TEXT("ExitState=Sustained"));
-			AITEST_TRUE(TEXT("Exit StateC"), LogOrder);
-			LogOrder = LogOrder.Then("StateCTask", TEXT("EnterState=Sustained"));
-			AITEST_TRUE(TEXT("Enter StateC"), LogOrder);
-			LogOrder = LogOrder.Then("StateDTask", TEXT("EnterState=Sustained"));
-			AITEST_TRUE(TEXT("Enter StateD"), LogOrder);
-			AITEST_FALSE(TEXT("Enter StateA"), Exec.Expect("StateATask", TEXT("EnterState0")));
-			AITEST_FALSE(TEXT("Enter StateB"), Exec.Expect("StateBTask", TEXT("EnterState0")));
-			AITEST_FALSE(TEXT("Exit StateA"), Exec.Expect("StateATask", TEXT("ExitState0")));
-			AITEST_FALSE(TEXT("Exit StateB"), Exec.Expect("StateBTask", TEXT("ExitState0")));
-			Exec.LogClear();
-			TransitionToExecute = ECustomFunctionToRun::None;
-		}
-		{
-			TransitionToExecute = ECustomFunctionToRun::TransitionC_To_B;
-			if (!ResetInstanceData())
+				const bool bUseCompletedRule = EnumHasAllFlags(StateSelectionRules, EStateTreeStateSelectionRules::CompletedStateBeforeTransitionSourceFailsTransition);
+
+				FTestStateTreeExecutionContext Exec(StateTree, StateTree, InstanceData);
+				const EStateTreeRunStatus Status = Exec.Tick(0.1f);
+				AITEST_EQUAL(TEXT("Tick should complete with Running"), Status, EStateTreeRunStatus::Running);
+				if (bUseCompletedRule)
+				{
+					AITEST_TRUE(TEXT("Start should"), Exec.ExpectInActiveStates("StateA", "StateB", "StateC", "StateD"));
+				}
+				else
+				{
+					AITEST_TRUE(TEXT("Start should"), Exec.ExpectInActiveStates("StateA", "StateB", "StateC", "StateD", "StateE", "StateF"));
+				}
+				FTestStateTreeExecutionContext::FLogOrder LogOrder = Exec.Expect("StateATask", TEXT("Tick0"));
+				AITEST_TRUE(TEXT("Tick StateA"), LogOrder);
+				LogOrder = LogOrder.Then("StateBTask", TEXT("Tick0"));
+				AITEST_TRUE(TEXT("Tick StateB"), LogOrder);
+				LogOrder = LogOrder.Then("StateCTask", TEXT("Tick0"));
+				AITEST_TRUE(TEXT("Tick StateC"), LogOrder);
+				LogOrder = LogOrder.Then("StateDTask", TEXT("Tick0"));
+				AITEST_TRUE(TEXT("Tick StateD"), LogOrder);
+				if (bUseCompletedRule)
+				{
+					LogOrder = LogOrder.Then("StateDTask", ExitStateChangedStr);
+					AITEST_TRUE(TEXT("Enter StateD"), LogOrder);
+					LogOrder = LogOrder.Then("StateCTask", ExitStateChangedStr);
+					AITEST_TRUE(TEXT("Enter StateC"), LogOrder);
+					LogOrder = LogOrder.Then("StateBTask", ExitStateChangedStr);
+					AITEST_TRUE(TEXT("Enter StateB"), LogOrder);
+					LogOrder = LogOrder.Then("StateATask", ExitStateSustainedStr);
+					AITEST_TRUE(TEXT("Enter StateA"), LogOrder);
+					LogOrder = LogOrder.Then("StateATask", EnterStateSustainedStr);
+					AITEST_TRUE(TEXT("Enter StateA"), LogOrder);
+					LogOrder = LogOrder.Then("StateBTask", EnterStateChangedStr);
+					AITEST_TRUE(TEXT("Enter StateB"), LogOrder);
+					LogOrder = LogOrder.Then("StateCTask", EnterStateChangedStr);
+					AITEST_TRUE(TEXT("Enter StateC"), LogOrder);
+					LogOrder = LogOrder.Then("StateDTask", EnterStateChangedStr);
+					AITEST_TRUE(TEXT("Enter StateD"), LogOrder);
+
+					// Only one transition
+					LogOrder = LogOrder.Then("StateBTask", TEXT("EnterState0"));
+					AITEST_FALSE(TEXT("Enter StateD"), LogOrder);
+
+					AITEST_FALSE(TEXT("Enter StateE"), Exec.Expect("StateETask", TEXT("EnterState0")));
+					AITEST_FALSE(TEXT("Enter StateF"), Exec.Expect("StateFTask", TEXT("EnterState0")));
+					AITEST_FALSE(TEXT("Exit StateE"), Exec.Expect("StateETask", TEXT("ExitState0")));
+					AITEST_FALSE(TEXT("Exit StateF"), Exec.Expect("StateFTask", TEXT("ExitState0")));
+				}
+				else
+				{
+					LogOrder = LogOrder.Then("StateETask", EnterStateChangedStr);
+					AITEST_TRUE(TEXT("Enter StateE"), LogOrder);
+					LogOrder = LogOrder.Then("StateFTask", EnterStateChangedStr);
+					AITEST_TRUE(TEXT("Enter StateF"), LogOrder);
+					AITEST_FALSE(TEXT("Enter StateA"), Exec.Expect("StateATask", TEXT("EnterState0")));
+					AITEST_FALSE(TEXT("Enter StateB"), Exec.Expect("StateBTask", TEXT("EnterState0")));
+					AITEST_FALSE(TEXT("Enter StateC"), Exec.Expect("StateCTask", TEXT("EnterState0")));
+					AITEST_FALSE(TEXT("Exit StateA"), Exec.Expect("StateATask", TEXT("ExitState0")));
+					AITEST_FALSE(TEXT("Exit StateB"), Exec.Expect("StateBTask", TEXT("ExitState0")));
+					AITEST_FALSE(TEXT("Exit StateC"), Exec.Expect("StateCTask", TEXT("ExitState0")));
+				}
+				Exec.LogClear();
+				TransitionToExecute = ECustomFunctionToRun::None;
+				AllStateHandle.bFinishTasksB = false;
+			}
+			// Select a sibling
+			{
+				TransitionToExecute = ECustomFunctionToRun::TransitionB_To_I;
+				if (!ResetInstanceData())
+				{
+					return false;
+				}
+
+				FTestStateTreeExecutionContext Exec(StateTree, StateTree, InstanceData);
+				const EStateTreeRunStatus Status = Exec.Tick(0.1f);
+				AITEST_EQUAL(TEXT("Tick should complete with Running"), Status, EStateTreeRunStatus::Running);
+				AITEST_TRUE(TEXT("Start should"), Exec.ExpectInActiveStates("StateA", "StateB", "StateI", "StateJ"));
+				FTestStateTreeExecutionContext::FLogOrder LogOrder = Exec.Expect("StateATask", TEXT("Tick0"));
+				AITEST_TRUE(TEXT("Tick StateA"), LogOrder);
+				LogOrder = LogOrder.Then("StateBTask", TEXT("Tick0"));
+				AITEST_TRUE(TEXT("Tick StateB"), LogOrder);
+				LogOrder = LogOrder.Then("StateCTask", TEXT("Tick0"));
+				AITEST_TRUE(TEXT("Tick StateC"), LogOrder);
+				LogOrder = LogOrder.Then("StateDTask", TEXT("Tick0"));
+				AITEST_TRUE(TEXT("Tick StateD"), LogOrder);
+				LogOrder = LogOrder.Then("StateDTask", ExitStateChangedStr);
+				AITEST_TRUE(TEXT("Exit StateD"), LogOrder);
+				LogOrder = LogOrder.Then("StateCTask", ExitStateChangedStr);
+				AITEST_TRUE(TEXT("Exit StateC"), LogOrder);
+				LogOrder = LogOrder.Then("StateITask", EnterStateChangedStr);
+				AITEST_TRUE(TEXT("Enter StateI"), LogOrder);
+				LogOrder = LogOrder.Then("StateJTask", EnterStateChangedStr);
+				AITEST_TRUE(TEXT("Enter StateJ"), LogOrder);
+				AITEST_FALSE(TEXT("Enter StateA"), Exec.Expect("StateATask", TEXT("EnterState0")));
+				AITEST_FALSE(TEXT("Enter StateB"), Exec.Expect("StateBTask", TEXT("EnterState0")));
+				AITEST_FALSE(TEXT("Exit StateA"), Exec.Expect("StateATask", TEXT("ExitState0")));
+				AITEST_FALSE(TEXT("Exit StateB"), Exec.Expect("StateBTask", TEXT("ExitState0")));
+				Exec.LogClear();
+				TransitionToExecute = ECustomFunctionToRun::None;
+			}
+			// No state completed. Reselect child.
 			{
-				return false;
+				TransitionToExecute = ECustomFunctionToRun::TransitionB_To_C;
+				if (!ResetInstanceData())
+				{
+					return false;
+				}
+
+				FTestStateTreeExecutionContext Exec(StateTree, StateTree, InstanceData);
+				const EStateTreeRunStatus Status = Exec.Tick(0.1f);
+				AITEST_EQUAL(TEXT("Tick should complete with Running"), Status, EStateTreeRunStatus::Running);
+				AITEST_TRUE(TEXT("Start should"), Exec.ExpectInActiveStates("StateA", "StateB", "StateC", "StateD"));
+				FTestStateTreeExecutionContext::FLogOrder LogOrder = Exec.Expect("StateATask", TEXT("Tick0"));
+				AITEST_TRUE(TEXT("Tick StateA"), LogOrder);
+				LogOrder = LogOrder.Then("StateBTask", TEXT("Tick0"));
+				AITEST_TRUE(TEXT("Tick StateB"), LogOrder);
+				LogOrder = LogOrder.Then("StateCTask", TEXT("Tick0"));
+				AITEST_TRUE(TEXT("Tick StateC"), LogOrder);
+				LogOrder = LogOrder.Then("StateDTask", TEXT("Tick0"));
+				AITEST_TRUE(TEXT("Tick StateD"), LogOrder);
+				LogOrder = LogOrder.Then("StateDTask", ExitStateSustainedStr);
+				AITEST_TRUE(TEXT("Exit StateD"), LogOrder);
+				LogOrder = LogOrder.Then("StateCTask", ExitStateSustainedStr);
+				AITEST_TRUE(TEXT("Exit StateC"), LogOrder);
+				LogOrder = LogOrder.Then("StateCTask", EnterStateSustainedStr);
+				AITEST_TRUE(TEXT("Enter StateC"), LogOrder);
+				LogOrder = LogOrder.Then("StateDTask", EnterStateSustainedStr);
+				AITEST_TRUE(TEXT("Enter StateD"), LogOrder);
+				AITEST_FALSE(TEXT("Enter StateA"), Exec.Expect("StateATask", TEXT("EnterState0")));
+				AITEST_FALSE(TEXT("Enter StateB"), Exec.Expect("StateBTask", TEXT("EnterState0")));
+				AITEST_FALSE(TEXT("Exit StateA"), Exec.Expect("StateATask", TEXT("ExitState0")));
+				AITEST_FALSE(TEXT("Exit StateB"), Exec.Expect("StateBTask", TEXT("ExitState0")));
+				Exec.LogClear();
+				TransitionToExecute = ECustomFunctionToRun::None;
 			}
+			// State completed. Reselect child.
+			{
+				TransitionToExecute = ECustomFunctionToRun::TransitionB_To_C;
+				AllStateHandle.bFinishTasksC = true;
+				if (!ResetInstanceData())
+				{
+					return false;
+				}
 
-			FTestStateTreeExecutionContext Exec(StateTree, StateTree, InstanceData);
-			const EStateTreeRunStatus Status = Exec.Tick(0.1f);
-			AITEST_EQUAL(TEXT("Tick should complete with Running"), Status, EStateTreeRunStatus::Running);
-			AITEST_TRUE(TEXT("Start should"), Exec.ExpectInActiveStates("StateA", "StateB", "StateC", "StateD"));
-			FTestStateTreeExecutionContext::FLogOrder LogOrder = Exec.Expect("StateATask", TEXT("Tick0"));
-			AITEST_TRUE(TEXT("Tick StateA"), LogOrder);
-			LogOrder = LogOrder.Then("StateBTask", TEXT("Tick0"));
-			AITEST_TRUE(TEXT("Tick StateB"), LogOrder);
-			LogOrder = LogOrder.Then("StateCTask", TEXT("Tick0"));
-			AITEST_TRUE(TEXT("Tick StateC"), LogOrder);
-			LogOrder = LogOrder.Then("StateDTask", TEXT("Tick0"));
-			AITEST_TRUE(TEXT("Tick StateD"), LogOrder);
-			LogOrder = LogOrder.Then("StateDTask", TEXT("ExitState=Sustained"));
-			AITEST_TRUE(TEXT("Exit StateD"), LogOrder);
-			LogOrder = LogOrder.Then("StateCTask", TEXT("ExitState=Sustained"));
-			AITEST_TRUE(TEXT("Exit StateC"), LogOrder);
-			LogOrder = LogOrder.Then("StateBTask", TEXT("ExitState=Sustained"));
-			AITEST_TRUE(TEXT("Exit StateB"), LogOrder);
-			LogOrder = LogOrder.Then("StateBTask", TEXT("EnterState=Sustained"));
-			AITEST_TRUE(TEXT("Enter StateB"), LogOrder);
-			LogOrder = LogOrder.Then("StateCTask", TEXT("EnterState=Sustained"));
-			AITEST_TRUE(TEXT("Enter StateC"), LogOrder);
-			LogOrder = LogOrder.Then("StateDTask", TEXT("EnterState=Sustained"));
-			AITEST_TRUE(TEXT("Enter StateD"), LogOrder);
-			AITEST_FALSE(TEXT("Enter StateA"), Exec.Expect("StateATask", TEXT("EnterState0")));
-			AITEST_FALSE(TEXT("Exit StateA"), Exec.Expect("StateATask", TEXT("ExitState0")));
+				const bool bUseCompletedRule = EnumHasAllFlags(StateSelectionRules, EStateTreeStateSelectionRules::CompletedTransitionStatesCreateNewStates);
+				const TCHAR* SchemaEnterStateStr = bUseCompletedRule ? EnterStateChangedStr : EnterStateSustainedStr;
+				const TCHAR* SchemaExitStateStr = bUseCompletedRule ? ExitStateChangedStr : ExitStateSustainedStr;
 
-			Exec.LogClear();
-			TransitionToExecute = ECustomFunctionToRun::None;
-		}
-		{
-			TransitionToExecute = ECustomFunctionToRun::TransitionC_To_B;
-			AllStateHandle.bFinishTasksC = true;
-			if (!ResetInstanceData())
+				FTestStateTreeExecutionContext Exec(StateTree, StateTree, InstanceData);
+				const EStateTreeRunStatus Status = Exec.Tick(0.1f);
+				AITEST_EQUAL(TEXT("Tick should complete with Running"), Status, EStateTreeRunStatus::Running);
+				AITEST_TRUE(TEXT("Start should"), Exec.ExpectInActiveStates("StateA", "StateB", "StateC", "StateD"));
+				FTestStateTreeExecutionContext::FLogOrder LogOrder = Exec.Expect("StateATask", TEXT("Tick0"));
+				AITEST_TRUE(TEXT("Tick StateA"), LogOrder);
+				LogOrder = LogOrder.Then("StateBTask", TEXT("Tick0"));
+				AITEST_TRUE(TEXT("Tick StateB"), LogOrder);
+				LogOrder = LogOrder.Then("StateCTask", TEXT("Tick0"));
+				AITEST_TRUE(TEXT("Tick StateC"), LogOrder);
+				LogOrder = LogOrder.Then("StateDTask", TEXT("Tick0"));
+				AITEST_TRUE(TEXT("Tick StateD"), LogOrder);
+				LogOrder = LogOrder.Then("StateDTask", SchemaExitStateStr);
+				AITEST_TRUE(TEXT("Exit StateD"), LogOrder);
+				LogOrder = LogOrder.Then("StateCTask", SchemaExitStateStr);
+				AITEST_TRUE(TEXT("Exit StateC"), LogOrder);
+				LogOrder = LogOrder.Then("StateCTask", SchemaEnterStateStr);
+				AITEST_TRUE(TEXT("Enter StateC"), LogOrder);
+				LogOrder = LogOrder.Then("StateDTask", SchemaEnterStateStr);
+				AITEST_TRUE(TEXT("Enter StateD"), LogOrder);
+				AITEST_FALSE(TEXT("Enter StateA"), Exec.Expect("StateATask", TEXT("EnterState0")));
+				AITEST_FALSE(TEXT("Enter StateB"), Exec.Expect("StateBTask", TEXT("EnterState0")));
+				AITEST_FALSE(TEXT("Exit StateA"), Exec.Expect("StateATask", TEXT("ExitState0")));
+				AITEST_FALSE(TEXT("Exit StateB"), Exec.Expect("StateBTask", TEXT("ExitState0")));
+				Exec.LogClear();
+				TransitionToExecute = ECustomFunctionToRun::None;
+				AllStateHandle.bFinishTasksC = false;
+			}
+			// No state completed. Reselect parent.
 			{
-				return false;
+				TransitionToExecute = ECustomFunctionToRun::TransitionC_To_B;
+				if (!ResetInstanceData())
+				{
+					return false;
+				}
+
+				FTestStateTreeExecutionContext Exec(StateTree, StateTree, InstanceData);
+				const EStateTreeRunStatus Status = Exec.Tick(0.1f);
+				AITEST_EQUAL(TEXT("Tick should complete with Running"), Status, EStateTreeRunStatus::Running);
+				AITEST_TRUE(TEXT("Start should"), Exec.ExpectInActiveStates("StateA", "StateB", "StateC", "StateD"));
+				FTestStateTreeExecutionContext::FLogOrder LogOrder = Exec.Expect("StateATask", TEXT("Tick0"));
+				AITEST_TRUE(TEXT("Tick StateA"), LogOrder);
+				LogOrder = LogOrder.Then("StateBTask", TEXT("Tick0"));
+				AITEST_TRUE(TEXT("Tick StateB"), LogOrder);
+				LogOrder = LogOrder.Then("StateCTask", TEXT("Tick0"));
+				AITEST_TRUE(TEXT("Tick StateC"), LogOrder);
+				LogOrder = LogOrder.Then("StateDTask", TEXT("Tick0"));
+				AITEST_TRUE(TEXT("Tick StateD"), LogOrder);
+				LogOrder = LogOrder.Then("StateDTask", ExitStateSustainedStr);
+				AITEST_TRUE(TEXT("Exit StateD"), LogOrder);
+				LogOrder = LogOrder.Then("StateCTask", ExitStateSustainedStr);
+				AITEST_TRUE(TEXT("Exit StateC"), LogOrder);
+				LogOrder = LogOrder.Then("StateBTask", ExitStateSustainedStr);
+				AITEST_TRUE(TEXT("Exit StateB"), LogOrder);
+				LogOrder = LogOrder.Then("StateBTask", EnterStateSustainedStr);
+				AITEST_TRUE(TEXT("Enter StateB"), LogOrder);
+				LogOrder = LogOrder.Then("StateCTask", EnterStateSustainedStr);
+				AITEST_TRUE(TEXT("Enter StateC"), LogOrder);
+				LogOrder = LogOrder.Then("StateDTask", EnterStateSustainedStr);
+				AITEST_TRUE(TEXT("Enter StateD"), LogOrder);
+				AITEST_FALSE(TEXT("Enter StateA"), Exec.Expect("StateATask", TEXT("EnterState0")));
+				AITEST_FALSE(TEXT("Exit StateA"), Exec.Expect("StateATask", TEXT("ExitState0")));
+
+				Exec.LogClear();
+				TransitionToExecute = ECustomFunctionToRun::None;
 			}
+			// State completed. Reselect parent.
+			{
+				TransitionToExecute = ECustomFunctionToRun::TransitionC_To_B;
+				AllStateHandle.bFinishTasksC = true;
+				if (!ResetInstanceData())
+				{
+					return false;
+				}
 
-			FTestStateTreeExecutionContext Exec(StateTree, StateTree, InstanceData);
-			const EStateTreeRunStatus Status = Exec.Tick(0.1f);
-			AITEST_EQUAL(TEXT("Tick should complete with Running"), Status, EStateTreeRunStatus::Running);
-			AITEST_TRUE(TEXT("Start should"), Exec.ExpectInActiveStates("StateA", "StateB", "StateC", "StateD"));
-			FTestStateTreeExecutionContext::FLogOrder LogOrder = Exec.Expect("StateATask", TEXT("Tick0"));
-			AITEST_TRUE(TEXT("Tick StateA"), LogOrder);
-			LogOrder = LogOrder.Then("StateBTask", TEXT("Tick0"));
-			AITEST_TRUE(TEXT("Tick StateB"), LogOrder);
-			LogOrder = LogOrder.Then("StateCTask", TEXT("Tick0"));
-			AITEST_TRUE(TEXT("Tick StateC"), LogOrder);
-			LogOrder = LogOrder.Then("StateDTask", TEXT("Tick0"));
-			AITEST_TRUE(TEXT("Tick StateD"), LogOrder);
-			//@TODO State C completed. It should change for another state
-			//LogOrder = LogOrder.Then("StateDTask", TEXT("ExitState=Changed"));
-			//AITEST_TRUE(TEXT("Exit StateD"), LogOrder);
-			//@TODO State C completed. It should change for another state
-			//LogOrder = LogOrder.Then("StateCTask", TEXT("ExitState=Changed"));
-			//AITEST_TRUE(TEXT("Exit StateC"), LogOrder);
-			LogOrder = LogOrder.Then("StateBTask", TEXT("ExitState=Sustained"));
-			AITEST_TRUE(TEXT("Exit StateB"), LogOrder);
-			LogOrder = LogOrder.Then("StateBTask", TEXT("EnterState=Sustained"));
-			AITEST_TRUE(TEXT("Enter StateB"), LogOrder);
-			//@TODO State C completed. It should be a new state.
-			//LogOrder = LogOrder.Then("StateCTask", TEXT("EnterState=Changed"));
-			//AITEST_TRUE(TEXT("Enter StateC"), LogOrder);
-			//@TODO State C completed. It should be a new state.
-			//LogOrder = LogOrder.Then("StateDTask", TEXT("EnterState=Changed"));
-			//AITEST_TRUE(TEXT("Enter StateD"), LogOrder);
-			AITEST_FALSE(TEXT("Enter StateA"), Exec.Expect("StateATask", TEXT("EnterState0")));
-			AITEST_FALSE(TEXT("Exit StateA"), Exec.Expect("StateATask", TEXT("ExitState0")));
+				const bool bUseCompletedRule = EnumHasAllFlags(StateSelectionRules, EStateTreeStateSelectionRules::CompletedTransitionStatesCreateNewStates);
+				const TCHAR* SchemaEnterStateStr = bUseCompletedRule ? EnterStateChangedStr : EnterStateSustainedStr;
+				const TCHAR* SchemaExitStateStr = bUseCompletedRule ? ExitStateChangedStr : ExitStateSustainedStr;
 
-			Exec.LogClear();
-			TransitionToExecute = ECustomFunctionToRun::None;
-		}
-		{
-			TransitionToExecute = ECustomFunctionToRun::TransitionB_To_X;
-			if (!ResetInstanceData())
-			{
-				return false;
+				FTestStateTreeExecutionContext Exec(StateTree, StateTree, InstanceData);
+				const EStateTreeRunStatus Status = Exec.Tick(0.1f);
+				AITEST_EQUAL(TEXT("Tick should complete with Running"), Status, EStateTreeRunStatus::Running);
+				AITEST_TRUE(TEXT("Start should"), Exec.ExpectInActiveStates("StateA", "StateB", "StateC", "StateD"));
+				FTestStateTreeExecutionContext::FLogOrder LogOrder = Exec.Expect("StateATask", TEXT("Tick0"));
+				AITEST_TRUE(TEXT("Tick StateA"), LogOrder);
+				LogOrder = LogOrder.Then("StateBTask", TEXT("Tick0"));
+				AITEST_TRUE(TEXT("Tick StateB"), LogOrder);
+				LogOrder = LogOrder.Then("StateCTask", TEXT("Tick0"));
+				AITEST_TRUE(TEXT("Tick StateC"), LogOrder);
+				LogOrder = LogOrder.Then("StateDTask", TEXT("Tick0"));
+				AITEST_TRUE(TEXT("Tick StateD"), LogOrder);
+				LogOrder = LogOrder.Then("StateDTask", SchemaExitStateStr);
+				AITEST_TRUE(TEXT("Exit StateD"), LogOrder);
+				LogOrder = LogOrder.Then("StateCTask", SchemaExitStateStr);
+				AITEST_TRUE(TEXT("Exit StateC"), LogOrder);
+				LogOrder = LogOrder.Then("StateBTask", ExitStateSustainedStr);
+				AITEST_TRUE(TEXT("Exit StateB"), LogOrder);
+				LogOrder = LogOrder.Then("StateBTask", EnterStateSustainedStr);
+				AITEST_TRUE(TEXT("Enter StateB"), LogOrder);
+				LogOrder = LogOrder.Then("StateCTask", SchemaEnterStateStr);
+				AITEST_TRUE(TEXT("Enter StateC"), LogOrder);
+				LogOrder = LogOrder.Then("StateDTask", SchemaEnterStateStr);
+				AITEST_TRUE(TEXT("Enter StateD"), LogOrder);
+				AITEST_FALSE(TEXT("Enter StateA"), Exec.Expect("StateATask", TEXT("EnterState0")));
+				AITEST_FALSE(TEXT("Exit StateA"), Exec.Expect("StateATask", TEXT("ExitState0")));
+
+				Exec.LogClear();
+				AllStateHandle.bFinishTasksC = false;
+				TransitionToExecute = ECustomFunctionToRun::None;
 			}
+			// linked state
+			{
+				TransitionToExecute = ECustomFunctionToRun::TransitionB_To_X;
+				if (!ResetInstanceData())
+				{
+					return false;
+				}
 
-			FTestStateTreeExecutionContext Exec(StateTree, StateTree, InstanceData);
-			const EStateTreeRunStatus Status = Exec.Tick(0.1f);
-			AITEST_EQUAL(TEXT("Tick should complete with Running"), Status, EStateTreeRunStatus::Running);
-			AITEST_TRUE(TEXT("Start should"), Exec.ExpectInActiveStates("StateA", "StateX", "StateY"));
-			FTestStateTreeExecutionContext::FLogOrder LogOrder = Exec.Expect("StateATask", TEXT("Tick0"));
-			AITEST_TRUE(TEXT("Tick StateA"), LogOrder);
-			LogOrder = LogOrder.Then("StateBTask", TEXT("Tick0"));
-			AITEST_TRUE(TEXT("Tick StateB"), LogOrder);
-			LogOrder = LogOrder.Then("StateCTask", TEXT("Tick0"));
-			AITEST_TRUE(TEXT("Tick StateC"), LogOrder);
-			LogOrder = LogOrder.Then("StateDTask", TEXT("Tick0"));
-			AITEST_TRUE(TEXT("Tick StateD"), LogOrder);
-			LogOrder = LogOrder.Then("StateDTask", TEXT("ExitState=Changed"));
-			AITEST_TRUE(TEXT("Exit StateD"), LogOrder);
-			LogOrder = LogOrder.Then("StateCTask", TEXT("ExitState=Changed"));
-			AITEST_TRUE(TEXT("Exit StateC"), LogOrder);
-			LogOrder = LogOrder.Then("StateBTask", TEXT("ExitState=Changed"));
-			AITEST_TRUE(TEXT("Exit StateB"), LogOrder);
-			LogOrder = LogOrder.Then("StateXTask", TEXT("EnterState=Changed"));
-			AITEST_TRUE(TEXT("Enter StateX"), LogOrder);
-			LogOrder = LogOrder.Then("StateYTask", TEXT("EnterState=Changed"));
-			AITEST_TRUE(TEXT("Enter StateY"), LogOrder);
-			AITEST_FALSE(TEXT("Enter StateA"), Exec.Expect("StateATask", TEXT("EnterState0")));
-			AITEST_FALSE(TEXT("Exit StateA"), Exec.Expect("StateATask", TEXT("ExitState0")));
-			Exec.LogClear();
-			TransitionToExecute = ECustomFunctionToRun::None;
-		}
-		{
-			TransitionToExecute = ECustomFunctionToRun::TransitionB_To_B;
-			if (!ResetInstanceData())
+				FTestStateTreeExecutionContext Exec(StateTree, StateTree, InstanceData);
+				const EStateTreeRunStatus Status = Exec.Tick(0.1f);
+				AITEST_EQUAL(TEXT("Tick should complete with Running"), Status, EStateTreeRunStatus::Running);
+				AITEST_TRUE(TEXT("Start should"), Exec.ExpectInActiveStates("StateA", "StateX", "StateY"));
+				FTestStateTreeExecutionContext::FLogOrder LogOrder = Exec.Expect("StateATask", TEXT("Tick0"));
+				AITEST_TRUE(TEXT("Tick StateA"), LogOrder);
+				LogOrder = LogOrder.Then("StateBTask", TEXT("Tick0"));
+				AITEST_TRUE(TEXT("Tick StateB"), LogOrder);
+				LogOrder = LogOrder.Then("StateCTask", TEXT("Tick0"));
+				AITEST_TRUE(TEXT("Tick StateC"), LogOrder);
+				LogOrder = LogOrder.Then("StateDTask", TEXT("Tick0"));
+				AITEST_TRUE(TEXT("Tick StateD"), LogOrder);
+				LogOrder = LogOrder.Then("StateDTask", ExitStateChangedStr);
+				AITEST_TRUE(TEXT("Exit StateD"), LogOrder);
+				LogOrder = LogOrder.Then("StateCTask", ExitStateChangedStr);
+				AITEST_TRUE(TEXT("Exit StateC"), LogOrder);
+				LogOrder = LogOrder.Then("StateBTask", ExitStateChangedStr);
+				AITEST_TRUE(TEXT("Exit StateB"), LogOrder);
+				LogOrder = LogOrder.Then("StateXTask", EnterStateChangedStr);
+				AITEST_TRUE(TEXT("Enter StateX"), LogOrder);
+				LogOrder = LogOrder.Then("StateYTask", EnterStateChangedStr);
+				AITEST_TRUE(TEXT("Enter StateY"), LogOrder);
+				AITEST_FALSE(TEXT("Enter StateA"), Exec.Expect("StateATask", TEXT("EnterState0")));
+				AITEST_FALSE(TEXT("Exit StateA"), Exec.Expect("StateATask", TEXT("ExitState0")));
+				Exec.LogClear();
+				TransitionToExecute = ECustomFunctionToRun::None;
+			}
+			// Reselection same state
 			{
-				return false;
+				TransitionToExecute = ECustomFunctionToRun::TransitionB_To_B;
+				if (!ResetInstanceData())
+				{
+					return false;
+				}
+
+				FTestStateTreeExecutionContext Exec(StateTree, StateTree, InstanceData);
+				const EStateTreeRunStatus Status = Exec.Tick(0.1f);
+				AITEST_EQUAL(TEXT("Tick should complete with Running"), Status, EStateTreeRunStatus::Running);
+				AITEST_TRUE(TEXT("Start should"), Exec.ExpectInActiveStates("StateA", "StateB", "StateC", "StateD"));
+				FTestStateTreeExecutionContext::FLogOrder LogOrder = Exec.Expect("StateATask", TEXT("Tick0"));
+				AITEST_TRUE(TEXT("Tick StateA"), LogOrder);
+				LogOrder = LogOrder.Then("StateBTask", TEXT("Tick0"));
+				AITEST_TRUE(TEXT("Tick StateB"), LogOrder);
+				LogOrder = LogOrder.Then("StateCTask", TEXT("Tick0"));
+				AITEST_TRUE(TEXT("Tick StateC"), LogOrder);
+				LogOrder = LogOrder.Then("StateDTask", TEXT("Tick0"));
+				AITEST_TRUE(TEXT("Tick StateD"), LogOrder);
+				LogOrder = LogOrder.Then("StateDTask", ExitStateSustainedStr);
+				AITEST_TRUE(TEXT("Exit StateD"), LogOrder);
+				LogOrder = LogOrder.Then("StateCTask", ExitStateSustainedStr);
+				AITEST_TRUE(TEXT("Exit StateC"), LogOrder);
+				LogOrder = LogOrder.Then("StateBTask", ExitStateSustainedStr);
+				AITEST_TRUE(TEXT("Exit StateB"), LogOrder);
+				LogOrder = LogOrder.Then("StateBTask", EnterStateSustainedStr);
+				AITEST_TRUE(TEXT("Enter StateB"), LogOrder);
+				LogOrder = LogOrder.Then("StateCTask", EnterStateSustainedStr);
+				AITEST_TRUE(TEXT("Enter StateC"), LogOrder);
+				LogOrder = LogOrder.Then("StateDTask", EnterStateSustainedStr);
+				AITEST_TRUE(TEXT("Enter StateD"), LogOrder);
+				AITEST_FALSE(TEXT("Enter StateA"), Exec.Expect("StateATask", TEXT("EnterState0")));
+				AITEST_FALSE(TEXT("Exit StateA"), Exec.Expect("StateATask", TEXT("ExitState0")));
+				Exec.LogClear();
+				TransitionToExecute = ECustomFunctionToRun::None;
 			}
+			// Select another root.
+			{
+				TransitionToExecute = ECustomFunctionToRun::TransitionB_To_P;
+				if (!ResetInstanceData())
+				{
+					return false;
+				}
 
-			FTestStateTreeExecutionContext Exec(StateTree, StateTree, InstanceData);
-			const EStateTreeRunStatus Status = Exec.Tick(0.1f);
-			AITEST_EQUAL(TEXT("Tick should complete with Running"), Status, EStateTreeRunStatus::Running);
-			AITEST_TRUE(TEXT("Start should"), Exec.ExpectInActiveStates("StateA", "StateB", "StateC", "StateD"));
-			FTestStateTreeExecutionContext::FLogOrder LogOrder = Exec.Expect("StateATask", TEXT("Tick0"));
-			AITEST_TRUE(TEXT("Tick StateA"), LogOrder);
-			LogOrder = LogOrder.Then("StateBTask", TEXT("Tick0"));
-			AITEST_TRUE(TEXT("Tick StateB"), LogOrder);
-			LogOrder = LogOrder.Then("StateCTask", TEXT("Tick0"));
-			AITEST_TRUE(TEXT("Tick StateC"), LogOrder);
-			LogOrder = LogOrder.Then("StateDTask", TEXT("Tick0"));
-			AITEST_TRUE(TEXT("Tick StateD"), LogOrder);
-			LogOrder = LogOrder.Then("StateDTask", TEXT("ExitState=Sustained"));
-			AITEST_TRUE(TEXT("Exit StateD"), LogOrder);
-			LogOrder = LogOrder.Then("StateCTask", TEXT("ExitState=Sustained"));
-			AITEST_TRUE(TEXT("Exit StateC"), LogOrder);
-			LogOrder = LogOrder.Then("StateBTask", TEXT("ExitState=Sustained"));
-			AITEST_TRUE(TEXT("Exit StateB"), LogOrder);
-			LogOrder = LogOrder.Then("StateBTask", TEXT("EnterState=Sustained"));
-			AITEST_TRUE(TEXT("Enter StateB"), LogOrder);
-			LogOrder = LogOrder.Then("StateCTask", TEXT("EnterState=Sustained"));
-			AITEST_TRUE(TEXT("Enter StateC"), LogOrder);
-			LogOrder = LogOrder.Then("StateDTask", TEXT("EnterState=Sustained"));
-			AITEST_TRUE(TEXT("Enter StateD"), LogOrder);
-			AITEST_FALSE(TEXT("Enter StateA"), Exec.Expect("StateATask", TEXT("EnterState0")));
-			AITEST_FALSE(TEXT("Exit StateA"), Exec.Expect("StateATask", TEXT("ExitState0")));
-			Exec.LogClear();
-			TransitionToExecute = ECustomFunctionToRun::None;
-		}
-		{
-			TransitionToExecute = ECustomFunctionToRun::TransitionB_To_P;
-			if (!ResetInstanceData())
+				FTestStateTreeExecutionContext Exec(StateTree, StateTree, InstanceData);
+				const EStateTreeRunStatus Status = Exec.Tick(0.1f);
+				AITEST_EQUAL(TEXT("Tick should complete with Running"), Status, EStateTreeRunStatus::Running);
+				AITEST_TRUE(TEXT("Start should"), Exec.ExpectInActiveStates("StateO", "StateP"));
+				FTestStateTreeExecutionContext::FLogOrder LogOrder = Exec.Expect("StateATask", TEXT("Tick0"));
+				AITEST_TRUE(TEXT("Tick StateA"), LogOrder);
+				LogOrder = LogOrder.Then("StateBTask", TEXT("Tick0"));
+				AITEST_TRUE(TEXT("Tick StateB"), LogOrder);
+				LogOrder = LogOrder.Then("StateCTask", TEXT("Tick0"));
+				AITEST_TRUE(TEXT("Tick StateC"), LogOrder);
+				LogOrder = LogOrder.Then("StateDTask", TEXT("Tick0"));
+				AITEST_TRUE(TEXT("Tick StateD"), LogOrder);
+				LogOrder = LogOrder.Then("StateDTask", ExitStateChangedStr);
+				AITEST_TRUE(TEXT("Exit StateD"), LogOrder);
+				LogOrder = LogOrder.Then("StateCTask", ExitStateChangedStr);
+				AITEST_TRUE(TEXT("Exit StateC"), LogOrder);
+				LogOrder = LogOrder.Then("StateBTask", ExitStateChangedStr);
+				AITEST_TRUE(TEXT("Exit StateB"), LogOrder);
+				LogOrder = LogOrder.Then("StateATask", ExitStateChangedStr);
+				AITEST_TRUE(TEXT("Exit StateA"), LogOrder);
+				LogOrder = LogOrder.Then("StateOTask", EnterStateChangedStr);
+				AITEST_TRUE(TEXT("Enter StateO"), LogOrder);
+				LogOrder = LogOrder.Then("StatePTask", EnterStateChangedStr);
+				AITEST_TRUE(TEXT("Enter StateP"), LogOrder);
+				Exec.LogClear();
+				TransitionToExecute = ECustomFunctionToRun::None;
+			}
 			{
-				return false;
+				FTestStateTreeExecutionContext Exec(StateTree, StateTree, InstanceData);
+				EStateTreeRunStatus Status = Exec.Stop();
+				AITEST_EQUAL(TEXT("Should stop"), Status, EStateTreeRunStatus::Stopped);
 			}
-
-			FTestStateTreeExecutionContext Exec(StateTree, StateTree, InstanceData);
-			const EStateTreeRunStatus Status = Exec.Tick(0.1f);
-			AITEST_EQUAL(TEXT("Tick should complete with Running"), Status, EStateTreeRunStatus::Running);
-			AITEST_TRUE(TEXT("Start should"), Exec.ExpectInActiveStates("StateO", "StateP"));
-			FTestStateTreeExecutionContext::FLogOrder LogOrder = Exec.Expect("StateATask", TEXT("Tick0"));
-			AITEST_TRUE(TEXT("Tick StateA"), LogOrder);
-			LogOrder = LogOrder.Then("StateBTask", TEXT("Tick0"));
-			AITEST_TRUE(TEXT("Tick StateB"), LogOrder);
-			LogOrder = LogOrder.Then("StateCTask", TEXT("Tick0"));
-			AITEST_TRUE(TEXT("Tick StateC"), LogOrder);
-			LogOrder = LogOrder.Then("StateDTask", TEXT("Tick0"));
-			AITEST_TRUE(TEXT("Tick StateD"), LogOrder);
-			LogOrder = LogOrder.Then("StateDTask", TEXT("ExitState=Changed"));
-			AITEST_TRUE(TEXT("Exit StateD"), LogOrder);
-			LogOrder = LogOrder.Then("StateCTask", TEXT("ExitState=Changed"));
-			AITEST_TRUE(TEXT("Exit StateC"), LogOrder);
-			LogOrder = LogOrder.Then("StateBTask", TEXT("ExitState=Changed"));
-			AITEST_TRUE(TEXT("Exit StateB"), LogOrder);
-			LogOrder = LogOrder.Then("StateATask", TEXT("ExitState=Changed"));
-			AITEST_TRUE(TEXT("Exit StateA"), LogOrder);
-			LogOrder = LogOrder.Then("StateOTask", TEXT("EnterState=Changed"));
-			AITEST_TRUE(TEXT("Enter StateO"), LogOrder);
-			LogOrder = LogOrder.Then("StatePTask", TEXT("EnterState=Changed"));
-			AITEST_TRUE(TEXT("Enter StateP"), LogOrder);
-			Exec.LogClear();
-			TransitionToExecute = ECustomFunctionToRun::None;
-		}
-		{
-			FTestStateTreeExecutionContext Exec(StateTree, StateTree, InstanceData);
-			EStateTreeRunStatus Status = Exec.Stop();
-			AITEST_EQUAL(TEXT("Should stop"), Status, EStateTreeRunStatus::Stopped);
 		}
 
 		return true;
