[StateTree] Fix linking error on Linux introduced with CL 43343544

#rnx

diff --git a/Engine/Plugins/Runtime/StateTree/Source/StateTreeModule/Private/StateTreeTasksStatus.cpp b/Engine/Plugins/Runtime/StateTree/Source/StateTreeModule/Private/StateTreeTasksStatus.cpp
--- a/Engine/Plugins/Runtime/StateTree/Source/StateTreeModule/Private/StateTreeTasksStatus.cpp
+++ b/Engine/Plugins/Runtime/StateTree/Source/StateTreeModule/Private/StateTreeTasksStatus.cpp
@@ -14,32 +14,6 @@ namespace UE::StateTree::Private
 	}
 }
 
-namespace UE::StateTree
-{
-	template<typename T>
-	void TTasksCompletionStatus<T>::ResetStatus(int32 NumberOfTasksInTheCompletionMask)
-	{
-		T ClearBuffer = 0;
-		if (NumberOfTasksInTheCompletionMask == 0)
-		{
-			// All mask as at least one bit to mark the state.
-			ClearBuffer = 1 << BitIndex;
-		}
-		else if (NumberOfTasksInTheCompletionMask == MaxNumTasks)
-		{
-			// Prevent buffer overflow.
-			ClearBuffer = (T)(-1);
-		}
-		else
-		{
-			ClearBuffer = ((1 << NumberOfTasksInTheCompletionMask) - 1) << BitIndex;
-		}
-		const T ClearMask = ~ClearBuffer;
-		(*FirstCompletionBits) &= ClearMask;
-		(*SecondCompletionBits) &= ClearMask;
-	}
-}
-
 FStateTreeTasksCompletionStatus::FStateTreeTasksCompletionStatus(const FCompactStateTreeFrame& Frame)
 {
 	BufferNum = Frame.NumberOfTasksStatusMasks;
diff --git a/Engine/Plugins/Runtime/StateTree/Source/StateTreeModule/Public/StateTreeTasksStatus.h b/Engine/Plugins/Runtime/StateTree/Source/StateTreeModule/Public/StateTreeTasksStatus.h
--- a/Engine/Plugins/Runtime/StateTree/Source/StateTreeModule/Public/StateTreeTasksStatus.h
+++ b/Engine/Plugins/Runtime/StateTree/Source/StateTreeModule/Public/StateTreeTasksStatus.h
@@ -212,7 +212,27 @@ namespace UE::StateTree
 		}
 
 		/** Set the status of all tasks to running. */
-		STATETREEMODULE_API void ResetStatus(int32 NumberOfTasksInTheCompletionMask);
+		void ResetStatus(int32 NumberOfTasksInTheCompletionMask)
+		{
+			T ClearBuffer = 0;
+			if (NumberOfTasksInTheCompletionMask == 0)
+			{
+				// All mask as at least one bit to mark the state.
+				ClearBuffer = 1 << BitIndex;
+			}
+			else if (NumberOfTasksInTheCompletionMask == MaxNumTasks)
+			{
+				// Prevent buffer overflow.
+				ClearBuffer = (T)(-1);
+			}
+			else
+			{
+				ClearBuffer = ((1 << NumberOfTasksInTheCompletionMask) - 1) << BitIndex;
+			}
+			const T ClearMask = ~ClearBuffer;
+			(*FirstCompletionBits) &= ClearMask;
+			(*SecondCompletionBits) &= ClearMask;
+		}
 
 	private:
 		uint8 GetStatusInternal(int32 Index) const
