[StateTree] Add test for CL 42892807. The delegate binds to the wrong task index.
#rnx

diff --git a/Engine/Plugins/Runtime/StateTree/Source/StateTreeTestSuite/Private/StateTreeDelegateTest.cpp b/Engine/Plugins/Runtime/StateTree/Source/StateTreeTestSuite/Private/StateTreeDelegateTest.cpp
--- a/Engine/Plugins/Runtime/StateTree/Source/StateTreeTestSuite/Private/StateTreeDelegateTest.cpp
+++ b/Engine/Plugins/Runtime/StateTree/Source/StateTreeTestSuite/Private/StateTreeDelegateTest.cpp
@@ -23,9 +23,9 @@ struct FStateTreeTest_Delegate_ConcurrentListeners : FStateTreeTestBase
 		UStateTreeEditorData& EditorData = *Cast<UStateTreeEditorData>(StateTree.EditorData);
 		UStateTreeState& Root = EditorData.AddSubTree(FName("Root"));
 
-		auto& DispatcherTask = Root.AddTask<FTestTask_BroadcastDelegate>(FName("DispatcherTask"));
-		auto& ListenerTaskA = Root.AddTask<FTestTask_ListenDelegate>(FName("ListenerTaskA"));
-		auto& ListenerTaskB = Root.AddTask<FTestTask_ListenDelegate>(FName("ListenerTaskB"));
+		TStateTreeEditorNode<FTestTask_BroadcastDelegate>& DispatcherTask = Root.AddTask<FTestTask_BroadcastDelegate>(FName("DispatcherTask"));
+		TStateTreeEditorNode<FTestTask_ListenDelegate>& ListenerTaskA = Root.AddTask<FTestTask_ListenDelegate>(FName("ListenerTaskA"));
+		TStateTreeEditorNode<FTestTask_ListenDelegate>& ListenerTaskB = Root.AddTask<FTestTask_ListenDelegate>(FName("ListenerTaskB"));
 
 		EditorData.AddPropertyBinding(DispatcherTask, GET_MEMBER_NAME_STRING_CHECKED(FTestTask_BroadcastDelegate_InstanceData, OnTickDelegate), ListenerTaskB, GET_MEMBER_NAME_STRING_CHECKED(FTestTask_ListenDelegate_InstanceData, Listener));
 		EditorData.AddPropertyBinding(DispatcherTask, GET_MEMBER_NAME_STRING_CHECKED(FTestTask_BroadcastDelegate_InstanceData, OnTickDelegate), ListenerTaskA, GET_MEMBER_NAME_STRING_CHECKED(FTestTask_ListenDelegate_InstanceData, Listener));
@@ -77,10 +77,10 @@ struct FStateTreeTest_Delegate_MutuallyExclusiveListeners : FStateTreeTestBase
 		StateA.AddTransition(EStateTreeTransitionTrigger::OnTick, EStateTreeTransitionType::GotoState, &StateB);
 		StateB.AddTransition(EStateTreeTransitionTrigger::OnTick, EStateTreeTransitionType::GotoState, &StateA);
 
-		auto& DispatcherTask = Root.AddTask<FTestTask_BroadcastDelegate>(FName(TEXT("DispatcherTask")));
-		auto& ListenerTaskA0 = StateA.AddTask<FTestTask_ListenDelegate>(FName(TEXT("ListenerTaskA0")));
-		auto& ListenerTaskA1 = StateA.AddTask<FTestTask_ListenDelegate>(FName(TEXT("ListenerTaskA1")));
-		auto& ListenerTaskB = StateB.AddTask<FTestTask_ListenDelegate>(FName(TEXT("ListenerTaskB")));
+		TStateTreeEditorNode<FTestTask_BroadcastDelegate>& DispatcherTask = Root.AddTask<FTestTask_BroadcastDelegate>(FName(TEXT("DispatcherTask")));
+		TStateTreeEditorNode<FTestTask_ListenDelegate>& ListenerTaskA0 = StateA.AddTask<FTestTask_ListenDelegate>(FName(TEXT("ListenerTaskA0")));
+		TStateTreeEditorNode<FTestTask_ListenDelegate>& ListenerTaskA1 = StateA.AddTask<FTestTask_ListenDelegate>(FName(TEXT("ListenerTaskA1")));
+		TStateTreeEditorNode<FTestTask_ListenDelegate>& ListenerTaskB = StateB.AddTask<FTestTask_ListenDelegate>(FName(TEXT("ListenerTaskB")));
 
 		EditorData.AddPropertyBinding(DispatcherTask, GET_MEMBER_NAME_STRING_CHECKED(FTestTask_BroadcastDelegate_InstanceData, OnTickDelegate), ListenerTaskA0, GET_MEMBER_NAME_STRING_CHECKED(FTestTask_ListenDelegate_InstanceData, Listener));
 		EditorData.AddPropertyBinding(DispatcherTask, GET_MEMBER_NAME_STRING_CHECKED(FTestTask_BroadcastDelegate_InstanceData, OnTickDelegate), ListenerTaskA1, GET_MEMBER_NAME_STRING_CHECKED(FTestTask_ListenDelegate_InstanceData, Listener));
@@ -135,8 +135,8 @@ struct FStateTreeTest_Delegate_Transitions : FStateTreeTestBase
 		FStateTreeTransition& TransitionAToB = StateA.AddTransition(EStateTreeTransitionTrigger::OnDelegate, EStateTreeTransitionType::GotoState, &StateB);
 		FStateTreeTransition& TransitionBToA = StateB.AddTransition(EStateTreeTransitionTrigger::OnDelegate, EStateTreeTransitionType::GotoState, &StateA);
 
-		auto& DispatcherTask0 = Root.AddTask<FTestTask_BroadcastDelegate>(FName(TEXT("DispatcherTask0")));
-		auto& DispatcherTask1 = Root.AddTask<FTestTask_BroadcastDelegate>(FName(TEXT("DispatcherTask1")));
+		TStateTreeEditorNode<FTestTask_BroadcastDelegate>& DispatcherTask0 = Root.AddTask<FTestTask_BroadcastDelegate>(FName(TEXT("DispatcherTask0")));
+		TStateTreeEditorNode<FTestTask_BroadcastDelegate>& DispatcherTask1 = Root.AddTask<FTestTask_BroadcastDelegate>(FName(TEXT("DispatcherTask1")));
 		
 		EditorData.AddPropertyBinding(FPropertyBindingPath(DispatcherTask0.ID, GET_MEMBER_NAME_CHECKED(FTestTask_BroadcastDelegate_InstanceData, OnTickDelegate)),  FPropertyBindingPath(TransitionAToB.ID, GET_MEMBER_NAME_CHECKED(FStateTreeTransition, DelegateListener)));
 		EditorData.AddPropertyBinding(FPropertyBindingPath(DispatcherTask1.ID, GET_MEMBER_NAME_CHECKED(FTestTask_BroadcastDelegate_InstanceData, OnTickDelegate)), FPropertyBindingPath(TransitionBToA.ID, GET_MEMBER_NAME_CHECKED(FStateTreeTransition, DelegateListener)));
@@ -178,9 +178,9 @@ struct FStateTreeTest_Delegate_Rebroadcasting : FStateTreeTestBase
 		UStateTreeEditorData& EditorData = *Cast<UStateTreeEditorData>(StateTree.EditorData);
 		UStateTreeState& Root = EditorData.AddSubTree(FName(TEXT("Root")));
 
-		auto& DispatcherTask = Root.AddTask<FTestTask_BroadcastDelegate>(FName(TEXT("DispatcherTask")));
-		auto& RedispatcherTask = Root.AddTask<FTestTask_RebroadcastDelegate>(FName(TEXT("RedispatcherTask")));
-		auto& ListenerTask = Root.AddTask<FTestTask_ListenDelegate>(FName(TEXT("ListenerTask")));
+		TStateTreeEditorNode<FTestTask_BroadcastDelegate>& DispatcherTask = Root.AddTask<FTestTask_BroadcastDelegate>(FName(TEXT("DispatcherTask")));
+		TStateTreeEditorNode<FTestTask_RebroadcastDelegate>& RedispatcherTask = Root.AddTask<FTestTask_RebroadcastDelegate>(FName(TEXT("RedispatcherTask")));
+		TStateTreeEditorNode<FTestTask_ListenDelegate>& ListenerTask = Root.AddTask<FTestTask_ListenDelegate>(FName(TEXT("ListenerTask")));
 
 		EditorData.AddPropertyBinding(DispatcherTask, GET_MEMBER_NAME_STRING_CHECKED(FTestTask_BroadcastDelegate_InstanceData, OnTickDelegate), RedispatcherTask, TEXT("Listener"));
 		EditorData.AddPropertyBinding(RedispatcherTask, TEXT("Dispatcher"), ListenerTask, TEXT("Listener"));
@@ -221,8 +221,8 @@ struct FStateTreeTest_Delegate_SelfRemoval : FStateTreeTestBase
 		UStateTreeEditorData& EditorData = *Cast<UStateTreeEditorData>(StateTree.EditorData);
 		UStateTreeState& Root = EditorData.AddSubTree(FName(TEXT("Root")));
 
-		auto& DispatcherTask = Root.AddTask<FTestTask_BroadcastDelegate>(FName(TEXT("DispatcherTask")));
-		auto& CustomFuncTask = Root.AddTask<FTestTask_CustomFuncOnDelegate>(FName(TEXT("CustomFuncTask")));
+		TStateTreeEditorNode<FTestTask_BroadcastDelegate>& DispatcherTask = Root.AddTask<FTestTask_BroadcastDelegate>(FName(TEXT("DispatcherTask")));
+		TStateTreeEditorNode<FTestTask_CustomFuncOnDelegate>& CustomFuncTask = Root.AddTask<FTestTask_CustomFuncOnDelegate>(FName(TEXT("CustomFuncTask")));
 
 		uint32 TriggersCounter = 0;
 
@@ -274,8 +274,8 @@ struct FStateTreeTest_Delegate_WithoutRemoval : FStateTreeTestBase
 
 		FStateTreeTransition& TransitionAToB = StateA.AddTransition(EStateTreeTransitionTrigger::OnTick, EStateTreeTransitionType::GotoState, &StateB);
 
-		auto& DispatcherTask = Root.AddTask<FTestTask_BroadcastDelegate>(FName(TEXT("DispatcherTask")));
-		auto& ListenerTask = StateA.AddTask<FTestTask_ListenDelegate>(FName(TEXT("ListenerTask")));
+		TStateTreeEditorNode<FTestTask_BroadcastDelegate>& DispatcherTask = Root.AddTask<FTestTask_BroadcastDelegate>(FName(TEXT("DispatcherTask")));
+		TStateTreeEditorNode<FTestTask_ListenDelegate>& ListenerTask = StateA.AddTask<FTestTask_ListenDelegate>(FName(TEXT("ListenerTask")));
 		EditorData.AddPropertyBinding(DispatcherTask, GET_MEMBER_NAME_STRING_CHECKED(FTestTask_BroadcastDelegate_InstanceData, OnTickDelegate), ListenerTask, GET_MEMBER_NAME_STRING_CHECKED(FTestTask_ListenDelegate_InstanceData, Listener));
 		ListenerTask.GetNode().bRemoveOnExit = false;
 
@@ -317,11 +317,11 @@ struct FStateTreeTest_Delegate_GlobalDispatcherAndListener : FStateTreeTestBase
 		UStateTree& StateTree = NewStateTree();
 		UStateTreeEditorData& EditorData = *Cast<UStateTreeEditorData>(StateTree.EditorData);
 		UStateTreeState& Root = EditorData.AddSubTree(FName(TEXT("Root")));
-		auto& RootTask = Root.AddTask<FTestTask_Stand>();
+		TStateTreeEditorNode<FTestTask_Stand>& RootTask = Root.AddTask<FTestTask_Stand>();
 		RootTask.GetNode().TicksToCompletion = 100;
 
-		auto& DispatcherTask = EditorData.AddGlobalTask<FTestTask_BroadcastDelegate>(FName(TEXT("DispatcherTask")));
-		auto& ListenerTask = EditorData.AddGlobalTask<FTestTask_ListenDelegate>(FName(TEXT("ListenerTask")));
+		TStateTreeEditorNode<FTestTask_BroadcastDelegate>& DispatcherTask = EditorData.AddGlobalTask<FTestTask_BroadcastDelegate>(FName(TEXT("DispatcherTask")));
+		TStateTreeEditorNode<FTestTask_ListenDelegate>& ListenerTask = EditorData.AddGlobalTask<FTestTask_ListenDelegate>(FName(TEXT("ListenerTask")));
 		ListenerTask.GetNode().bRemoveOnExit = false;
 
 		EditorData.AddPropertyBinding(DispatcherTask, GET_MEMBER_NAME_STRING_CHECKED(FTestTask_BroadcastDelegate_InstanceData, OnTickDelegate), ListenerTask, GET_MEMBER_NAME_STRING_CHECKED(FTestTask_CustomFuncOnDelegate_InstanceData, Listener));
@@ -362,8 +362,8 @@ struct FStateTreeTest_Delegate_ListeningToDelegateOnExit : FStateTreeTestBase
 		UStateTreeEditorData& EditorData = *Cast<UStateTreeEditorData>(StateTree.EditorData);
 		UStateTreeState& Root = EditorData.AddSubTree(FName(TEXT("Root")));
 
-		auto& DispatcherTask = Root.AddTask<FTestTask_BroadcastDelegate>(FName(TEXT("DispatcherTask")));
-		auto& ListenerTask = Root.AddTask<FTestTask_ListenDelegate>(FName(TEXT("ListenerTask")));
+		TStateTreeEditorNode<FTestTask_BroadcastDelegate>& DispatcherTask = Root.AddTask<FTestTask_BroadcastDelegate>(FName(TEXT("DispatcherTask")));
+		TStateTreeEditorNode<FTestTask_ListenDelegate>& ListenerTask = Root.AddTask<FTestTask_ListenDelegate>(FName(TEXT("ListenerTask")));
 		ListenerTask.GetNode().bRemoveOnExit = false;
 
 		EditorData.AddPropertyBinding(DispatcherTask, GET_MEMBER_NAME_STRING_CHECKED(FTestTask_BroadcastDelegate_InstanceData, OnExitDelegate), ListenerTask, GET_MEMBER_NAME_STRING_CHECKED(FTestTask_CustomFuncOnDelegate_InstanceData, Listener));
@@ -391,6 +391,75 @@ struct FStateTreeTest_Delegate_ListeningToDelegateOnExit : FStateTreeTestBase
 };
 IMPLEMENT_STATE_TREE_INSTANT_TEST(FStateTreeTest_Delegate_ListeningToDelegateOnExit, "System.StateTree.Delegate.ListeningToDelegateOnExit");
 
+/** Test same state and the state index < number of instance data. */
+struct FStateTreeTest_Delegate_ListeningToDelegateOnExit2 : FStateTreeTestBase
+{
+	virtual bool InstantTest() override
+	{
+		UStateTree& StateTree = NewStateTree();
+		UStateTreeEditorData& EditorData = *Cast<UStateTreeEditorData>(StateTree.EditorData);
+		UStateTreeState& Root = EditorData.AddSubTree(FName(TEXT("Root")));
+		UStateTreeState& StateB = Root.AddChildState("StateA")
+			.AddChildState("StateB");
+		UStateTreeState& StateC= Root.AddChildState("StateC");
+		{
+			EditorData.AddGlobalTask<FTestTask_Stand>().GetNode().TicksToCompletion = 100;
+			EditorData.AddGlobalTask<FTestTask_Stand>().GetNode().TicksToCompletion = 100;
+			EditorData.AddGlobalTask<FTestTask_Stand>().GetNode().TicksToCompletion = 100;
+			EditorData.AddGlobalTask<FTestTask_Stand>().GetNode().TicksToCompletion = 100;
+			EditorData.AddGlobalTask<FTestTask_Stand>().GetNode().TicksToCompletion = 100;
+			EditorData.AddGlobalTask<FTestTask_Stand>().GetNode().TicksToCompletion = 100;
+			EditorData.AddGlobalTask<FTestTask_Stand>().GetNode().TicksToCompletion = 100;
+			EditorData.AddGlobalTask<FTestTask_Stand>().GetNode().TicksToCompletion = 100;
+			EditorData.AddGlobalTask<FTestTask_Stand>().GetNode().TicksToCompletion = 100;
+			EditorData.AddGlobalTask<FTestTask_Stand>().GetNode().TicksToCompletion = 100;
+		}
+		{
+			TStateTreeEditorNode<FTestTask_BroadcastDelegate>& DispatcherTask = StateB.AddTask<FTestTask_BroadcastDelegate>(FName(TEXT("DispatcherTask")));
+			TStateTreeEditorNode<FTestTask_ListenDelegate>& ListenerTask = StateB.AddTask<FTestTask_ListenDelegate>(FName(TEXT("ListenerTask")));
+			ListenerTask.GetNode().bRemoveOnExit = false;
+
+			EditorData.AddPropertyBinding(DispatcherTask, GET_MEMBER_NAME_STRING_CHECKED(FTestTask_BroadcastDelegate_InstanceData, OnExitDelegate), ListenerTask, GET_MEMBER_NAME_STRING_CHECKED(FTestTask_CustomFuncOnDelegate_InstanceData, Listener));
+
+			FStateTreeTransition& Transition = StateB.AddTransition(EStateTreeTransitionTrigger::OnTick, EStateTreeTransitionType::GotoState, &StateC);
+			Transition.bDelayTransition = true;
+			Transition.DelayDuration = 0.1f;
+		}
+		{
+			FStateTreeCompilerLog Log;
+			FStateTreeCompiler Compiler(Log);
+			const bool bResult = Compiler.Compile(StateTree);
+			AITEST_TRUE("StateTree should get compiled", bResult);
+		}
+
+		FStateTreeInstanceData InstanceData;
+		FTestStateTreeExecutionContext Exec(StateTree, StateTree, InstanceData);
+		const bool bInitSucceeded = Exec.IsValid();
+		AITEST_TRUE("StateTree should init", bInitSucceeded);
+
+		Exec.Start();
+		AITEST_TRUE("Expected Root to be active.", Exec.ExpectInActiveStates("Root", "StateA", "StateB"));
+		Exec.LogClear();
+
+		Exec.Tick(1.0f);
+		AITEST_TRUE("Expected Root to be active.", Exec.ExpectInActiveStates("Root", "StateA", "StateB"));
+		AITEST_FALSE(TEXT("StateTree Delegate shouldn't be triggered"), Exec.Expect("ListenerTask"));
+		Exec.LogClear();
+
+		Exec.Tick(1.0f);
+		AITEST_TRUE("Expected Root to be active.", Exec.ExpectInActiveStates("Root", "StateC"));
+		AITEST_FALSE(TEXT("StateTree Delegate shouldn't be triggered"), Exec.Expect("ListenerTask"));
+		Exec.LogClear();
+
+		Exec.Stop();
+		AITEST_FALSE(TEXT("StateTree Delegate shouldn't be triggered"), Exec.Expect("ListenerTask"));
+		Exec.LogClear();
+
+		return true;
+	}
+};
+IMPLEMENT_AI_INSTANT_TEST(FStateTreeTest_Delegate_ListeningToDelegateOnExit2, "System.StateTree.Delegate.ListeningToDelegateOnExit2");
+
 } // namespace UE::StateTree::Tests
 
 UE_ENABLE_OPTIMIZATION_SHIP
