[StateTree] Add no_discard to ExecutionContext. Fix bad code formatting introduced with CL 43343544 (result of an auto-merge).

#rnx

diff --git a/Engine/Plugins/Runtime/StateTree/Source/StateTreeModule/Private/StateTreeExecutionContext.cpp b/Engine/Plugins/Runtime/StateTree/Source/StateTreeModule/Private/StateTreeExecutionContext.cpp
--- a/Engine/Plugins/Runtime/StateTree/Source/StateTreeModule/Private/StateTreeExecutionContext.cpp
+++ b/Engine/Plugins/Runtime/StateTree/Source/StateTreeModule/Private/StateTreeExecutionContext.cpp
@@ -103,7 +103,7 @@ namespace UE::StateTree::ExecutionContext::Private
 		if (Frame)
 		{
 			return Frame;
-}
+		}
 
 		return TemporaryFrames.FindByPredicate(
 			[FrameID](const FStateTreeExecutionFrame& Frame)
@@ -240,7 +240,7 @@ namespace UE::StateTree::ExecutionContext
 			}
 
 			const int32 TransitionEnd = State->TransitionsBegin + State->TransitionsNum;
-			for (int32 TransitionIndex = State->TransitionsBegin; TransitionIndex <TransitionEnd; ++TransitionIndex)
+			for (int32 TransitionIndex = State->TransitionsBegin; TransitionIndex < TransitionEnd; ++TransitionIndex)
 			{
 				const FCompactStateTransition* Transition = StateTree->GetTransitionFromIndex(FStateTreeIndex16(TransitionIndex));
 				check(Transition);
@@ -736,7 +736,7 @@ FStateTreeInstanceDebugId FStateTreeReadOnlyExecutionContext::GetInstanceDebugId
 /**
  * FStateTreeMinimalExecutionContext implementation
  */
-// Deprecated
+ // Deprecated
 FStateTreeMinimalExecutionContext::FStateTreeMinimalExecutionContext(UObject& InOwner, const UStateTree& InStateTree, FStateTreeInstanceData& InInstanceData)
 	: FStateTreeMinimalExecutionContext(&InOwner, &InStateTree, InInstanceData.GetMutableStorage())
 {
@@ -880,7 +880,7 @@ FStateTreeExecutionFrame& FStateTreeExecutionContext::FSelectStateResult::MakeAn
 /**
  * FStateTreeExecutionContext::FCurrentlyProcessedFrameScope implementation
  */
-FStateTreeExecutionContext::FCurrentlyProcessedFrameScope::FCurrentlyProcessedFrameScope(FStateTreeExecutionContext& InContext, const FStateTreeExecutionFrame* CurrentParentFrame, const FStateTreeExecutionFrame& CurrentFrame): Context(InContext)
+FStateTreeExecutionContext::FCurrentlyProcessedFrameScope::FCurrentlyProcessedFrameScope(FStateTreeExecutionContext& InContext, const FStateTreeExecutionFrame* CurrentParentFrame, const FStateTreeExecutionFrame& CurrentFrame) : Context(InContext)
 {
 	check(CurrentFrame.StateTree);
 	FStateTreeInstanceStorage* SharedInstanceDataStorage = &CurrentFrame.StateTree->GetSharedInstanceData()->GetMutableStorage();
@@ -891,7 +891,7 @@ FStateTreeExecutionContext::FCurrentlyProcessedFrameScope::FCurrentlyProcessedFr
 	Context.CurrentlyProcessedFrame = &CurrentFrame;
 	Context.CurrentlyProcessedParentFrame = CurrentParentFrame;
 	Context.CurrentlyProcessedSharedInstanceStorage = SharedInstanceDataStorage;
-	
+
 	UE_STATETREE_DEBUG_INSTANCE_FRAME_EVENT(&Context, Context.CurrentlyProcessedFrame);
 }
 
@@ -1074,7 +1074,7 @@ void FStateTreeExecutionContext::SetLinkedStateTreeOverrides(FStateTreeReference
 				bValid = false;
 				break;
 			}
-				
+
 			const bool bIsSameSchema = RootStateTree.GetSchema()->GetClass() == OverrideSchema->GetClass();
 			if (!bIsSameSchema)
 			{
@@ -1097,7 +1097,7 @@ void FStateTreeExecutionContext::SetLinkedStateTreeOverrides(FStateTreeReference
 		LinkedAssetStateTreeOverrides.Reset();
 		bChanged = true;
 	}
-	
+
 	if (bChanged)
 	{
 		TInstancedStruct<FStateTreeExecutionExtension>& ExecutionExtension = Storage.GetMutableExecutionState().ExecutionExtension;
@@ -1126,17 +1126,17 @@ bool FStateTreeExecutionContext::FExternalGlobalParameters::Add(const FPropertyB
 	const int32 TypeHash = HashCombine(GetTypeHash(Copy.SourceLeafProperty), GetTypeHash(Copy.SourceIndirection));
 	const int32 NumMappings = Mappings.Num();
 	Mappings.Add(TypeHash, InParameterMemory);
-	return Mappings.Num() > NumMappings;			
+	return Mappings.Num() > NumMappings;
 }
 
 uint8* FStateTreeExecutionContext::FExternalGlobalParameters::Find(const FPropertyBindingCopyInfo& Copy) const
 {
 	const int32 TypeHash = HashCombine(GetTypeHash(Copy.SourceLeafProperty), GetTypeHash(Copy.SourceIndirection));
-	if(uint8* const* MappingPtr = Mappings.Find(TypeHash))
+	if (uint8* const* MappingPtr = Mappings.Find(TypeHash))
 	{
 		return *MappingPtr;
 	}
-		
+
 	checkf(false, TEXT("Missing external parameter data"));
 	return nullptr;
 }
@@ -1157,9 +1157,9 @@ bool FStateTreeExecutionContext::AreContextDataViewsValid() const
 	{
 		return false;
 	}
-	
+
 	bool bResult = true;
-	
+
 	for (const FStateTreeExternalDataDesc& DataDesc : RootStateTree.GetContextDataDescs())
 	{
 		const FStateTreeDataView& DataView = ContextAndExternalDataViews[DataDesc.Handle.DataHandle.GetIndex()];
@@ -1188,9 +1188,9 @@ bool FStateTreeExecutionContext::AreContextDataViewsValid() const
 bool FStateTreeExecutionContext::SetContextDataByName(const FName Name, FStateTreeDataView DataView)
 {
 	const FStateTreeExternalDataDesc* Desc = RootStateTree.GetContextDataDescs().FindByPredicate([&Name](const FStateTreeExternalDataDesc& Desc)
-	{
-		return Desc.Name == Name;
-	});
+		{
+			return Desc.Name == Name;
+		});
 	if (Desc)
 	{
 		ContextAndExternalDataViews[Desc->Handle.DataHandle.GetIndex()] = DataView;
@@ -1244,7 +1244,7 @@ PRAGMA_ENABLE_DEPRECATION_WARNINGS
 EStateTreeRunStatus FStateTreeExecutionContext::Start(const FInstancedPropertyBag* InitialParameters, int32 RandomSeed)
 {
 	const TOptional<int32> ParamRandomSeed = RandomSeed == -1 ? TOptional<int32>() : RandomSeed;
-	return Start(FStartParameters{.GlobalParameters = InitialParameters, .RandomSeed = ParamRandomSeed });
+	return Start(FStartParameters{ .GlobalParameters = InitialParameters, .RandomSeed = ParamRandomSeed });
 }
 
 void FStateTreeExecutionContext::SetUpdatePhaseInExecutionState(FStateTreeExecutionState& ExecutionState, const EStateTreeUpdatePhase UpdatePhase) const
@@ -1285,7 +1285,7 @@ EStateTreeRunStatus FStateTreeExecutionContext::Start(FStartParameters Parameter
 
 	FStateTreeExecutionState& Exec = GetExecState();
 	if (!ensureMsgf(Exec.CurrentPhase == EStateTreeUpdatePhase::Unset, TEXT("%hs can't be called while already in %s ('%s' using StateTree '%s')."),
-			__FUNCTION__, *UEnum::GetDisplayValueAsText(Exec.CurrentPhase).ToString(), *GetNameSafe(&Owner), *GetFullNameSafe(&RootStateTree)))
+		__FUNCTION__, *UEnum::GetDisplayValueAsText(Exec.CurrentPhase).ToString(), *GetNameSafe(&Owner), *GetFullNameSafe(&RootStateTree)))
 	{
 		return EStateTreeRunStatus::Failed;
 	}
@@ -1328,7 +1328,7 @@ EStateTreeRunStatus FStateTreeExecutionContext::Start(FStartParameters Parameter
 	constexpr bool bIsGlobalFrame = true;
 	FStateTreeExecutionFrame& InitFrame = SelectStateResult.MakeAndAddTemporaryFrame(InitFrameID, InitFrameKey, bIsGlobalFrame);
 	InitFrame.GlobalParameterDataHandle = FStateTreeDataHandle(EStateTreeDataSourceType::GlobalParameterData);
-	
+
 	if (!CollectActiveExternalData(SelectStateResult.TemporaryFrames))
 	{
 		STATETREE_LOG(Warning, TEXT("%hs: Failed to collect external data ('%s' using StateTree '%s')"),
@@ -1427,7 +1427,7 @@ EStateTreeRunStatus FStateTreeExecutionContext::Start(FStartParameters Parameter
 
 		// No active states or global tasks anymore, reset frames.
 		Exec.ActiveFrames.Reset();
-		
+
 		RemoveAllDelegateListeners();
 
 		// We are not considered as running yet so we only set the status without requiring a stop.
@@ -1440,7 +1440,7 @@ EStateTreeRunStatus FStateTreeExecutionContext::Start(FStartParameters Parameter
 
 	// Use local for resulting run state since Stop will reset the instance data.
 	EStateTreeRunStatus Result = Exec.TreeRunStatus;
-	
+
 	if (Exec.RequestedStop != EStateTreeRunStatus::Unset
 		&& Exec.TreeRunStatus == EStateTreeRunStatus::Running)
 	{
@@ -1448,7 +1448,7 @@ EStateTreeRunStatus FStateTreeExecutionContext::Start(FStartParameters Parameter
 		UE_STATETREE_DEBUG_LOG_EVENT(this, Log, TEXT("Processing Deferred Stop"));
 		Result = Stop(Exec.RequestedStop);
 	}
-	
+
 	return Result;
 }
 
@@ -1496,7 +1496,7 @@ EStateTreeRunStatus FStateTreeExecutionContext::Stop(EStateTreeRunStatus Complet
 	SetUpdatePhaseInExecutionState(Exec, EStateTreeUpdatePhase::StopTree);
 
 	EStateTreeRunStatus Result = Exec.TreeRunStatus;
-	
+
 	// Exit states if still in some valid state.
 	if (Exec.TreeRunStatus == EStateTreeRunStatus::Running)
 	{
@@ -1504,7 +1504,7 @@ EStateTreeRunStatus FStateTreeExecutionContext::Stop(EStateTreeRunStatus Complet
 		FSelectStateResult EmptySelectionResult;
 		FStateTreeTransitionResult Transition;
 		Transition.TargetState = FStateTreeStateHandle::FromCompletionStatus(CompletionStatus);
-		Transition.CurrentRunStatus = CompletionStatus;	
+		Transition.CurrentRunStatus = CompletionStatus;
 		ExitState(EmptySelectionResult, Transition);
 
 		// No active states or global tasks anymore, reset frames.
@@ -1554,7 +1554,7 @@ EStateTreeRunStatus FStateTreeExecutionContext::TickPrelude()
 	}
 
 	if (!ensureMsgf(Exec.CurrentPhase == EStateTreeUpdatePhase::Unset, TEXT("%hs can't be called while already in %s ('%s' using StateTree '%s')."),
-			__FUNCTION__, *UEnum::GetDisplayValueAsText(Exec.CurrentPhase).ToString(), *GetNameSafe(&Owner), *GetFullNameSafe(&RootStateTree)))
+		__FUNCTION__, *UEnum::GetDisplayValueAsText(Exec.CurrentPhase).ToString(), *GetNameSafe(&Owner), *GetFullNameSafe(&RootStateTree)))
 	{
 		return EStateTreeRunStatus::Failed;
 	}
@@ -1575,7 +1575,7 @@ EStateTreeRunStatus FStateTreeExecutionContext::TickPostlude()
 
 	// Use local for resulting run state since Stop will reset the instance data.
 	EStateTreeRunStatus Result = Exec.TreeRunStatus;
-	
+
 	if (Exec.RequestedStop != EStateTreeRunStatus::Unset)
 	{
 		STATETREE_LOG(VeryVerbose, TEXT("Processing Deferred Stop"));
@@ -1620,10 +1620,10 @@ EStateTreeRunStatus FStateTreeExecutionContext::TickUpdateTasks(const float Delt
 	}
 
 	TickUpdateTasksInternal(DeltaTime);
-	
+
 	return TickPostlude();
 }
-	
+
 EStateTreeRunStatus FStateTreeExecutionContext::TickTriggerTransitions()
 {
 	CSV_SCOPED_TIMING_STAT_EXCLUSIVE(StateTree_Tick);
@@ -1651,7 +1651,7 @@ void FStateTreeExecutionContext::TickUpdateTasksInternal(float DeltaTime)
 	{
 		return;
 	}
-	
+
 	// Prevent wrong user input.
 	DeltaTime = FMath::Max(0.f, DeltaTime);
 
@@ -1707,18 +1707,18 @@ void FStateTreeExecutionContext::TickUpdateTasksInternal(float DeltaTime)
 				// Only set RequestStop if it's the first frame (root)
 				if (Exec.ActiveFrames.Num() > 0)
 				{
-				const UStateTree* StateTree = Exec.ActiveFrames[0].StateTree;
-				check(StateTree == &RootStateTree);
-				const ETaskCompletionStatus GlobalTaskStatus = Exec.ActiveFrames[0].ActiveTasksStatus.GetStatus(StateTree).GetCompletionStatus();
-				const EStateTreeRunStatus GlobalRunStatus = ExecutionContext::CastToRunStatus(GlobalTaskStatus);
-				if (GlobalRunStatus != EStateTreeRunStatus::Running)
-				{
-					Exec.RequestedStop = ExecutionContext::GetPriorityRunStatus(Exec.RequestedStop, GlobalRunStatus);
-					LogRequestStop();
+					const UStateTree* StateTree = Exec.ActiveFrames[0].StateTree;
+					check(StateTree == &RootStateTree);
+					const ETaskCompletionStatus GlobalTaskStatus = Exec.ActiveFrames[0].ActiveTasksStatus.GetStatus(StateTree).GetCompletionStatus();
+					const EStateTreeRunStatus GlobalRunStatus = ExecutionContext::CastToRunStatus(GlobalTaskStatus);
+					if (GlobalRunStatus != EStateTreeRunStatus::Running)
+					{
+						Exec.RequestedStop = ExecutionContext::GetPriorityRunStatus(Exec.RequestedStop, GlobalRunStatus);
+						LogRequestStop();
+					}
 				}
-			}
-			else
-			{
+				else
+				{
 					// Start failed and the user called Tick anyway.
 					Exec.RequestedStop = EStateTreeRunStatus::Failed;
 					LogRequestStop();
@@ -1774,12 +1774,12 @@ void FStateTreeExecutionContext::TickTriggerTransitionsInternal()
 				// Stop evaluators and global tasks (handled in ExitState)
 				if (!ensure(Exec.ActiveFrames.Num() == 0))
 				{
-				StopEvaluatorsAndGlobalTasks(Exec.TreeRunStatus);
+					StopEvaluatorsAndGlobalTasks(Exec.TreeRunStatus);
 
-				// No active states or global tasks anymore, reset frames.
-				Exec.ActiveFrames.Reset();
+					// No active states or global tasks anymore, reset frames.
+					Exec.ActiveFrames.Reset();
 
-				RemoveAllDelegateListeners();
+					RemoveAllDelegateListeners();
 				}
 
 				break;
@@ -1913,7 +1913,7 @@ void FStateTreeExecutionContext::RequestTransition(const FStateTreeTransitionReq
 	if (bAllowDirectTransitions)
 	{
 		checkf(CurrentlyProcessedFrame, TEXT("Expecting CurrentlyProcessedFrame to be valid when called during TriggerTransitions()."));
-		
+
 		STATETREE_LOG(Verbose, TEXT("Request transition to '%s' at priority %s"), *GetSafeStateName(*CurrentlyProcessedFrame, Request.TargetState), *UEnum::GetDisplayValueAsText(Request.Priority).ToString());
 
 		const FActiveStateID CurrentlyProcessedStateID = CurrentlyProcessedFrame->ActiveStates.FindStateID(CurrentlyProcessedState);
@@ -1939,7 +1939,7 @@ void FStateTreeExecutionContext::RequestTransition(const FStateTreeTransitionReq
 				__FUNCTION__, *GetNameSafe(&Owner), *GetFullNameSafe(&RootStateTree));
 			return;
 		}
-		
+
 		STATETREE_LOG(Verbose, TEXT("Request transition to '%s' at priority %s"), *GetSafeStateName(*RootFrame, Request.TargetState), *UEnum::GetDisplayValueAsText(Request.Priority).ToString());
 
 		FStateTreeTransitionRequest RequestWithSource = Request;
@@ -1995,7 +1995,7 @@ void FStateTreeExecutionContext::FinishTask(const FStateTreeTaskBase& Task, ESta
 
 		const int32 ActiveStateIndex = CurrentlyProcessedFrame->ActiveStates.IndexOfReverse(CurrentlyProcessedState);
 		check(ActiveStateIndex != INDEX_NONE);
-		
+
 		const int32 StateTaskIndex = CurrentNodeIndex - State.TasksBegin;
 		check(StateTaskIndex >= 0);
 
@@ -2111,26 +2111,26 @@ void FStateTreeExecutionContext::UpdateInstanceData(FSelectStateResult& Args)
 
 				const FStateTreeExecutionFrame* Frame = FindExecutionFrame(CurrentFrameID, MakeArrayView(Exec.ActiveFrames), MakeArrayView(Args.TemporaryFrames));
 				if (ensure(Frame) && Frame->bIsGlobalFrame)
-		{
+				{
 					check(Frame->StateTree);
 
 					CurrentStateTree = Frame->StateTree;
 					EstimatedNumStructs += CurrentStateTree->NumGlobalInstanceData;
 				}
-		}
+			}
 
 			// State
 			if (ensure(CurrentStateTree))
 			{
 				const FCompactStateTreeState* CurrentState = CurrentStateTree->GetStateFromHandle(SelectedState.GetStateHandle());
 				if (ensure(CurrentState))
-		{
+				{
 					EstimatedNumStructs += CurrentState->InstanceDataNum;
 				}
 			}
 		}
 	}
-	
+
 	TArray<FConstStructView> InstanceStructs;
 	InstanceStructs.Reserve(EstimatedNumStructs);
 
@@ -2142,14 +2142,13 @@ void FStateTreeExecutionContext::UpdateInstanceData(FSelectStateResult& Args)
 
 	TArrayView<FStateTreeTemporaryInstanceData> TempInstances = Storage.GetMutableTemporaryInstances();
 	auto FindInstanceTempData = [&TempInstances](FActiveFrameID FrameID, FStateTreeDataHandle DataHandle)
-	{
-			FStateTreeTemporaryInstanceData* TempData = TempInstances.FindByPredicate([&FrameID, &DataHandle](const FStateTreeTemporaryInstanceData& Data)
 		{
+			FStateTreeTemporaryInstanceData* TempData = TempInstances.FindByPredicate([&FrameID, &DataHandle](const FStateTreeTemporaryInstanceData& Data)
+				{
 					return Data.FrameID == FrameID && Data.DataHandle == DataHandle;
-		});
-		return TempData ? &TempData->Instance : nullptr;
-	};
-	
+				});
+			return TempData ? &TempData->Instance : nullptr;
+		};
 
 	// Find next instance data sources and find common/existing section of instance data at start.
 	int32 CurrentGlobalInstanceIndexBase = 0;
@@ -2157,7 +2156,7 @@ void FStateTreeExecutionContext::UpdateInstanceData(FSelectStateResult& Args)
 
 	const UStruct* NextStateParameterDataStruct = nullptr;
 	FStateTreeDataHandle NextStateParameterDataHandle = FStateTreeDataHandle::Invalid;
-	
+
 	FStateTreeDataHandle CurrentGlobalParameterDataHandle = FStateTreeDataHandle(EStateTreeDataSourceType::GlobalParameterData);
 	{
 		FActiveFrameID CurrentFrameID;
@@ -2170,7 +2169,7 @@ void FStateTreeExecutionContext::UpdateInstanceData(FSelectStateResult& Args)
 
 			// Global
 			if (SelectedState.GetFrameID() != CurrentFrameID)
-	{
+			{
 				CurrentFrameID = SelectedState.GetFrameID();
 				CurrentFrame = Exec.FindActiveFrame(CurrentFrameID);
 
@@ -2190,51 +2189,51 @@ void FStateTreeExecutionContext::UpdateInstanceData(FSelectStateResult& Args)
 
 				// Global Nodes
 				if (CurrentFrame->bIsGlobalFrame)
-		{
-			// Handle global tree parameters
-			if (NextStateParameterDataHandle.IsValid())
-			{
-				// Point to the parameter block set by linked state.
+				{
+					// Handle global tree parameters
+					if (NextStateParameterDataHandle.IsValid())
+					{
+						// Point to the parameter block set by linked state.
 						check(NextStateParameterDataStruct == CurrentStateTree->GetDefaultParameters().GetPropertyBagStruct());
-				CurrentGlobalParameterDataHandle = NextStateParameterDataHandle;
-				NextStateParameterDataHandle = FStateTreeDataHandle::Invalid; // Mark as used.
-			}
-			
-			const int32 BaseIndex = InstanceStructs.Num();
-			CurrentGlobalInstanceIndexBase = BaseIndex;
-			
+						CurrentGlobalParameterDataHandle = NextStateParameterDataHandle;
+						NextStateParameterDataHandle = FStateTreeDataHandle::Invalid; // Mark as used.
+					}
+
+					const int32 BaseIndex = InstanceStructs.Num();
+					CurrentGlobalInstanceIndexBase = BaseIndex;
+
 					InstanceStructs.AddDefaulted(CurrentStateTree->NumGlobalInstanceData);
 					TempInstanceStructs.AddZeroed(CurrentStateTree->NumGlobalInstanceData);
-			
+
 					// Global Evals
 					for (int32 EvalIndex = CurrentStateTree->EvaluatorsBegin; EvalIndex < (CurrentStateTree->EvaluatorsBegin + CurrentStateTree->EvaluatorsNum); EvalIndex++)
-			{
+					{
 						const FStateTreeEvaluatorBase& Eval = CurrentStateTree->Nodes[EvalIndex].Get<const FStateTreeEvaluatorBase>();
 						const FConstStructView EvalInstanceData = CurrentStateTree->DefaultInstanceData.GetStruct(Eval.InstanceTemplateIndex.Get());
-				InstanceStructs[BaseIndex + Eval.InstanceDataHandle.GetIndex()] = EvalInstanceData;
+						InstanceStructs[BaseIndex + Eval.InstanceDataHandle.GetIndex()] = EvalInstanceData;
 						if (!bIsFrameCommon)
-				{
+						{
 							TempInstanceStructs[BaseIndex + Eval.InstanceDataHandle.GetIndex()] = FindInstanceTempData(CurrentFrameID, Eval.InstanceDataHandle);
-				}
-			}
+						}
+					}
 
-			// Global tasks
+					// Global tasks
 					for (int32 TaskIndex = CurrentStateTree->GlobalTasksBegin; TaskIndex < (CurrentStateTree->GlobalTasksBegin + CurrentStateTree->GlobalTasksNum); TaskIndex++)
-			{
+					{
 						const FStateTreeTaskBase& Task = CurrentStateTree->Nodes[TaskIndex].Get<const FStateTreeTaskBase>();
 						const FConstStructView TaskInstanceData = CurrentStateTree->DefaultInstanceData.GetStruct(Task.InstanceTemplateIndex.Get());
-				InstanceStructs[BaseIndex + Task.InstanceDataHandle.GetIndex()] = TaskInstanceData;
+						InstanceStructs[BaseIndex + Task.InstanceDataHandle.GetIndex()] = TaskInstanceData;
 						if (!bIsFrameCommon)
-				{
+						{
 							TempInstanceStructs[BaseIndex + Task.InstanceDataHandle.GetIndex()] = FindInstanceTempData(CurrentFrameID, Task.InstanceDataHandle);
-				}
-			}
+						}
+					}
 
 					if (bIsFrameCommon)
-			{
-				NumCommonInstanceData = InstanceStructs.Num();
-			}
-		}
+					{
+						NumCommonInstanceData = InstanceStructs.Num();
+					}
+				}
 
 				CurrentFrameBaseIndex = InstanceStructs.Num();
 
@@ -2249,7 +2248,7 @@ void FStateTreeExecutionContext::UpdateInstanceData(FSelectStateResult& Args)
 				const FStateTreeStateHandle CurrentStateHandle = SelectedState.GetStateHandle();
 				const FCompactStateTreeState* CurrentState = CurrentStateTree->GetStateFromHandle(CurrentStateHandle);
 				if (ensure(CurrentState))
-		{
+				{
 					// Find if the state is common (before the transition or sustained from the transition).
 					//The frame may contain valid states (sustained or not). SelectedState will be added later.
 					const bool bIsStateCommon = CurrentFrame->ActiveStates.Contains(SelectedState.GetStateID());
@@ -2257,120 +2256,120 @@ void FStateTreeExecutionContext::UpdateInstanceData(FSelectStateResult& Args)
 					InstanceStructs.AddDefaulted(CurrentState->InstanceDataNum);
 					TempInstanceStructs.AddZeroed(CurrentState->InstanceDataNum);
 
-			bool bCanHaveTempData = false;
+					bool bCanHaveTempData = false;
 					if (CurrentState->Type == EStateTreeStateType::Subtree)
-			{
+					{
 						check(CurrentState->ParameterDataHandle.IsValid());
 						check(CurrentState->ParameterTemplateIndex.IsValid());
 						const FConstStructView ParamsInstanceData = CurrentStateTree->DefaultInstanceData.GetStruct(CurrentState->ParameterTemplateIndex.Get());
-				if (!NextStateParameterDataHandle.IsValid())
-				{
-					// Parameters are not set by a linked state, create instance data.
+						if (!NextStateParameterDataHandle.IsValid())
+						{
+							// Parameters are not set by a linked state, create instance data.
 							InstanceStructs[CurrentFrameBaseIndex + CurrentState->ParameterDataHandle.GetIndex()] = ParamsInstanceData;
 							CurrentFrame->StateParameterDataHandle = CurrentState->ParameterDataHandle;
-					bCanHaveTempData = true;
-				}
-				else
-				{
-					// Point to the parameter block set by linked state.
-					const FCompactStateTreeParameters* Params = ParamsInstanceData.GetPtr<const FCompactStateTreeParameters>();
-					const UStruct* StateParameterDataStruct = Params ? Params->Parameters.GetPropertyBagStruct() : nullptr;
-					check(NextStateParameterDataStruct == StateParameterDataStruct);
-					
+							bCanHaveTempData = true;
+						}
+						else
+						{
+							// Point to the parameter block set by linked state.
+							const FCompactStateTreeParameters* Params = ParamsInstanceData.GetPtr<const FCompactStateTreeParameters>();
+							const UStruct* StateParameterDataStruct = Params ? Params->Parameters.GetPropertyBagStruct() : nullptr;
+							check(NextStateParameterDataStruct == StateParameterDataStruct);
+
 							CurrentFrame->StateParameterDataHandle = NextStateParameterDataHandle;
-					NextStateParameterDataHandle = FStateTreeDataHandle::Invalid; // Mark as used.
+							NextStateParameterDataHandle = FStateTreeDataHandle::Invalid; // Mark as used.
 
-					// This state will not instantiate parameter data, so we don't care about the temp data either.
-					bCanHaveTempData = false;
-				}
-			}
-			else
-			{
+							// This state will not instantiate parameter data, so we don't care about the temp data either.
+							bCanHaveTempData = false;
+						}
+					}
+					else
+					{
 						if (CurrentState->ParameterTemplateIndex.IsValid())
-				{
-					// Linked state's instance data is the parameters.
+						{
+							// Linked state's instance data is the parameters.
 							check(CurrentState->ParameterDataHandle.IsValid());
 
-					const FCompactStateTreeParameters* Params = nullptr;
+							const FCompactStateTreeParameters* Params = nullptr;
 							if (FInstancedStruct* TempParamsInstanceData = FindInstanceTempData(CurrentFrameID, CurrentState->ParameterDataHandle))
-					{
-						// If we have temp data for the parameters, then setup the instance data with just a type, so that we can steal the temp data below (TempInstanceStructs).
-						// We expect overridden linked assets to hit this code path. 
+							{
+								// If we have temp data for the parameters, then setup the instance data with just a type, so that we can steal the temp data below (TempInstanceStructs).
+								// We expect overridden linked assets to hit this code path. 
 								InstanceStructs[CurrentFrameBaseIndex + CurrentState->ParameterDataHandle.GetIndex()] = FConstStructView(TempParamsInstanceData->GetScriptStruct());
-						Params = TempParamsInstanceData->GetPtr<const FCompactStateTreeParameters>();
-						bCanHaveTempData = true;
-					}
-					else
-					{
-						// If not temp data, use the states or linked assets default values.
-						FConstStructView ParamsInstanceData;
+								Params = TempParamsInstanceData->GetPtr<const FCompactStateTreeParameters>();
+								bCanHaveTempData = true;
+							}
+							else
+							{
+								// If not temp data, use the states or linked assets default values.
+								FConstStructView ParamsInstanceData;
 								if (CurrentState->Type == EStateTreeStateType::LinkedAsset)
-						{
+								{
 									// This state is a container for the linked state tree.
 									// Its instance data matches the linked state tree parameters.
 									// The linked state tree asset is the next frame.
 									const bool bIsLastFrame = SelectedStateIndex == Args.SelectedStates.Num() - 1;
-							if (!bIsLastFrame)
-							{
+									if (!bIsLastFrame)
+									{
 										const FActiveState& FollowingSelectedState = Args.SelectedStates[SelectedStateIndex + 1];
 										const FStateTreeExecutionFrame* FollowingNextFrame = FindExecutionFrame(FollowingSelectedState.GetFrameID(), MakeArrayView(Exec.ActiveFrames), MakeArrayView(Args.TemporaryFrames));
 										if (ensure(FollowingNextFrame))
 										{
 											ParamsInstanceData = FConstStructView::Make(TempParams.Emplace_GetRef(FollowingNextFrame->StateTree->GetDefaultParameters()));
 										}
-							}
-						}
-						if (!ParamsInstanceData.IsValid())
-						{
+									}
+								}
+								if (!ParamsInstanceData.IsValid())
+								{
 									ParamsInstanceData = CurrentStateTree->DefaultInstanceData.GetStruct(CurrentState->ParameterTemplateIndex.Get());
-						}
+								}
 								InstanceStructs[CurrentFrameBaseIndex + CurrentState->ParameterDataHandle.GetIndex()] = ParamsInstanceData;
-						Params = ParamsInstanceData.GetPtr<const FCompactStateTreeParameters>();
-						bCanHaveTempData = true;
-					}
+								Params = ParamsInstanceData.GetPtr<const FCompactStateTreeParameters>();
+								bCanHaveTempData = true;
+							}
 
 							if (CurrentState->Type == EStateTreeStateType::Linked
 								|| CurrentState->Type == EStateTreeStateType::LinkedAsset)
-					{
-						// Store the index of the parameter data, so that we can point the linked state to it.
+							{
+								// Store the index of the parameter data, so that we can point the linked state to it.
 								check(CurrentState->ParameterDataHandle.GetSource() == EStateTreeDataSourceType::StateParameterData);
-						checkf(!NextStateParameterDataHandle.IsValid(), TEXT("NextStateParameterDataIndex not should be set yet when we encounter a linked state."));
+								checkf(!NextStateParameterDataHandle.IsValid(), TEXT("NextStateParameterDataIndex not should be set yet when we encounter a linked state."));
 								NextStateParameterDataHandle = CurrentState->ParameterDataHandle;
-						NextStateParameterDataStruct = Params ? Params->Parameters.GetPropertyBagStruct() : nullptr;
+								NextStateParameterDataStruct = Params ? Params->Parameters.GetPropertyBagStruct() : nullptr;
+							}
+						}
 					}
-				}
-			}
-			
+
 					if (!bIsStateCommon && bCanHaveTempData)
-			{
+					{
 						TempInstanceStructs[CurrentFrameBaseIndex + CurrentState->ParameterDataHandle.GetIndex()] = FindInstanceTempData(CurrentFrameID, CurrentState->ParameterDataHandle);
-			}
+					}
 
 					if (CurrentState->EventDataIndex.IsValid())
-			{
+					{
 						InstanceStructs[CurrentFrameBaseIndex + CurrentState->EventDataIndex.Get()] = FConstStructView(FStateTreeSharedEvent::StaticStruct());
-			}
+					}
 
 					for (int32 TaskIndex = CurrentState->TasksBegin; TaskIndex < (CurrentState->TasksBegin + CurrentState->TasksNum); TaskIndex++)
-			{
+					{
 						const FStateTreeTaskBase& Task = CurrentStateTree->Nodes[TaskIndex].Get<const FStateTreeTaskBase>();
 						const FConstStructView TaskInstanceData = CurrentStateTree->DefaultInstanceData.GetStruct(Task.InstanceTemplateIndex.Get());
 						InstanceStructs[CurrentFrameBaseIndex + Task.InstanceDataHandle.GetIndex()] = TaskInstanceData;
 						if (!bIsStateCommon)
-				{
+						{
 							TempInstanceStructs[CurrentFrameBaseIndex + Task.InstanceDataHandle.GetIndex()] = FindInstanceTempData(CurrentFrameID, Task.InstanceDataHandle);
-				}
-			}
+						}
+					}
 
 					if (bIsStateCommon)
-			{
-				NumCommonInstanceData = InstanceStructs.Num();
+					{
+						NumCommonInstanceData = InstanceStructs.Num();
+					}
+				}
 			}
 		}
 	}
-		}
-	}
-	
+
 	// Common section should match.
 #if WITH_STATETREE_DEBUG
 	for (int32 Index = 0; Index < NumCommonInstanceData; Index++)
@@ -2378,7 +2377,7 @@ void FStateTreeExecutionContext::UpdateInstanceData(FSelectStateResult& Args)
 		check(Index < InstanceData.Num());
 
 		FConstStructView ExistingInstanceDataView = InstanceData.GetStruct(Index);
-		FConstStructView NewInstanceDataView = InstanceStructs[Index]; 
+		FConstStructView NewInstanceDataView = InstanceStructs[Index];
 
 		check(NewInstanceDataView.GetScriptStruct() == ExistingInstanceDataView.GetScriptStruct());
 
@@ -2408,29 +2407,29 @@ FStateTreeDataView FStateTreeExecutionContext::GetDataView(const FStateTreeExecu
 	switch (Handle.GetSource())
 	{
 	case EStateTreeDataSourceType::ContextData:
-		check(!ContextAndExternalDataViews.IsEmpty())
+		check(!ContextAndExternalDataViews.IsEmpty());
 		return ContextAndExternalDataViews[Handle.GetIndex()];
 
 	case EStateTreeDataSourceType::ExternalData:
-		check(!ContextAndExternalDataViews.IsEmpty())
+		check(!ContextAndExternalDataViews.IsEmpty());
 		return ContextAndExternalDataViews[CurrentFrame.ExternalDataBaseIndex.Get() + Handle.GetIndex()];
 
 	case EStateTreeDataSourceType::TransitionEvent:
+	{
+		if (CurrentlyProcessedTransitionEvent)
 		{
-			if (CurrentlyProcessedTransitionEvent)
-			{
-				// const_cast because events are read only, but we cannot express that in FStateTreeDataView.
-				return FStateTreeDataView(FStructView::Make(*const_cast<FStateTreeEvent*>(CurrentlyProcessedTransitionEvent)));
-			}
-
-			return nullptr;
+			// const_cast because events are read only, but we cannot express that in FStateTreeDataView.
+			return FStateTreeDataView(FStructView::Make(*const_cast<FStateTreeEvent*>(CurrentlyProcessedTransitionEvent)));
 		}
 
+		return nullptr;
+	}
+
 	case EStateTreeDataSourceType::StateEvent:
-		{
-			// If state selection is going, return FStateTreeEvent of the event currently captured by the state selection.
+	{
+		// If state selection is going, return FStateTreeEvent of the event currently captured by the state selection.
 		if (CurrentlyProcessedStateSelectionResult)
-			{
+		{
 			FSelectionEventWithID* FoundEvent = CurrentlyProcessedStateSelectionResult->SelectionEvents.FindByPredicate(
 				[FrameID = CurrentFrame.FrameID, StateHandle = Handle.GetState()]
 				(const FSelectionEventWithID& Event)
@@ -2438,23 +2437,21 @@ FStateTreeDataView FStateTreeExecutionContext::GetDataView(const FStateTreeExecu
 					return Event.State.GetStateHandle() == StateHandle && Event.State.GetFrameID() == FrameID;
 				});
 			if (FoundEvent)
-				{
-					// Events are read only, but we cannot express that in FStateTreeDataView.
+			{
+				// Events are read only, but we cannot express that in FStateTreeDataView.
 				return FStateTreeDataView(FStructView::Make(*FoundEvent->Event.GetMutable()));
-					}
-				}
+			}
+		}
 
 		if (UE::StateTree::InstanceData::Private::IsActiveInstanceHandleSourceValid(Storage, CurrentFrame, Handle))
 		{
 			return UE::StateTree::InstanceData::GetDataView(Storage, CurrentlyProcessedSharedInstanceStorage, ParentFrame, CurrentFrame, Handle);
-			}
 		}
+	}
 
 	case EStateTreeDataSourceType::ExternalGlobalParameterData:
-		{
-			checkf(false, TEXT("External global parameter data currently not supported for linked state-trees"));
+		checkf(false, TEXT("External global parameter data currently not supported for linked state-trees"));
 		return nullptr;
-		}
 
 	default:
 		return UE::StateTree::InstanceData::GetDataView(Storage, CurrentlyProcessedSharedInstanceStorage, ParentFrame, CurrentFrame, Handle);
@@ -2462,7 +2459,7 @@ FStateTreeDataView FStateTreeExecutionContext::GetDataView(const FStateTreeExecu
 }
 
 FStateTreeDataView FStateTreeExecutionContext::GetDataView(const FStateTreeExecutionFrame* ParentFrame, const FStateTreeExecutionFrame& CurrentFrame, const FPropertyBindingCopyInfo& CopyInfo)
-{	
+{
 	const FStateTreeDataHandle Handle = CopyInfo.SourceDataHandle.Get<FStateTreeDataHandle>();
 	if (Handle.GetSource() == EStateTreeDataSourceType::ExternalGlobalParameterData)
 	{
@@ -2572,9 +2569,9 @@ EStateTreeRunStatus FStateTreeExecutionContext::ForceTransition(const FRecordedS
 const FStateTreeExecutionFrame* FStateTreeExecutionContext::FindFrame(const UStateTree* StateTree, FStateTreeStateHandle RootState, TConstArrayView<FStateTreeExecutionFrame> Frames, const FStateTreeExecutionFrame*& OutParentFrame)
 {
 	const int32 FrameIndex = Frames.IndexOfByPredicate([StateTree, RootState](const FStateTreeExecutionFrame& Frame)
-	{
-		return Frame.HasRoot(StateTree, RootState);
-	});
+		{
+			return Frame.HasRoot(StateTree, RootState);
+		});
 
 	if (FrameIndex == INDEX_NONE)
 	{
@@ -2641,7 +2638,7 @@ FStateTreeDataView FStateTreeExecutionContext::GetDataViewOrTemporary(const FSta
 	{
 		return GetDataView(ParentFrame, CurrentFrame, Handle);
 	}
-	
+
 	return GetTemporaryDataView(ParentFrame, CurrentFrame, Handle);
 }
 
@@ -2651,10 +2648,10 @@ FStateTreeDataView FStateTreeExecutionContext::GetDataViewOrTemporary(const FSta
 	if (Handle.GetSource() == EStateTreeDataSourceType::ExternalGlobalParameterData)
 	{
 		uint8* MemoryPtr = ExternalGlobalParameters->Find(CopyInfo);
-		return FStateTreeDataView(CopyInfo.SourceStructType, MemoryPtr);	
+		return FStateTreeDataView(CopyInfo.SourceStructType, MemoryPtr);
 	}
 
-	return GetDataViewOrTemporary(ParentFrame, CurrentFrame, Handle);	
+	return GetDataViewOrTemporary(ParentFrame, CurrentFrame, Handle);
 }
 
 FStateTreeDataView FStateTreeExecutionContext::GetTemporaryDataView(const FStateTreeExecutionFrame* ParentFrame,
@@ -2807,7 +2804,7 @@ bool FStateTreeExecutionContext::CopyBatchWithValidation(const FStateTreeExecuti
 			bSucceed = false;
 			break;
 		}
-		
+
 		bSucceed &= CurrentFrame.StateTree->PropertyBindings.Super::CopyProperty(Copy, SourceView, TargetView);
 	}
 	return bSucceed;
@@ -2827,7 +2824,7 @@ bool FStateTreeExecutionContext::CollectActiveExternalData(const TArrayView<FSta
 
 	bool bAllExternalDataValid = true;
 	const FStateTreeExecutionFrame* PrevFrame = nullptr;
-	
+
 	for (FStateTreeExecutionFrame& Frame : Frames)
 	{
 		if (PrevFrame && PrevFrame->StateTree == Frame.StateTree)
@@ -2843,7 +2840,7 @@ bool FStateTreeExecutionContext::CollectActiveExternalData(const TArrayView<FSta
 		{
 			bAllExternalDataValid = false;
 		}
-		
+
 		PrevFrame = &Frame;
 	}
 
@@ -2851,7 +2848,7 @@ bool FStateTreeExecutionContext::CollectActiveExternalData(const TArrayView<FSta
 	{
 		bActiveExternalDataCollected = true;
 	}
-	
+
 	return bAllExternalDataValid;
 }
 
@@ -2870,7 +2867,7 @@ FStateTreeIndex16 FStateTreeExecutionContext::CollectExternalData(const UStateTr
 			return Cache.BaseIndex;
 		}
 	}
-	
+
 	const TConstArrayView<FStateTreeExternalDataDesc> ExternalDataDescs = StateTree->GetExternalDataDescs();
 	const int32 BaseIndex = ContextAndExternalDataViews.Num();
 	const int32 NumDescs = ExternalDataDescs.Num();
@@ -2879,7 +2876,7 @@ FStateTreeIndex16 FStateTreeExecutionContext::CollectExternalData(const UStateTr
 	if (NumDescs > 0)
 	{
 		ContextAndExternalDataViews.AddDefaulted(NumDescs);
-		const TArrayView<FStateTreeDataView> DataViews = MakeArrayView(ContextAndExternalDataViews.GetData() + BaseIndex, NumDescs);  
+		const TArrayView<FStateTreeDataView> DataViews = MakeArrayView(ContextAndExternalDataViews.GetData() + BaseIndex, NumDescs);
 
 		if (ensureMsgf(CollectExternalDataDelegate.IsBound(), TEXT("The StateTree asset has external data, expecting CollectExternalData delegate to be provided.")))
 		{
@@ -2980,17 +2977,17 @@ void FStateTreeExecutionContext::CaptureNewStateEvents(FSelectStateResult& Args)
 		check(CurrentFrame);
 		const bool bIsStateCommon = CurrentFrame->ActiveStates.Contains(SelectedState.GetStateID());
 		if (bIsStateCommon)
-			{
+		{
 			continue;
 		}
 
 		if (const FCompactStateTreeState* State = CurrentFrame->StateTree->GetStateFromHandle(SelectedState.GetStateHandle()))
-			{
-			if (State->EventDataIndex.IsValid())
 		{
+			if (State->EventDataIndex.IsValid())
+			{
 				FStateTreeSharedEvent& StateTreeEvent = Storage.GetMutableStruct(CurrentFrame->ActiveInstanceIndexBase.Get() + State->EventDataIndex.Get()).Get<FStateTreeSharedEvent>();
 				const FSelectionEventWithID* EventToCapturePtr = Args.SelectionEvents.FindByPredicate([SelectedState](const FSelectionEventWithID& EventID)
-			{
+					{
 						return EventID.State == SelectedState;
 					});
 				if (ensureAlways(EventToCapturePtr))
@@ -3100,11 +3097,11 @@ EStateTreeRunStatus FStateTreeExecutionContext::EnterState(FSelectStateResult& A
 		// States which were active before and will remain active, but are not on target branch will not get
 		// EnterState called. That is, a transition is handled as "replan from this state".
 		if (bAreCommon && !bSustainedState)
-	{
+		{
 			// Is already in the active and we do not need to call EnterState with Sustained
 			continue;
 		}
-		
+
 		FStateTreeExecutionFrame* CurrentParentFrame = FrameIndex > 0 ? &Exec.ActiveFrames[FrameIndex - 1] : nullptr;
 		const UStateTree* CurrentStateTree = CurrentFrame.StateTree;
 
@@ -3175,7 +3172,7 @@ EStateTreeRunStatus FStateTreeExecutionContext::EnterState(FSelectStateResult& A
 				{
 					const FStateTreeDataView StateParamsDataView = GetDataView(CurrentParentFrame, CurrentFrame, CurrentState.ParameterDataHandle);
 					CopyBatchOnActiveInstances(CurrentParentFrame, CurrentFrame, StateParamsDataView, CurrentState.ParameterBindingsBatch);
-			}
+				}
 			}
 
 			const EStateTreeStateChangeType ChangeType = bSustainedState ? EStateTreeStateChangeType::Sustained : EStateTreeStateChangeType::Changed;
@@ -3199,7 +3196,7 @@ EStateTreeRunStatus FStateTreeExecutionContext::EnterState(FSelectStateResult& A
 					if (Cond.bHasShouldCallStateChangeEvents)
 					{
 						const bool bShouldCallEnterState = CurrentTransition.ChangeType == EStateTreeStateChangeType::Changed
-															|| (CurrentTransition.ChangeType == EStateTreeStateChangeType::Sustained && Cond.bShouldStateChangeOnReselect);
+							|| (CurrentTransition.ChangeType == EStateTreeStateChangeType::Sustained && Cond.bShouldStateChangeOnReselect);
 
 						if (bShouldCallEnterState)
 						{
@@ -3251,7 +3248,7 @@ EStateTreeRunStatus FStateTreeExecutionContext::EnterState(FSelectStateResult& A
 				}
 
 				const bool bShouldCallEnterState = CurrentTransition.ChangeType == EStateTreeStateChangeType::Changed
-													|| (CurrentTransition.ChangeType == EStateTreeStateChangeType::Sustained && Task.bShouldStateChangeOnReselect);
+					|| (CurrentTransition.ChangeType == EStateTreeStateChangeType::Sustained && Task.bShouldStateChangeOnReselect);
 
 				if (bShouldCallEnterState)
 				{
@@ -3262,7 +3259,7 @@ EStateTreeRunStatus FStateTreeExecutionContext::EnterState(FSelectStateResult& A
 					{
 						QUICK_SCOPE_CYCLE_COUNTER(StateTree_Task_EnterState);
 						CSV_SCOPED_TIMING_STAT_EXCLUSIVE(StateTree_Task_EnterState);
-					
+
 						TaskRunStatus = Task.EnterState(*this, CurrentTransition);
 					}
 
@@ -3286,7 +3283,7 @@ EStateTreeRunStatus FStateTreeExecutionContext::EnterState(FSelectStateResult& A
 		}
 
 		if (Result == EStateTreeRunStatus::Failed)
-			{
+		{
 			break;
 		}
 	}
@@ -3352,13 +3349,13 @@ void FStateTreeExecutionContext::ExitState(const FSelectStateResult& Args, const
 			const FStateTreeStateHandle CurrentHandle = CurrentFrame.ActiveStates[StateIndex];
 			const FActiveStateID CurrentStateID = CurrentFrame.ActiveStates.StateIDs[StateIndex];
 			const FCompactStateTreeState& CurrentState = CurrentStateTree->States[CurrentHandle.Index];
-		
+
 			const FActiveState CurrentActiveState(CurrentFrame.FrameID, CurrentStateID, CurrentHandle);
 			const bool bAreCommon = !ChangedStates.Contains(CurrentActiveState);
 
 			// Can it stop calling Changed
 			if (bAreCommon && !bHasSustained)
-				{
+			{
 				// It will also stop the frame loop.
 				bContinue = false;
 				break;
@@ -3368,22 +3365,22 @@ void FStateTreeExecutionContext::ExitState(const FSelectStateResult& Args, const
 			{
 				// this is the last sustained
 				bContinue = false;
-		}
+			}
 
 			const bool bIsStateSustained = bAreCommon;
 
 			TGuardValue<FStateTreeStateHandle> CurrentTransitionCurrentStateGuard(CurrentTransition.CurrentState, CurrentHandle);
 			CurrentTransition.ChangeType = bIsStateSustained ? EStateTreeStateChangeType::Sustained : EStateTreeStateChangeType::Changed;
 
-			STATETREE_LOG(Log, TEXT("%*sState '%s' (%s)"), (FrameIndex + StateIndex + 1)*UE::StateTree::Debug::IndentSize, TEXT("")
+			STATETREE_LOG(Log, TEXT("%*sState '%s' (%s)"), (FrameIndex + StateIndex + 1) * UE::StateTree::Debug::IndentSize, TEXT("")
 				, *GetSafeStateName(CurrentFrame, CurrentHandle)
 				, *UEnum::GetDisplayValueAsText(CurrentTransition.ChangeType).ToString());
 
-				FCurrentlyProcessedStateScope StateScope(*this, CurrentHandle);
-				UE_STATETREE_DEBUG_STATE_EVENT(this, CurrentHandle, EStateTreeTraceEventType::OnExiting);
+			FCurrentlyProcessedStateScope StateScope(*this, CurrentHandle);
+			UE_STATETREE_DEBUG_STATE_EVENT(this, CurrentHandle, EStateTreeTraceEventType::OnExiting);
 
 			for (int32 TaskIndex = (CurrentState.TasksBegin + CurrentState.TasksNum) - 1; TaskIndex >= CurrentState.TasksBegin; TaskIndex--)
-					{
+			{
 				const FStateTreeTaskBase& Task = CurrentStateTree->Nodes[TaskIndex].Get<const FStateTreeTaskBase>();
 
 				// Ignore disabled task
@@ -3396,27 +3393,27 @@ void FStateTreeExecutionContext::ExitState(const FSelectStateResult& Args, const
 						const FStateTreeDataView TaskInstanceView = GetDataView(CurrentParentFrame, CurrentFrame, Task.InstanceDataHandle);
 						FNodeInstanceDataScope DataScope(*this, &Task, TaskIndex, Task.InstanceDataHandle, TaskInstanceView);
 
-							const bool bShouldCallStateChange = CurrentTransition.ChangeType == EStateTreeStateChangeType::Changed
-									|| (CurrentTransition.ChangeType == EStateTreeStateChangeType::Sustained && Task.bShouldStateChangeOnReselect);
+						const bool bShouldCallStateChange = CurrentTransition.ChangeType == EStateTreeStateChangeType::Changed
+							|| (CurrentTransition.ChangeType == EStateTreeStateChangeType::Sustained && Task.bShouldStateChangeOnReselect);
 
-							if (bShouldCallStateChange)
+						if (bShouldCallStateChange)
+						{
+							STATETREE_LOG(Verbose, TEXT("%*sTask '%s'.ExitState()"), (FrameIndex + StateIndex + 1) * UE::StateTree::Debug::IndentSize, TEXT(""), *Task.Name.ToString());
+							UE_STATETREE_DEBUG_TASK_EXIT_STATE(this, CurrentStateTree, FStateTreeIndex16(TaskIndex));
 							{
-								STATETREE_LOG(Verbose, TEXT("%*sTask '%s'.ExitState()"), (FrameIndex + StateIndex + 1)*UE::StateTree::Debug::IndentSize, TEXT(""), *Task.Name.ToString());
-								UE_STATETREE_DEBUG_TASK_EXIT_STATE(this, CurrentStateTree, FStateTreeIndex16(TaskIndex));
-								{
-									QUICK_SCOPE_CYCLE_COUNTER(StateTree_Task_ExitState);
-									CSV_SCOPED_TIMING_STAT_EXCLUSIVE(StateTree_Task_ExitState);
-									Task.ExitState(*this, CurrentTransition);
-								}
-								UE_STATETREE_DEBUG_TASK_EVENT(this, TaskIndex, TaskInstanceView, EStateTreeTraceEventType::OnExited, Transition.CurrentRunStatus);
+								QUICK_SCOPE_CYCLE_COUNTER(StateTree_Task_ExitState);
+								CSV_SCOPED_TIMING_STAT_EXCLUSIVE(StateTree_Task_ExitState);
+								Task.ExitState(*this, CurrentTransition);
 							}
+							UE_STATETREE_DEBUG_TASK_EVENT(this, TaskIndex, TaskInstanceView, EStateTreeTraceEventType::OnExited, Transition.CurrentRunStatus);
 						}
+					}
 				}
-						else
-						{
-							STATETREE_LOG(VeryVerbose, TEXT("%*sSkipped 'ExitState' for disabled Task: '%s'"), UE::StateTree::Debug::IndentSize, TEXT(""), *Task.Name.ToString());
-						}
-					CurrentFrame.ActiveTaskNodeIndex = FStateTreeIndex16(TaskIndex - 1);
+				else
+				{
+					STATETREE_LOG(VeryVerbose, TEXT("%*sSkipped 'ExitState' for disabled Task: '%s'"), UE::StateTree::Debug::IndentSize, TEXT(""), *Task.Name.ToString());
+				}
+				CurrentFrame.ActiveTaskNodeIndex = FStateTreeIndex16(TaskIndex - 1);
 
 				// Call state change events on conditions if needed.
 				if (CurrentState.bHasStateChangeConditions)
@@ -3427,7 +3424,7 @@ void FStateTreeExecutionContext::ExitState(const FSelectStateResult& Args, const
 						if (Cond.bHasShouldCallStateChangeEvents)
 						{
 							const bool bShouldCallStateChange = CurrentTransition.ChangeType == EStateTreeStateChangeType::Changed
-										|| (CurrentTransition.ChangeType == EStateTreeStateChangeType::Sustained && Cond.bShouldStateChangeOnReselect);
+								|| (CurrentTransition.ChangeType == EStateTreeStateChangeType::Sustained && Cond.bShouldStateChangeOnReselect);
 
 							if (bShouldCallStateChange)
 							{
@@ -3535,7 +3532,7 @@ void FStateTreeExecutionContext::StateCompleted()
 
 			FCurrentlyProcessedStateScope StateScope(*this, CurrentHandle);
 			UE_STATETREE_DEBUG_STATE_EVENT(this, CurrentHandle, EStateTreeTraceEventType::OnStateCompleted);
-			STATETREE_LOG(Verbose, TEXT("%*sState '%s'"), (FrameIndex + StateIndex + 1)*UE::StateTree::Debug::IndentSize, TEXT(""), *GetSafeStateName(CurrentFrame, CurrentHandle));
+			STATETREE_LOG(Verbose, TEXT("%*sState '%s'"), (FrameIndex + StateIndex + 1) * UE::StateTree::Debug::IndentSize, TEXT(""), *GetSafeStateName(CurrentFrame, CurrentHandle));
 
 			// Notify Tasks
 			for (int32 TaskIndex = (State.TasksBegin + State.TasksNum) - 1; TaskIndex >= State.TasksBegin; TaskIndex--)
@@ -3555,11 +3552,11 @@ void FStateTreeExecutionContext::StateCompleted()
 					const FStateTreeDataView TaskInstanceView = GetDataView(CurrentParentFrame, CurrentFrame, Task.InstanceDataHandle);
 					FNodeInstanceDataScope DataScope(*this, &Task, TaskIndex, Task.InstanceDataHandle, TaskInstanceView);
 
-					STATETREE_LOG(Verbose, TEXT("%*sTask '%s'.StateCompleted()"), (FrameIndex + StateIndex + 1)*UE::StateTree::Debug::IndentSize, TEXT(""), *Task.Name.ToString());
+					STATETREE_LOG(Verbose, TEXT("%*sTask '%s'.StateCompleted()"), (FrameIndex + StateIndex + 1) * UE::StateTree::Debug::IndentSize, TEXT(""), *Task.Name.ToString());
 					Task.StateCompleted(*this, Exec.LastTickStatus, CurrentFrame.ActiveStates);
 				}
 			}
-				
+
 			// Call state change events on conditions if needed.
 			if (State.bHasStateChangeConditions)
 			{
@@ -3576,7 +3573,7 @@ void FStateTreeExecutionContext::StateCompleted()
 							// Use validated copy, since we test in situations where the sources are not always valid (e.g. enter conditions may try to access inactive parent state).
 							CopyBatchOnActiveInstances(CurrentParentFrame, CurrentFrame, ConditionInstanceView, Cond.BindingsBatch);
 						}
-							
+
 						Cond.StateCompleted(*this, Exec.LastTickStatus, CurrentFrame.ActiveStates);
 
 						// Reset copied properties that might contain object references.
@@ -4206,7 +4203,7 @@ EStateTreeRunStatus FStateTreeExecutionContext::TickTasks(const float DeltaTime)
 			StateResult = GetPriorityRunStatus(StateResult, CastToRunStatus(StateTasksStatus));
 		}
 	}
-	
+
 	if (ExecutionContext::Private::bGlobalTasksCompleteOwningFrame && FirstFrameResult != EStateTreeRunStatus::Running)
 	{
 		Exec.RequestedStop = ExecutionContext::GetPriorityRunStatus(Exec.RequestedStop, FrameResult);
@@ -4338,14 +4335,14 @@ FStateTreeExecutionContext::FTickTaskResult FStateTreeExecutionContext::TickTask
 		if (Args.TasksCompletionStatus->IsConsideredForCompletion(OwnerTaskIndex))
 		{
 			if (TaskRunStatus == EStateTreeRunStatus::Failed)
-		{
-			bShouldTickTasks = false;
+			{
+				bShouldTickTasks = false;
 				break; // Stop copy binding.
 			}
 		}
 	}
 
-	return FTickTaskResult{bShouldTickTasks};
+	return FTickTaskResult{ bShouldTickTasks };
 }
 
 // Deprecated
@@ -4381,14 +4378,14 @@ bool FStateTreeExecutionContext::TestAllConditionsInternal(const FStateTreeExecu
 	TStaticArray<bool, UE::StateTree::MaxExpressionIndent + 1> Values(InPlace, false);
 
 	int32 Level = 0;
-	
+
 	for (int32 Index = 0; Index < ConditionsNum; Index++)
 	{
 		const int32 ConditionIndex = ConditionsOffset + Index;
 		const FStateTreeConditionBase& Cond = CurrentFrame.StateTree->Nodes[ConditionIndex].Get<const FStateTreeConditionBase>();
-		checkf(Cond.InstanceDataHandle.GetSource() == EStateTreeDataSourceType::SharedInstanceData 
-			   || Cond.InstanceDataHandle.GetSource() == EStateTreeDataSourceType::SharedInstanceDataObject, 
-			   TEXT("Condition only support SharedInstanceData"));
+		checkf(Cond.InstanceDataHandle.GetSource() == EStateTreeDataSourceType::SharedInstanceData
+			|| Cond.InstanceDataHandle.GetSource() == EStateTreeDataSourceType::SharedInstanceDataObject,
+			TEXT("Condition only support SharedInstanceData"));
 
 		const FStateTreeDataView ConditionInstanceView = GetDataView(CurrentParentFrame, CurrentFrame, Cond.InstanceDataHandle);
 		FNodeInstanceDataScope DataScope(*this, &Cond, ConditionIndex, Cond.InstanceDataHandle, ConditionInstanceView);
@@ -4416,7 +4413,7 @@ bool FStateTreeExecutionContext::TestAllConditionsInternal(const FStateTreeExecu
 
 			bValue = Cond.TestCondition(*this);
 			UE_STATETREE_DEBUG_CONDITION_EVENT(this, ConditionIndex, ConditionInstanceView, bValue ? EStateTreeTraceEventType::Passed : EStateTreeTraceEventType::Failed);
-			
+
 			// Reset copied properties that might contain object references.
 			if (Cond.BindingsBatch.IsValid())
 			{
@@ -4463,7 +4460,7 @@ bool FStateTreeExecutionContext::TestAllConditionsInternal(const FStateTreeExecu
 			Operands[Level] = EStateTreeExpressionOperand::Copy;
 		}
 	}
-	
+
 	return Values[0];
 }
 
@@ -4498,8 +4495,8 @@ float FStateTreeExecutionContext::EvaluateUtilityWithValidation(const FStateTree
 		const FStateTreeConsiderationBase& Consideration = CurrentFrame.StateTree->Nodes[ConsiderationIndex].Get<const FStateTreeConsiderationBase>();
 
 		checkf(Consideration.InstanceDataHandle.GetSource() == EStateTreeDataSourceType::SharedInstanceData
-			   || Consideration.InstanceDataHandle.GetSource() == EStateTreeDataSourceType::SharedInstanceDataObject
-			   , TEXT("Consideration only support SharedInstanceData"));
+			|| Consideration.InstanceDataHandle.GetSource() == EStateTreeDataSourceType::SharedInstanceDataObject
+			, TEXT("Consideration only support SharedInstanceData"));
 
 		const FStateTreeDataView ConsiderationInstanceView = GetDataView(CurrentParentFrame, CurrentFrame, Consideration.InstanceDataHandle);
 		FNodeInstanceDataScope DataScope(*this, &Consideration, ConsiderationIndex, Consideration.InstanceDataHandle, ConsiderationInstanceView);
@@ -4576,9 +4573,9 @@ void FStateTreeExecutionContext::EvaluatePropertyFunctionsOnActiveInstances(cons
 			// Use validated copy, since we test in situations where the sources are not always valid (e.g. enter conditions may try to access inactive parent state). 
 			CopyBatchOnActiveInstances(CurrentParentFrame, CurrentFrame, FuncInstanceView, Func.BindingsBatch);
 		}
-			
+
 		Func.Execute(*this);
-			
+
 		// Reset copied properties that might contain object references.
 		if (Func.BindingsBatch.IsValid())
 		{
@@ -4601,9 +4598,9 @@ void FStateTreeExecutionContext::EvaluatePropertyFunctionsWithValidation(const F
 			// Use validated copy, since we test in situations where the sources are not always valid (e.g. enter conditions may try to access inactive parent state). 
 			CopyBatchWithValidation(CurrentParentFrame, CurrentFrame, FuncInstanceView, Func.BindingsBatch);
 		}
-			
+
 		Func.Execute(*this);
-			
+
 		// Reset copied properties that might contain object references.
 		if (Func.BindingsBatch.IsValid())
 		{
@@ -4629,11 +4626,11 @@ FString FStateTreeExecutionContext::DebugGetEventsAsString() const
 
 				const bool bHasTag = Event->Tag.IsValid();
 				const bool bHasPayload = Event->Payload.GetScriptStruct() != nullptr;
-			
+
 				if (bHasTag || bHasPayload)
 				{
 					StrBuilder << (TEXT('('));
-				
+
 					if (bHasTag)
 					{
 						StrBuilder << TEXT("Tag: '");
@@ -4735,7 +4732,7 @@ bool FStateTreeExecutionContext::RequestTransitionInternal(
 
 	FSelectStateResult StateSelectionResult;
 	if (SelectState(SelectStateArgs, StateSelectionResult))
-		{
+	{
 		const FStateTreeExecutionState& Exec = GetExecState();
 
 		SetupNextTransition(SourceFrame, SourceStateID, TargetState, Priority);
@@ -4774,15 +4771,15 @@ bool FStateTreeExecutionContext::RequestTransitionInternal(
 
 				const FSelectionEventWithID* FoundEvent = RequestedTransition->Selection.SelectionEvents.FindByPredicate(
 					[SelectedState](const FSelectionEventWithID& Event)
-				{
+					{
 						return Event.State == SelectedState;
 					});
 				if (FoundEvent)
-					{
+				{
 					CurrentSelectionEvents->Events[CurrentSelectionEventsIndex] = FoundEvent->Event;
 				}
 			}
-					}
+		}
 
 		//@todo In TriggerTransitions, the events are processed bottom up. If an event with a higher priority is process after, the event will already be consumed.
 		// Consume events from states, if required. The transition might also want to consume the event. (Not a bug.)
@@ -4800,7 +4797,7 @@ bool FStateTreeExecutionContext::RequestTransitionInternal(
 		}
 
 PRAGMA_DISABLE_DEPRECATION_WARNINGS
-UE_SUPPRESS(LogStateTree, Verbose,)
+		UE_SUPPRESS(LogStateTree, Verbose,
 		{
 			check(RequestedTransition->Selection.SelectedStates.Num() > 0);
 			const FActiveState& LastSelectedState = RequestedTransition->Selection.SelectedStates.Last();
@@ -4811,7 +4808,7 @@ UE_SUPPRESS(LogStateTree, Verbose,)
 				*GetSafeStateName(TargetState.StateTree, TargetState.StateHandle),
 				*GetSafeStateName(*SelectedFrame, LastSelectedState.GetStateHandle())
 			);
-		}
+		})
 PRAGMA_ENABLE_DEPRECATION_WARNINGS
 
 		return true;
@@ -4853,9 +4850,9 @@ void FStateTreeExecutionContext::SetupNextTransition(
 	OutTransitionResult.SourceStateID = SourceStateID;
 PRAGMA_DISABLE_DEPRECATION_WARNINGS
 	if (SourceStateID.IsValid())
-		{
+	{
 		OutTransitionResult.SourceState = SourceFrame.ActiveStates.FindStateHandle(SourceStateID);
-		}
+	}
 	OutTransitionResult.SourceStateTree = SourceFrame.StateTree;
 	OutTransitionResult.SourceRootState = SourceFrame.RootState;
 PRAGMA_ENABLE_DEPRECATION_WARNINGS
@@ -4950,7 +4947,7 @@ bool FStateTreeExecutionContext::TriggerTransitions()
 	struct FTransitionHandler
 	{
 		FTransitionHandler() = default;
-		
+
 		FTransitionHandler(const int32 InFrameIndex, const FStateTreeStateHandle InStateHandle, const UE::StateTree::FActiveStateID InStateID, const EStateTreeTransitionPriority InPriority)
 			: StateHandle(InStateHandle)
 			, StateID(InStateID)
@@ -5029,10 +5026,10 @@ bool FStateTreeExecutionContext::TriggerTransitions()
 
 				// Has expired transition delayed.
 				const bool bHasActiveTransitionDelayed = ExpiredTransitionsDelayed.ContainsByPredicate([StateID](const FStateTreeTransitionDelayedState& Other)
-				{
-					return Other.StateID == StateID;
-				});
-				
+					{
+						return Other.StateID == StateID;
+					});
+
 				// Regular transitions on state
 				//or A transition task can trigger an event. We need to add the state if that is a possibility
 				//or Expired transition delayed
@@ -5124,7 +5121,7 @@ bool FStateTreeExecutionContext::TriggerTransitions()
 				{
 					continue;
 				}
-				
+
 				// No need to test the transition if same or higher priority transition has already been processed.
 				if (RequestedTransition.IsValid() && Transition.Priority <= RequestedTransition->Transition.Priority)
 				{
@@ -5176,7 +5173,7 @@ bool FStateTreeExecutionContext::TriggerTransitions()
 
 				for (const FStateTreeSharedEvent& TransitionEvent : TransitionEvents)
 				{
-					bool bPassed = false; 
+					bool bPassed = false;
 					{
 						FCurrentlyProcessedTransitionEventScope TransitionEventScope(*this, TransitionEvent.IsValid() ? TransitionEvent.Get() : nullptr);
 						UE_STATETREE_DEBUG_TRANSITION_EVENT(this, FStateTreeTransitionSource(CurrentFrame.StateTree, FStateTreeIndex16(TransitionIndex), Transition.State, Transition.Priority), EStateTreeTraceEventType::OnEvaluating);
@@ -5191,16 +5188,16 @@ bool FStateTreeExecutionContext::TriggerTransitions()
 						{
 							uint32 TransitionEventHash = 0u;
 							if (TransitionEvent.IsValid())
-							{		
+							{
 								TransitionEventHash = GetTypeHash(TransitionEvent.Get());
 							}
 
-							const bool bIsDelayedTransitionExisting = Exec.DelayedTransitions.ContainsByPredicate([CurrentStateID =  Handler.StateID, TransitionIndex, TransitionEventHash](const FStateTreeTransitionDelayedState& DelayedState)
-							{
-								return DelayedState.StateID == CurrentStateID
-									&& DelayedState.TransitionIndex.Get() == TransitionIndex
-									&& DelayedState.CapturedEventHash == TransitionEventHash;
-							});
+							const bool bIsDelayedTransitionExisting = Exec.DelayedTransitions.ContainsByPredicate([CurrentStateID = Handler.StateID, TransitionIndex, TransitionEventHash](const FStateTreeTransitionDelayedState& DelayedState)
+								{
+									return DelayedState.StateID == CurrentStateID
+										&& DelayedState.TransitionIndex.Get() == TransitionIndex
+										&& DelayedState.CapturedEventHash == TransitionEventHash;
+								});
 
 							if (!bIsDelayedTransitionExisting)
 							{
@@ -5225,7 +5222,7 @@ PRAGMA_ENABLE_DEPRECATION_WARNINGS
 									BeginDelayedTransition(DelayedState);
 									STATETREE_LOG(Verbose, TEXT("Delayed transition triggered from '%s' (%s) -> '%s' %.1fs"),
 										*GetSafeStateName(CurrentFrame, CurrentFrame.ActiveStates.Last()), *State.Name.ToString(), *GetSafeStateName(CurrentFrame, Transition.State), DelayedState.TimeLeft);
-								
+
 									// Delay state added, skip requesting the transition.
 									continue;
 								}
@@ -5247,7 +5244,7 @@ PRAGMA_ENABLE_DEPRECATION_WARNINGS
 							{
 								ConsumeEvent(TransitionEvent);
 							}
-							
+
 							RequestedTransition->Source = FStateTreeTransitionSource(CurrentFrame.StateTree, FStateTreeIndex16(TransitionIndex), Transition.State, Transition.Priority);
 							break;
 						}
@@ -5275,7 +5272,7 @@ PRAGMA_ENABLE_DEPRECATION_WARNINGS
 	{
 		// Find the pending completed frame/state. Don't cache the result because this function is reentrant.
 		//Stop at the first completion.
-		int32 FrameIndexToStart = -1; 
+		int32 FrameIndexToStart = -1;
 		int32 StateIndexToStart = -1;
 		EStateTreeRunStatus CurrentStatus = EStateTreeRunStatus::Unset;
 		for (int32 FrameIndex = 0; FrameIndex < Exec.ActiveFrames.Num(); ++FrameIndex)
@@ -5306,7 +5303,7 @@ PRAGMA_ENABLE_DEPRECATION_WARNINGS
 						STATETREE_LOG(Verbose, TEXT("Completed subtree '%s' from global: %s"),
 							*GetSafeStateName(ParentFrame, ParentLinkedState),
 							*UEnum::GetDisplayValueAsText(CastToRunStatus(FrameTasksStatus)).ToString()
-							);
+						);
 
 						const FCompactStateTreeState& State = ParentFrame.StateTree->States[ParentLinkedState.Index];
 						ParentFrame.ActiveTasksStatus.GetStatus(State).SetCompletionStatus(FrameTasksStatus);
@@ -5384,7 +5381,7 @@ PRAGMA_ENABLE_DEPRECATION_WARNINGS
 								continue;
 							}
 
-							const bool bTransitionAccepted = !bIsCurrentStatusStopped 
+							const bool bTransitionAccepted = !bIsCurrentStatusStopped
 								? EnumHasAnyFlags(Transition.Trigger, CompletionTrigger)
 								: Transition.Trigger == EStateTreeTransitionTrigger::OnStateCompleted;
 							if (bTransitionAccepted)
@@ -5437,7 +5434,7 @@ PRAGMA_ENABLE_DEPRECATION_WARNINGS
 				check(RootFrame.ActiveStates.Num() != 0);
 				FCurrentlyProcessedFrameScope RootFrameScope(*this, nullptr, RootFrame);
 				FCurrentlyProcessedStateScope RootStateScope(*this, FStateTreeStateHandle::Root);
-				
+
 				UE::StateTree::ExecutionContext::FStateHandleContext TargetState(RootFrame.StateTree, FStateTreeStateHandle::Root);
 				if (RequestTransitionInternal(RootFrame, RootFrame.ActiveStates.StateIDs[0], TargetState))
 				{
@@ -5482,7 +5479,7 @@ PRAGMA_ENABLE_DEPRECATION_WARNINGS
 					*GetSafeStateName(ParentFrame, ParentLinkedState),
 					*GetSafeStateName(SourceFrame, NextTransitionSourceState),
 					*UEnum::GetDisplayValueAsText(RunStatus).ToString()
-					);
+				);
 #endif
 
 				// Set the parent linked state as last completed state, and update tick status to the status from the transition.
@@ -5648,10 +5645,10 @@ PRAGMA_ENABLE_DEPRECATION_WARNINGS
 
 // Deprecated
 bool FStateTreeExecutionContext::SelectState(const FStateTreeExecutionFrame& CurrentFrame,
-											const FStateTreeStateHandle NextState,
+	const FStateTreeStateHandle NextState,
 	FSelectStateResult& OutSelectionResult,
-											const FStateTreeSharedEvent* TransitionEvent,
-											const EStateTreeSelectionFallback Fallback)
+	const FStateTreeSharedEvent* TransitionEvent,
+	const EStateTreeSelectionFallback Fallback)
 {
 	using namespace UE::StateTree;
 	using namespace UE::StateTree::ExecutionContext;
@@ -5708,12 +5705,6 @@ bool FStateTreeExecutionContext::SelectState(const FSelectStateArguments& Args,
 	//  New D, CD, BCD if completed.
 	//  ExitState: Sustained CD, C, -, -. Changed -, D, CD, BCD. EnterState: Sustained CD, C, -, -. Changed -, D', C'D', B'C'D'.
 	//
-	// (if ReselectedParentCreatesNewChildStates)
-	//           ABCD   | ABC    | AB     | ABC'D', AB'C'D'
-	//  Select a parent state. New B always new CD
-	//  ExitState: Sustained B, -. Changed CD, BCD. EnterState: Sustained B, -. Changed C'D', B'C'D'
-	//
-	// (if not ReselectedStateCreatesNewStates)
 	//           ABCD   | ABC     | AB    | ABCD, ABCD', ABC'D', AB'C'D'
 	//  New D, CD, BCD if completed.
 	//  ExitState: Sustained BCD, BC, B, -. Changed -, D, CD, BCD. EnterState: Sustained BCD, BC, B, -. Changed -, D', C'D', B'C'D'.
@@ -5749,16 +5740,16 @@ bool FStateTreeExecutionContext::SelectState(const FSelectStateArguments& Args,
 
 	TArray<FStateHandleContext, TInlineAllocator<FStateTreeActiveStates::MaxStates>> PathToTargetState;
 	if (!GetStatesListToState(Args.TargetState, PathToTargetState))
-		{
-			STATETREE_LOG(Error, TEXT("%hs: Reached max execution depth when trying to select state %s from '%s'.  '%s' using StateTree '%s'."),
+	{
+		STATETREE_LOG(Error, TEXT("%hs: Reached max execution depth when trying to select state %s from '%s'.  '%s' using StateTree '%s'."),
 			__FUNCTION__,
 			*GetSafeStateName(Args.TargetState.StateTree, Args.TargetState.StateHandle),
 			*GetStateStatusString(Exec),
 			*GetNameSafe(&Owner),
 			*GetFullNameSafe(&RootStateTree)
 		);
-			return false;
-		}
+		return false;
+	}
 	check(!PathToTargetState.IsEmpty());
 
 	const FExecutionFrameKey TargetFrameKey = FExecutionFrameKey(Args.TargetState.StateTree, PathToTargetState[0].StateHandle);
@@ -5786,31 +5777,31 @@ bool FStateTreeExecutionContext::SelectState(const FSelectStateArguments& Args,
 			SourceStates = Args.ActiveStates.Left(FoundIndex + 1);
 		}
 		else
-	{
+		{
 			// Pick the source frame first state.
 			//This is usually when the request is out of the scope or from a global task. It should start from the root of the frame.
 			int32 ActiveStateIndex = 0;
 			for (; ActiveStateIndex < Args.ActiveStates.Num(); ++ActiveStateIndex)
-		{
+			{
 				const FActiveFrameID CurrentActiveFrameID = Args.ActiveStates[ActiveStateIndex].GetFrameID();
 				if (CurrentActiveFrameID == Args.SourceState.GetFrameID())
-			{
-				break;
+				{
+					break;
+				}
 			}
-		}
 			if (ActiveStateIndex < Args.ActiveStates.Num())
 			{
 				SourceStates = Args.ActiveStates.Left(ActiveStateIndex + 1);
-	}
+			}
 			else
 			{
 				// SourceFrame exists, so the state must be in the active states.
 				if (!ensure(Args.ActiveStates.Num() == 0))
-	{
+				{
 					return false;
 				}
 			}
-	}
+		}
 	}
 	checkf(SourceStates.Num() == 0 || FActiveStatePath::StartsWith(Args.ActiveStates, SourceStates), TEXT("Source is part of the active path."));
 	FStateTreeExecutionFrame* SourceExecFrame = FindExecutionFrame(Args.SourceState.GetFrameID(), MakeArrayView(Exec.ActiveFrames), MakeArrayView(OutSelectionResult.TemporaryFrames));
@@ -5818,7 +5809,7 @@ bool FStateTreeExecutionContext::SelectState(const FSelectStateArguments& Args,
 	{
 		return false;
 	}
-	
+
 	const int32 SelectedStatesSlack = Args.ActiveStates.Num() + 10; // Number to help with buffer reallocation
 	OutSelectionResult.SelectedStates.Empty(SelectedStatesSlack);
 
@@ -5833,9 +5824,9 @@ bool FStateTreeExecutionContext::SelectState(const FSelectStateArguments& Args,
 			// Copy active to the new selected until we reach the frame.
 			FActiveFrameID CurrentFrameID;
 			for (const FActiveState& ActiveState : Args.ActiveStates)
-	{
+			{
 				if (CurrentFrameID != ActiveState.GetFrameID())
-		{
+				{
 					CurrentFrameID = ActiveState.GetFrameID();
 					if (ActiveState.GetFrameID() == Args.SourceState.GetFrameID())
 					{
@@ -5879,11 +5870,11 @@ bool FStateTreeExecutionContext::SelectState(const FSelectStateArguments& Args,
 							}
 						}
 						else if (bFoundStateTree)
-			{
-				break;
-			}
-		}
-	}
+						{
+							break;
+						}
+					}
+				}
 
 				if (FoundActiveStateIndex < 0 && !bFoundStateTree)
 				{
@@ -5898,7 +5889,7 @@ bool FStateTreeExecutionContext::SelectState(const FSelectStateArguments& Args,
 				}
 
 				// Copy active states to the new selected until we reach the frame.
-				for (int32 ActiveStateIndex = 0; ActiveStateIndex < FoundActiveStateIndex+1; ++ActiveStateIndex)
+				for (int32 ActiveStateIndex = 0; ActiveStateIndex < FoundActiveStateIndex + 1; ++ActiveStateIndex)
 				{
 					const FActiveState& ActiveState = Args.ActiveStates[ActiveStateIndex];
 					if (CurrentFrameID != ActiveState.GetFrameID())
@@ -5911,7 +5902,7 @@ bool FStateTreeExecutionContext::SelectState(const FSelectStateArguments& Args,
 			}
 		}
 		else
-	{
+		{
 			// SourceExecFrame do not matches and there are no source states.
 			STATETREE_LOG(Error, TEXT("%hs: Encountered out of range state %s during state selection from '%s'.  '%s' using StateTree '%s'."),
 				__FUNCTION__,
@@ -5923,7 +5914,7 @@ bool FStateTreeExecutionContext::SelectState(const FSelectStateArguments& Args,
 			return false;
 		}
 	}
-	
+
 	FSelectStateInternalArguments InternalArgs;
 	InternalArgs.MissingActiveStates = Args.ActiveStates.Mid(OutSelectionResult.SelectedStates.Num());
 	InternalArgs.MissingSourceStates = SourceStates.Mid(OutSelectionResult.SelectedStates.Num());
@@ -6062,9 +6053,9 @@ bool FStateTreeExecutionContext::SelectStateInternal(
 			FStateTreeExecutionFrame* MissingFramePtr = FindExecutionFrame(MissingFrameID, MakeArrayView(GetExecState().ActiveFrames), MakeArrayView(OutSelectionResult.TemporaryFrames));
 			check(MissingFramePtr);
 			if (!ensure(MissingFramePtr->StateTree == NextStateTree))
-	{
-		return false;
-	}
+			{
+				return false;
+			}
 
 			FActiveFrameID NextFrameID = MissingFrameID;
 			if (InternalArgs.MissingSourceStates.Num() != 0)
@@ -6093,7 +6084,7 @@ bool FStateTreeExecutionContext::SelectStateInternal(
 					OutSelectionResult.SelectedFrames.Add(NextFrameID);
 				}
 			}
-			else 
+			else
 			{
 				OutSelectionResult.SelectedFrames.Add(NextFrameID);
 			}
@@ -6200,33 +6191,33 @@ bool FStateTreeExecutionContext::SelectStateInternal(
 	{
 		if (ensure(NextState.RequiredEventToEnter.IsValid()))
 		{
-		// Use the same event as performed transition unless it didn't lead to this state as only state selected by the transition should get it's event.
+			// Use the same event as performed transition unless it didn't lead to this state as only state selected by the transition should get it's event.
 			if (Args.TransitionEvent.IsValid() && bIsTargetState)
-		{
-				if (NextState.RequiredEventToEnter.DoesEventMatchDesc(*Args.TransitionEvent))
 			{
+				if (NextState.RequiredEventToEnter.DoesEventMatchDesc(*Args.TransitionEvent))
+				{
 					StateSelectionEvents.Emplace(Args.TransitionEvent);
+				}
 			}
-		}
-		else
-		{
-			TArrayView<FStateTreeSharedEvent> EventsQueue = GetMutableEventsToProcessView();
-			for (FStateTreeSharedEvent& Event : EventsQueue)
+			else
 			{
-				check(Event.IsValid());
-				if (NextState.RequiredEventToEnter.DoesEventMatchDesc(*Event))
+				TArrayView<FStateTreeSharedEvent> EventsQueue = GetMutableEventsToProcessView();
+				for (FStateTreeSharedEvent& Event : EventsQueue)
 				{
+					check(Event.IsValid());
+					if (NextState.RequiredEventToEnter.DoesEventMatchDesc(*Event))
+					{
 						StateSelectionEvents.Emplace(Event);
+					}
 				}
-			}
 
-			// Couldn't find matching state's event, but it's marked as not required. Adding an empty event which allows us to continue the state selection.
-			if (!bShouldPrerequisitesBeChecked && StateSelectionEvents.IsEmpty())
-			{
-				StateSelectionEvents.Emplace();
+				// Couldn't find matching state's event, but it's marked as not required. Adding an empty event which allows us to continue the state selection.
+				if (!bShouldPrerequisitesBeChecked && StateSelectionEvents.IsEmpty())
+				{
+					StateSelectionEvents.Emplace();
+				}
 			}
 		}
-		}
 
 		if (StateSelectionEvents.IsEmpty())
 		{
@@ -6277,19 +6268,19 @@ bool FStateTreeExecutionContext::SelectStateInternal(
 				break;
 			}
 			++StateInFrameCounter;
-	}
+		}
 
 		if (StateInFrameCounter > FStateTreeActiveStates::MaxStates)
-	{
-		STATETREE_LOG(Error, TEXT("%hs: Reached max execution depth when trying to select state %s from '%s'.  '%s' using StateTree '%s'."),
+		{
+			STATETREE_LOG(Error, TEXT("%hs: Reached max execution depth when trying to select state %s from '%s'.  '%s' using StateTree '%s'."),
 				__FUNCTION__,
 				*GetSafeStateName(NextFrame, NextStateHandle),
 				*GetFullNameSafe(NextFrame.StateTree),
 				*GetNameSafe(&Owner),
 				*GetFullNameSafe(&RootStateTree)
 			);
-		return false;
-	}
+			return false;
+		}
 	}
 	else
 	{
@@ -6304,7 +6295,7 @@ bool FStateTreeExecutionContext::SelectStateInternal(
 	}
 
 	for (const FStateTreeSharedEvent& StateSelectionEvent : StateSelectionEvents)
-		{
+	{
 		// Add the event for GetDataView and for consuming it later.
 		const bool bContainsSelectionEvent = OutSelectionResult.SelectionEvents.ContainsByPredicate(
 			[&LatestState = OutSelectionResult.SelectedStates.Last()](const FSelectionEventWithID& Other)
@@ -6318,18 +6309,18 @@ bool FStateTreeExecutionContext::SelectStateInternal(
 		{
 			OutSelectionResult.SelectionEvents.Emplace(OutSelectionResult.SelectedStates.Last(), StateSelectionEvent);
 			bRemoveSelectionEvents = true;
-	}
+		}
 
 		auto RemoveStateSelectionEvent = [bRemoveSelectionEvents](FSelectStateResult& SelectionResult)
-	{
+			{
 				if (bRemoveSelectionEvents)
-		{
+				{
 					check(SelectionResult.SelectionEvents.Num() > 0);
 					ensureMsgf(SelectionResult.SelectionEvents.Last().State == SelectionResult.SelectedStates.Last(), TEXT("We should remove from the same element it was added."));
 					SelectionResult.SelectionEvents.Pop();
-		}
+				}
 			};
-		
+
 		if (bShouldPrerequisitesBeChecked)
 		{
 			// Check that the state can be entered
@@ -6343,7 +6334,7 @@ bool FStateTreeExecutionContext::SelectStateInternal(
 				continue;
 			}
 		}
-		
+
 		FSelectStateInternalArguments NewInternalArgs = FSelectStateInternalArguments{
 			.MissingActiveStates = bShouldCreateNewState ? TArrayView<const FActiveState>() : InternalArgs.MissingActiveStates.Mid(1),
 			.MissingSourceFrameID = FActiveFrameID(),
@@ -6398,25 +6389,25 @@ bool FStateTreeExecutionContext::SelectStateInternal(
 				bSucceededToSelectState = true;
 			}
 			else
-				{
+			{
 				// Continue the force selection. Next child state is already known.
 				bSucceededToSelectState = SelectStateInternal(Args, NewInternalArgs, OutSelectionResult);
 			}
-				}
+		}
 
 
 		if (bSucceededToSelectState)
-				{
-					break;
-				}
+		{
+			break;
+		}
 		else
-					{
+		{
 			RemoveStateSelectionEvent(OutSelectionResult);
 		}
-				}
+	}
 
 	return bSucceededToSelectState;
-				}
+}
 
 TOptional<bool> FStateTreeExecutionContext::SelectStateFromSourceInternal(
 	const FSelectStateArguments& Args,
@@ -6434,7 +6425,7 @@ TOptional<bool> FStateTreeExecutionContext::SelectStateFromSourceInternal(
 	if (InternalArgs.MissingSourceStates.Num() > 0
 		&& InternalArgs.MissingStatesToReachTarget.Num() > 1
 		&& !bNewFrameCreated)
-				{
+	{
 		check(InternalArgs.MissingActiveStates.Num() > 0);
 		check(InternalArgs.MissingSourceStates[0] == InternalArgs.MissingActiveStates[0]);
 		const FActiveState& MissingActiveState = InternalArgs.MissingSourceStates[0];
@@ -6455,12 +6446,12 @@ TOptional<bool> FStateTreeExecutionContext::SelectStateFromSourceInternal(
 				if (EnumHasAllFlags(NextFrame.StateTree->StateSelectionRules, EStateTreeStateSelectionRules::CompletedTransitionStatesCreateNewStates))
 				{
 					bCompletedState = NextFrame.ActiveTasksStatus.GetStatus(NextState).IsCompleted();
+				}
 			}
 		}
-			}
 
 		if (!bContinueWithStateSelection)
-			{
+		{
 			OutSelectionResult.SelectedStates.Add(InternalArgs.MissingSourceStates[0]);
 
 			FSelectStateInternalArguments NewInternalArgs = FSelectStateInternalArguments{
@@ -6477,7 +6468,7 @@ TOptional<bool> FStateTreeExecutionContext::SelectStateFromSourceInternal(
 				//The completed state should fail the transition(because of EStateTreeStateSelectionRules::CompletedStateBeforeTargetFailsTransition).
 				//The transition can contain a state with "follow transition" that would remove this state from the transition.
 				if (OutSelectionResult.SelectedStates.Contains(MissingActiveState))
-			{
+				{
 					STATETREE_LOG(VeryVerbose, TEXT("%hs: Selection fails because state '%s' is completed.  '%s' using StateTree '%s'."),
 						__FUNCTION__,
 						*NextState.Name.ToString(),
@@ -6491,8 +6482,8 @@ TOptional<bool> FStateTreeExecutionContext::SelectStateFromSourceInternal(
 		}
 	}
 	return {};
-			}
-				
+}
+
 namespace UE::StateTree::ExecutionContext::Private
 {
 	bool PreventRecursionCheck(const FExecutionFrameKey& LinkStateFrameKey,
@@ -6500,7 +6491,7 @@ namespace UE::StateTree::ExecutionContext::Private
 		TArrayView<const FStateTreeExecutionFrame> ActiveFrames,
 		TArrayView<const FStateTreeExecutionFrame> TemporaryFrames)
 	{
-			// Check and prevent recursion.
+		// Check and prevent recursion.
 		return SelectedFrames.ContainsByPredicate(
 			[&LinkStateFrameKey, &ActiveFrames, &TemporaryFrames](const FActiveFrameID& FrameID)
 			{
@@ -6508,7 +6499,7 @@ namespace UE::StateTree::ExecutionContext::Private
 				return Frame && Frame->HasRoot(LinkStateFrameKey);
 			});
 	}
-				
+
 	FStateTreeExecutionFrame* SelectedFrameLinkedFrame(FStateTreeExecutionState& Exec,
 		bool bShouldCreateNewState,
 		TArrayView<const UE::StateTree::FActiveState> MatchingActiveStates,
@@ -6517,7 +6508,7 @@ namespace UE::StateTree::ExecutionContext::Private
 	{
 		FStateTreeExecutionFrame* Result = nullptr;
 		if (!bShouldCreateNewState)
-			{
+		{
 			// Get the next frame ID
 			if (ensure(MatchingActiveStates.Num() >= 1))
 			{
@@ -6529,12 +6520,12 @@ namespace UE::StateTree::ExecutionContext::Private
 					if (EnumHasAllFlags(Result->StateTree->GetStateSelectionRules(), EStateTreeStateSelectionRules::CompletedTransitionStatesCreateNewStates)
 						&& Result->ActiveTasksStatus.GetStatus(Result->StateTree).IsCompleted()
 						)
-			{
+					{
 						return nullptr;
 					}
 				}
 			}
-			}
+		}
 		return Result;
 	}
 }
@@ -6546,7 +6537,7 @@ bool FStateTreeExecutionContext::SelectStateInternal_Linked(
 	TNotNull<const UStateTree*> NextStateTree,
 	const FCompactStateTreeState& TargetState,
 	bool bShouldCreateNewState)
-			{
+{
 	using namespace UE::StateTree;
 	using namespace UE::StateTree::ExecutionContext;
 	using namespace UE::StateTree::ExecutionContext::Private;
@@ -6554,7 +6545,7 @@ bool FStateTreeExecutionContext::SelectStateInternal_Linked(
 	FStateTreeExecutionState& Exec = GetExecState();
 
 	if (!TargetState.LinkedState.IsValid())
-				{
+	{
 		STATETREE_LOG(Warning, TEXT("%hs: Trying to enter invalid linked subtree from '%s'. '%s' using StateTree '%s'."),
 			__FUNCTION__,
 			*GetStateStatusString(Exec),
@@ -6562,8 +6553,8 @@ bool FStateTreeExecutionContext::SelectStateInternal_Linked(
 			*NextStateTree->GetFullName()
 		);
 		return false;
-				}
-					
+	}
+
 	const FExecutionFrameKey LinkStateFrameKey = FExecutionFrameKey(NextStateTree, TargetState.LinkedState);
 
 	const bool bHasMissingState = InternalArgs.MissingStatesToReachTarget.Num() > 0 && Args.Behavior == ESelectStateBehavior::Forced;
@@ -6573,18 +6564,18 @@ bool FStateTreeExecutionContext::SelectStateInternal_Linked(
 		//Ex: a previous transition go to root, then another transition go to another top level state (new root)
 		if (InternalArgs.MissingStatesToReachTarget[0].StateTree != LinkStateFrameKey.StateTree
 			|| InternalArgs.MissingStatesToReachTarget[0].StateHandle != LinkStateFrameKey.RootState)
-				{
+		{
 			STATETREE_LOG(Error, TEXT("%hs: The missing state is not from the same state tree. '%s' using StateTree '%s'."),
 				__FUNCTION__,
 				*Owner.GetName(),
 				*NextStateTree->GetFullName()
 			);
 			return false;
-				}
-			}
-				
+		}
+	}
+
 	if (PreventRecursionCheck(LinkStateFrameKey, OutSelectionResult.SelectedFrames, MakeConstArrayView(Exec.ActiveFrames), MakeConstArrayView(OutSelectionResult.TemporaryFrames)))
-			{
+	{
 		STATETREE_LOG(Error, TEXT("%hs: Trying to recursively enter subtree '%s' from '%s'. '%s' using StateTree '%s'."),
 			__FUNCTION__,
 			*GetSafeStateName(LinkStateFrameKey.StateTree, LinkStateFrameKey.RootState),
@@ -6593,11 +6584,11 @@ bool FStateTreeExecutionContext::SelectStateInternal_Linked(
 			*NextStateTree->GetFullName()
 		);
 		return false;
-			}
+	}
 
 	const FCompactStateTreeFrame* LinkStateTreeFrame = FindStateTreeFrame(LinkStateFrameKey);
 	if (LinkStateTreeFrame == nullptr)
-			{
+	{
 		STATETREE_LOG(Error, TEXT("%hs: The frame '%s' from '%s' does not exist. '%s' using StateTree '%s'."),
 			__FUNCTION__,
 			*GetSafeStateName(LinkStateFrameKey.StateTree, LinkStateFrameKey.RootState),
@@ -6606,8 +6597,8 @@ bool FStateTreeExecutionContext::SelectStateInternal_Linked(
 			*NextStateTree->GetFullName()
 		);
 		return false;
-			}
-				
+	}
+
 	// Do we have an existing frame.
 	FStateTreeExecutionFrame* SelectedFrame = SelectedFrameLinkedFrame(Exec, bShouldCreateNewState, InternalArgs.MissingActiveStates, LinkStateFrameKey);
 	if (SelectedFrame && !SelectedFrame->HasRoot(LinkStateFrameKey))
@@ -6627,22 +6618,22 @@ bool FStateTreeExecutionContext::SelectStateInternal_Linked(
 			);
 			return false;
 		}
-		}
+	}
 
 	const bool bIsNewFrame = SelectedFrame == nullptr;
 	if (bIsNewFrame)
-			{
+	{
 		// Note. Adding to TemporaryFrame can invalidate TargetFrame.
 		FStateTreeIndex16 ExternalDataBaseIndex;
 		FStateTreeDataHandle GlobalParameterDataHandle;
 		FStateTreeIndex16 GlobalInstanceIndexBase;
-				{
+		{
 			const FStateTreeExecutionFrame* NextFrame = FindExecutionFrame(OutSelectionResult.SelectedFrames.Last(), MakeConstArrayView(Exec.ActiveFrames), MakeConstArrayView(OutSelectionResult.TemporaryFrames));
 			check(NextFrame);
 			ExternalDataBaseIndex = NextFrame->ExternalDataBaseIndex;
 			GlobalParameterDataHandle = NextFrame->GlobalParameterDataHandle;
 			GlobalInstanceIndexBase = NextFrame->GlobalInstanceIndexBase;
-				}
+		}
 
 		constexpr bool bIsGlobalFrame = false;
 		FStateTreeExecutionFrame& NewFrame = OutSelectionResult.MakeAndAddTemporaryFrame(FActiveFrameID(Storage.GenerateUniqueId()), LinkStateFrameKey, bIsGlobalFrame);
@@ -6652,7 +6643,7 @@ bool FStateTreeExecutionContext::SelectStateInternal_Linked(
 		NewFrame.StateParameterDataHandle = TargetState.ParameterDataHandle; // Temporary allocated earlier if did not exists.
 
 		SelectedFrame = &NewFrame;
-				}
+	}
 
 	OutSelectionResult.SelectedFrames.Add(SelectedFrame->FrameID);
 
@@ -6667,12 +6658,12 @@ bool FStateTreeExecutionContext::SelectStateInternal_Linked(
 	const bool bSucceededToSelectState = SelectStateInternal(Args, NewInternalArgs, OutSelectionResult);
 
 	if (!bSucceededToSelectState)
-				{
+	{
 		OutSelectionResult.SelectedFrames.Pop();
-				}
+	}
 
 	return bSucceededToSelectState;
-				}
+}
 
 bool FStateTreeExecutionContext::SelectStateInternal_LinkedAsset(
 	const FSelectStateArguments& Args,
@@ -6682,7 +6673,7 @@ bool FStateTreeExecutionContext::SelectStateInternal_LinkedAsset(
 	const FCompactStateTreeState& NextState,
 	const UStateTree* NextLinkedStateAsset,
 	bool bShouldCreateNewState)
-				{
+{
 	using namespace UE::StateTree;
 	using namespace UE::StateTree::ExecutionContext;
 	using namespace UE::StateTree::ExecutionContext::Private;
@@ -6690,32 +6681,32 @@ bool FStateTreeExecutionContext::SelectStateInternal_LinkedAsset(
 	FStateTreeExecutionState& Exec = GetExecState();
 
 	if (NextLinkedStateAsset == nullptr || NextLinkedStateAsset->States.Num() == 0)
-					{
+	{
 		STATETREE_LOG(Error, TEXT("%hs: The linked State Tree is invalid. '%s' using StateTree '%s'."),
 			__FUNCTION__,
 			*Owner.GetName(),
 			*NextStateTree->GetFullName()
 		);
 		return false;
-					}
+	}
 
 	FStateTreeStateHandle NextLinkedStateRoot = FStateTreeStateHandle::Root;
 	const bool bHasMissingState = InternalArgs.MissingStatesToReachTarget.Num() > 0 && Args.Behavior == ESelectStateBehavior::Forced;
 	if (bHasMissingState)
-					{
+	{
 		// In a force transition, the root could be different from what is expected.
 		//Ex: a previous transition go to root, then another transition go to another top level state (new root)
 		if (InternalArgs.MissingStatesToReachTarget[0].StateTree != NextLinkedStateAsset)
-						{
+		{
 			STATETREE_LOG(Error, TEXT("%hs: The missing state is not from the same state tree. '%s' using StateTree '%s'."),
 				__FUNCTION__,
 				*Owner.GetName(),
 				*NextStateTree->GetFullName()
 			);
 			return false;
-							}
+		}
 		NextLinkedStateRoot = InternalArgs.MissingStatesToReachTarget[0].StateHandle;
-						}
+	}
 
 	const FExecutionFrameKey LinkStateFrameKey = FExecutionFrameKey(NextLinkedStateAsset, NextLinkedStateRoot);
 
@@ -6732,7 +6723,7 @@ bool FStateTreeExecutionContext::SelectStateInternal_LinkedAsset(
 			*NextStateTree->GetFullName()
 		);
 		return false;
-					}
+	}
 
 	if (PreventRecursionCheck(LinkStateFrameKey, OutSelectionResult.SelectedFrames, MakeConstArrayView(Exec.ActiveFrames), MakeConstArrayView(OutSelectionResult.TemporaryFrames)))
 	{
@@ -6744,11 +6735,11 @@ bool FStateTreeExecutionContext::SelectStateInternal_LinkedAsset(
 			*NextStateTree->GetFullName()
 		);
 		return false;
-				}
+	}
 
 	const FCompactStateTreeFrame* LinkStateTreeFrame = FindStateTreeFrame(LinkStateFrameKey);
 	if (LinkStateTreeFrame == nullptr)
-				{
+	{
 		STATETREE_LOG(Error, TEXT("%hs: The frame '%s' from '%s' does not exist. '%s' using StateTree '%s'."),
 			__FUNCTION__,
 			*GetSafeStateName(LinkStateFrameKey.StateTree, LinkStateFrameKey.RootState),
@@ -6757,18 +6748,18 @@ bool FStateTreeExecutionContext::SelectStateInternal_LinkedAsset(
 			*NextStateTree->GetFullName()
 		);
 		return false;
-				}
+	}
 
 	// Do we have an existing frame.
 	FStateTreeExecutionFrame* SelectedFrame = SelectedFrameLinkedFrame(Exec, bShouldCreateNewState, InternalArgs.MissingActiveStates, LinkStateFrameKey);
 	if (SelectedFrame && !SelectedFrame->HasRoot(LinkStateFrameKey))
-				{
+	{
 		if (Args.Behavior == ESelectStateBehavior::Forced)
-					{
+		{
 			SelectedFrame = nullptr;
-				}
-				else
-				{
+		}
+		else
+		{
 			STATETREE_LOG(Error, TEXT("%hs: The frame '%s' from '%s' does not have the same root as the active frame. '%s' using StateTree '%s'."),
 				__FUNCTION__,
 				*GetSafeStateName(LinkStateFrameKey.StateTree, LinkStateFrameKey.RootState),
@@ -6778,15 +6769,15 @@ bool FStateTreeExecutionContext::SelectStateInternal_LinkedAsset(
 			);
 			return false;
 		}
-				}
+	}
 
 	const bool bIsNewFrame = SelectedFrame == nullptr;
 	if (bIsNewFrame)
-				{
+	{
 		// Collect external data if needed
 		const FStateTreeIndex16 ExternalDataBaseIndex = CollectExternalData(LinkStateFrameKey.StateTree);
 		if (!ExternalDataBaseIndex.IsValid())
-					{
+		{
 			STATETREE_LOG(VeryVerbose, TEXT("%hs: Cannot select state '%s' because failed to collect external data for nested tree '%s' from '%s'. '%s' using StateTree '%s'."),
 				__FUNCTION__,
 				*GetSafeStateName(LinkStateFrameKey.StateTree, LinkStateFrameKey.RootState),
@@ -6812,7 +6803,7 @@ bool FStateTreeExecutionContext::SelectStateInternal_LinkedAsset(
 	OutSelectionResult.SelectedFrames.Add(SelectedFrame->FrameID);
 
 	if (bIsNewFrame)
-					{
+	{
 		// Start global tasks and evaluators temporarily, so that their data is available already during select.
 		const FStateTreeExecutionFrame* NextParentFrame = FindExecutionFrame(OutSelectionResult.SelectedFrames.Last(1), MakeConstArrayView(Exec.ActiveFrames), MakeConstArrayView(OutSelectionResult.TemporaryFrames));
 		FStateTreeExecutionFrame* NextFrame = FindExecutionFrame(OutSelectionResult.SelectedFrames.Last(), MakeArrayView(Exec.ActiveFrames), MakeArrayView(OutSelectionResult.TemporaryFrames));
@@ -6821,7 +6812,7 @@ bool FStateTreeExecutionContext::SelectStateInternal_LinkedAsset(
 
 		const EStateTreeRunStatus StartResult = StartTemporaryEvaluatorsAndGlobalTasks(NextParentFrame, *NextFrame);
 		if (StartResult != EStateTreeRunStatus::Running)
-						{
+		{
 			STATETREE_LOG(VeryVerbose, TEXT("%hs: Cannot select state '%s' because cannot start nested tree's '%s' global tasks and evaluators. '%s' using StateTree '%s'."),
 				__FUNCTION__,
 				*GetSafeStateName(LinkStateFrameKey.StateTree, LinkStateFrameKey.RootState),
@@ -6835,8 +6826,8 @@ bool FStateTreeExecutionContext::SelectStateInternal_LinkedAsset(
 
 			OutSelectionResult.SelectedFrames.Pop();
 			return false;
-						}
-					}
+		}
+	}
 
 	// Select the root state of the new frame.
 	FStateHandleContext RootState = FStateHandleContext(LinkStateFrameKey.StateTree, LinkStateFrameKey.RootState);
@@ -6860,12 +6851,12 @@ bool FStateTreeExecutionContext::SelectStateInternal_LinkedAsset(
 			constexpr EStateTreeRunStatus CompletionStatus = EStateTreeRunStatus::Stopped;
 			StopTemporaryEvaluatorsAndGlobalTasks(NextParentFrame, *NextFrame, CompletionStatus);
 			CleanFrame(Exec, SelectedFrame->FrameID);
-				}
+		}
 		OutSelectionResult.SelectedFrames.Pop();
-			}
+	}
 
 	return bSucceededToSelectState;
-			}
+}
 
 bool FStateTreeExecutionContext::SelectStateInternal_TrySelectChildrenInOrder(
 	const FSelectStateArguments& Args,
@@ -6874,32 +6865,32 @@ bool FStateTreeExecutionContext::SelectStateInternal_TrySelectChildrenInOrder(
 	TNotNull<const UStateTree*> NextStateTree,
 	const FCompactStateTreeState& NextState,
 	const FStateTreeStateHandle NextStateHandle)
-		{
+{
 	using namespace UE::StateTree::ExecutionContext;
 
 	if (!NextState.HasChildren())
-			{
+	{
 		// Select this state (For backwards compatibility)
 		UE_STATETREE_DEBUG_STATE_EVENT(this, NextStateHandle, EStateTreeTraceEventType::OnStateSelected);
 		return true;
 	}
 
-				UE_STATETREE_DEBUG_SCOPED_STATE_PHASE(this, NextStateHandle, EStateTreeUpdatePhase::TrySelectBehavior);
+	UE_STATETREE_DEBUG_SCOPED_STATE_PHASE(this, NextStateHandle, EStateTreeUpdatePhase::TrySelectBehavior);
 
-				// If the state has children, proceed to select children.
+	// If the state has children, proceed to select children.
 	bool bSucceededToSelectState = false;
 	for (uint16 ChildStateIndex = NextState.ChildrenBegin; ChildStateIndex < NextState.ChildrenEnd; ChildStateIndex = NextStateTree->States[ChildStateIndex].GetNextSibling())
-				{
+	{
 		FStateHandleContext ChildState = FStateHandleContext(NextStateTree, FStateTreeStateHandle(ChildStateIndex));
 		FSelectStateInternalArguments NewInternalArgs = InternalArgs;
 		NewInternalArgs.MissingStatesToReachTarget = MakeArrayView(&ChildState, 1);
 		if (SelectStateInternal(Args, NewInternalArgs, OutSelectionResult))
-					{
-						// Selection succeeded
-						bSucceededToSelectState = true;
-						break;
-					}
-				}
+		{
+			// Selection succeeded
+			bSucceededToSelectState = true;
+			break;
+		}
+	}
 
 	return bSucceededToSelectState;
 }
@@ -6911,40 +6902,40 @@ bool FStateTreeExecutionContext::SelectStateInternal_TrySelectChildrenAtRandom(
 	TNotNull<const UStateTree*> NextStateTree,
 	const FCompactStateTreeState& NextState,
 	const FStateTreeStateHandle NextStateHandle)
-				{
+{
 	using namespace UE::StateTree::ExecutionContext;
 
 	if (!NextState.HasChildren())
-			{
-				// Select this state (For backwards compatibility)
-				UE_STATETREE_DEBUG_STATE_EVENT(this, NextStateHandle, EStateTreeTraceEventType::OnStateSelected);
+	{
+		// Select this state (For backwards compatibility)
+		UE_STATETREE_DEBUG_STATE_EVENT(this, NextStateHandle, EStateTreeTraceEventType::OnStateSelected);
 		return true;
-			}
+	}
+
+	UE_STATETREE_DEBUG_SCOPED_STATE_PHASE(this, NextStateHandle, EStateTreeUpdatePhase::TrySelectBehavior);
 
-				UE_STATETREE_DEBUG_SCOPED_STATE_PHASE(this, NextStateHandle, EStateTreeUpdatePhase::TrySelectBehavior);
-				
 	TArray<uint16, TInlineAllocator<16>> NextLevelChildStates;
 	for (uint16 ChildStateIndex = NextState.ChildrenBegin; ChildStateIndex < NextState.ChildrenEnd; ChildStateIndex = NextStateTree->States[ChildStateIndex].GetNextSibling())
-				{
+	{
 		NextLevelChildStates.Push(ChildStateIndex);
-				}
+	}
 
 	const FStateTreeExecutionState& Exec = GetExecState();
-				while (!NextLevelChildStates.IsEmpty())
-				{
-					const int32 ChildStateIndex = Exec.RandomStream.RandRange(0, NextLevelChildStates.Num() - 1);
+	while (!NextLevelChildStates.IsEmpty())
+	{
+		const int32 ChildStateIndex = Exec.RandomStream.RandRange(0, NextLevelChildStates.Num() - 1);
 		FStateHandleContext ChildState = FStateHandleContext(NextStateTree, FStateTreeStateHandle(NextLevelChildStates[ChildStateIndex]));
 		FSelectStateInternalArguments NewInternalArgs = InternalArgs;
 		NewInternalArgs.MissingStatesToReachTarget = MakeArrayView(&ChildState, 1);
 		if (SelectStateInternal(Args, NewInternalArgs, OutSelectionResult))
-					{
-						// Selection succeeded
+		{
+			// Selection succeeded
 			return true;
-					}
+		}
 
-					constexpr EAllowShrinking AllowShrinking = EAllowShrinking::No;
-					NextLevelChildStates.RemoveAtSwap(ChildStateIndex, AllowShrinking);
-				}
+		constexpr EAllowShrinking AllowShrinking = EAllowShrinking::No;
+		NextLevelChildStates.RemoveAtSwap(ChildStateIndex, AllowShrinking);
+	}
 
 	return false;
 }
@@ -6956,20 +6947,20 @@ bool FStateTreeExecutionContext::SelectStateInternal_TrySelectChildrenWithHighes
 	TNotNull<const UStateTree*> NextStateTree,
 	const FCompactStateTreeState& NextState,
 	const FStateTreeStateHandle NextStateHandle)
-				{
+{
 	using namespace UE::StateTree::ExecutionContext;
 
 	if (!NextState.HasChildren())
-			{
-				// Select this state (For backwards compatibility)
-				UE_STATETREE_DEBUG_STATE_EVENT(this, NextStateHandle, EStateTreeTraceEventType::OnStateSelected);
+	{
+		// Select this state (For backwards compatibility)
+		UE_STATETREE_DEBUG_STATE_EVENT(this, NextStateHandle, EStateTreeTraceEventType::OnStateSelected);
 		return true;
-			}
+	}
 
 	UE_STATETREE_DEBUG_SCOPED_STATE_PHASE(this, NextStateHandle, EStateTreeUpdatePhase::TrySelectBehavior);
 
 	TArray<TPair<uint16, float>, TInlineAllocator<16>> NextLevelChildStates;
-		{
+	{
 		using namespace UE::StateTree::ExecutionContext::Private;
 		const FStateTreeExecutionFrame* NextParentFrame = OutSelectionResult.SelectedFrames.Num() > 1
 			? FindExecutionFrame(OutSelectionResult.SelectedFrames.Last(1), MakeConstArrayView(GetExecState().ActiveFrames), MakeConstArrayView(OutSelectionResult.TemporaryFrames))
@@ -6978,30 +6969,30 @@ bool FStateTreeExecutionContext::SelectStateInternal_TrySelectChildrenWithHighes
 		check(NextFrame);
 
 		for (uint16 ChildState = NextState.ChildrenBegin; ChildState < NextState.ChildrenEnd; ChildState = NextStateTree->States[ChildState].GetNextSibling())
-				{
+		{
 			const FCompactStateTreeState& CurrentState = NextStateTree->States[ChildState];
 			const float Score = EvaluateUtilityWithValidation(NextParentFrame, *NextFrame, FStateTreeStateHandle(ChildState), CurrentState.UtilityConsiderationsBegin, CurrentState.UtilityConsiderationsNum, CurrentState.Weight);
 			NextLevelChildStates.Emplace(ChildState, Score);
 		}
-				}
+	}
 
-				while (!NextLevelChildStates.IsEmpty())
-				{
-					//Find one with highest score in the remaining candidates
+	while (!NextLevelChildStates.IsEmpty())
+	{
+		//Find one with highest score in the remaining candidates
 		float HighestScore = -std::numeric_limits<float>::infinity();
-					int32 ArrayIndexWithHighestScore = INDEX_NONE;
-					for (int32 Index = 0; Index < NextLevelChildStates.Num(); ++Index)
-					{
+		int32 ArrayIndexWithHighestScore = INDEX_NONE;
+		for (int32 Index = 0; Index < NextLevelChildStates.Num(); ++Index)
+		{
 			const float Score = NextLevelChildStates[Index].Get<1>();
-						if (Score > HighestScore)
-						{
-							HighestScore = Score;
-							ArrayIndexWithHighestScore = Index;
-						}
-					}
+			if (Score > HighestScore)
+			{
+				HighestScore = Score;
+				ArrayIndexWithHighestScore = Index;
+			}
+		}
 
 		if (!NextLevelChildStates.IsValidIndex(ArrayIndexWithHighestScore))
-					{
+		{
 			return false;
 		}
 
@@ -7009,17 +7000,17 @@ bool FStateTreeExecutionContext::SelectStateInternal_TrySelectChildrenWithHighes
 		FSelectStateInternalArguments NewInternalArgs = InternalArgs;
 		NewInternalArgs.MissingStatesToReachTarget = MakeArrayView(&ChildState, 1);
 		if (SelectStateInternal(Args, NewInternalArgs, OutSelectionResult))
-						{
+		{
 			return true;
-						}
-						
-						// Disqualify the state we failed to enter
+		}
+
+		// Disqualify the state we failed to enter
 		constexpr EAllowShrinking AllowShrinking = EAllowShrinking::No;
 		NextLevelChildStates.RemoveAtSwap(ArrayIndexWithHighestScore, AllowShrinking);
-					}
+	}
 
 	return false;
-				}
+}
 
 bool FStateTreeExecutionContext::SelectStateInternal_TrySelectChildrenAtRandomWeightedByUtility(
 	const FSelectStateArguments& Args,
@@ -7028,21 +7019,21 @@ bool FStateTreeExecutionContext::SelectStateInternal_TrySelectChildrenAtRandomWe
 	TNotNull<const UStateTree*> NextStateTree,
 	const FCompactStateTreeState& NextState,
 	const FStateTreeStateHandle NextStateHandle)
-				{
+{
 	using namespace UE::StateTree::ExecutionContext;
 
 	if (!NextState.HasChildren())
-			{
-				// Select this state (For backwards compatibility)
-				UE_STATETREE_DEBUG_STATE_EVENT(this, NextStateHandle, EStateTreeTraceEventType::OnStateSelected);
+	{
+		// Select this state (For backwards compatibility)
+		UE_STATETREE_DEBUG_STATE_EVENT(this, NextStateHandle, EStateTreeTraceEventType::OnStateSelected);
 		return true;
-			}
+	}
 
 	UE_STATETREE_DEBUG_SCOPED_STATE_PHASE(this, NextStateHandle, EStateTreeUpdatePhase::TrySelectBehavior);
 
 	TArray<TTuple<uint16, float>, TInlineAllocator<16>> NextLevelChildStates;
 	float TotalScore = 0.0f;
-		{
+	{
 		using namespace UE::StateTree::ExecutionContext::Private;
 		const FStateTreeExecutionFrame* NextParentFrame = OutSelectionResult.SelectedFrames.Num() > 1
 			? FindExecutionFrame(OutSelectionResult.SelectedFrames.Last(1), MakeConstArrayView(GetExecState().ActiveFrames), MakeConstArrayView(OutSelectionResult.TemporaryFrames))
@@ -7051,42 +7042,42 @@ bool FStateTreeExecutionContext::SelectStateInternal_TrySelectChildrenAtRandomWe
 		check(NextFrame);
 
 		for (uint16 ChildState = NextState.ChildrenBegin; ChildState < NextState.ChildrenEnd; ChildState = NextStateTree->States[ChildState].GetNextSibling())
-			{
+		{
 			const FCompactStateTreeState& CurrentState = NextStateTree->States[ChildState];
 			const float Score = EvaluateUtilityWithValidation(NextParentFrame, *NextFrame, FStateTreeStateHandle(ChildState), CurrentState.UtilityConsiderationsBegin, CurrentState.UtilityConsiderationsNum, CurrentState.Weight);
 			if (Score > 0.0f)
-				{
+			{
 				NextLevelChildStates.Emplace(ChildState, Score);
 				TotalScore += Score;
 			}
 		}
-				}
+	}
 
 	const FStateTreeExecutionState& Exec = GetExecState();
-				while (!NextLevelChildStates.IsEmpty())
-				{
-					const float RandomScore = Exec.RandomStream.FRand() * TotalScore;
+	while (!NextLevelChildStates.IsEmpty())
+	{
+		const float RandomScore = Exec.RandomStream.FRand() * TotalScore;
 		float AccumulatedScore = 0.0f;
-					for (int32 Index = 0; Index < NextLevelChildStates.Num(); ++Index)
-					{
-						const TTuple<uint16, float>& StateScorePair = NextLevelChildStates[Index];
-						const uint16 StateIndex = StateScorePair.Key;
-						const float StateScore = StateScorePair.Value;
-						AccumulatedScore += StateScore;
+		for (int32 Index = 0; Index < NextLevelChildStates.Num(); ++Index)
+		{
+			const TTuple<uint16, float>& StateScorePair = NextLevelChildStates[Index];
+			const uint16 StateIndex = StateScorePair.Key;
+			const float StateScore = StateScorePair.Value;
+			AccumulatedScore += StateScore;
 
 			if (RandomScore <= AccumulatedScore || (Index == (NextLevelChildStates.Num() - 1)))
-						{
+			{
 				FStateHandleContext ChildState = FStateHandleContext(NextStateTree, FStateTreeStateHandle(StateIndex));
 				FSelectStateInternalArguments NewInternalArgs = InternalArgs;
 				NewInternalArgs.MissingStatesToReachTarget = MakeArrayView(&ChildState, 1);
 				if (SelectStateInternal(Args, NewInternalArgs, OutSelectionResult))
-							{
+				{
 					return true;
-							}
+				}
 
 				// Disqualify the state we failed to enter
-							constexpr EAllowShrinking AllowShrinking = EAllowShrinking::No;
-							NextLevelChildStates.RemoveAtSwap(Index, AllowShrinking);
+				constexpr EAllowShrinking AllowShrinking = EAllowShrinking::No;
+				NextLevelChildStates.RemoveAtSwap(Index, AllowShrinking);
 				break;
 			}
 		}
@@ -7127,25 +7118,25 @@ bool FStateTreeExecutionContext::SelectStateInternal_TryFollowTransitions(
 		if (Transition.Priority <= CurrentPriority)
 		{
 			continue;
-						}
+		}
 
 		// Skip completion transitions
 		if (EnumHasAnyFlags(Transition.Trigger, EStateTreeTransitionTrigger::OnStateCompleted))
 		{
 			continue;
-					}
+		}
 
 		// Skip invalid state or completion state
 		if (!Transition.State.IsValid() || Transition.State.IsCompletionState())
-					{
+		{
 			continue;
-					}
+		}
 
 		// Cannot follow transitions with delay.
 		if (Transition.HasDelay())
 		{
 			continue;
-				}
+		}
 
 		// Try to prevent (infinite) loops in the selection.
 		const bool bSelectionContainsState = OutSelectionResult.SelectedStates.ContainsByPredicate(
@@ -7154,7 +7145,7 @@ bool FStateTreeExecutionContext::SelectStateInternal_TryFollowTransitions(
 				return State == Other.GetStateHandle();
 			});
 		if (bSelectionContainsState)
-				{
+		{
 			STATETREE_LOG(Warning, TEXT("%hs: Loop detected when trying to select state %s from '%s'. Prior states: %s.  '%s' using StateTree '%s'.")
 				, __FUNCTION__
 				, *GetSafeStateName(NextStateTree, Transition.State)
@@ -7163,7 +7154,7 @@ bool FStateTreeExecutionContext::SelectStateInternal_TryFollowTransitions(
 				, *Owner.GetName()
 				, *NextStateTree->GetFullName());
 			continue;
-				}
+		}
 
 		FSharedEventInlineArray TransitionEvents;
 		GetTriggerTransitionEvent(Transition, Storage, Args.TransitionEvent, GetEventsToProcessView(), TransitionEvents);
@@ -7198,10 +7189,10 @@ bool FStateTreeExecutionContext::SelectStateInternal_TryFollowTransitions(
 				FSelectStateResult StateSelectionResult;
 				StateSelectionResult.TemporaryFrames = OutSelectionResult.TemporaryFrames;
 				if (SelectState(SelectStateArgs, StateSelectionResult))
-			{
+				{
 					OutSelectionResult = MoveTemp(StateSelectionResult);
 					CurrentPriority = Transition.Priority;
-				bSucceededToSelectState = true;
+					bSucceededToSelectState = true;
 
 					//@todo sort the transition by priority at compile time. This will solve having to loop back once we found a valid transition.
 					// Consume the transition event
@@ -7211,11 +7202,11 @@ bool FStateTreeExecutionContext::SelectStateInternal_TryFollowTransitions(
 					}
 
 					// Cannot return because higher priority transitions may override the selection. 
-				break;
+					break;
+				}
 			}
 		}
 	}
-	}
 
 	return bSucceededToSelectState;
 }
@@ -7250,7 +7241,7 @@ FString FStateTreeExecutionContext::DebugGetStatePath(TConstArrayView<FStateTree
 {
 	FString StatePath;
 	const UStateTree* LastStateTree = &RootStateTree;
-		
+
 	for (const FStateTreeExecutionFrame& Frame : ActiveFrames)
 	{
 		if (!ensure(Frame.StateTree))
@@ -7270,14 +7261,14 @@ FString FStateTreeExecutionContext::DebugGetStatePath(TConstArrayView<FStateTree
 			StatePath.Appendf(TEXT("[%s]"), *GetNameSafe(Frame.StateTree));
 			LastStateTree = Frame.StateTree;
 		}
-		
+
 		for (int32 i = 0; i < Num; i++)
 		{
 			const FCompactStateTreeState& State = Frame.StateTree->States[Frame.ActiveStates[i].Index];
 			StatePath.Appendf(TEXT("%s%s"), i == 0 ? TEXT("") : TEXT("."), *State.Name.ToString());
 		}
 	}
-		
+
 	return StatePath;
 }
 
diff --git a/Engine/Plugins/Runtime/StateTree/Source/StateTreeModule/Public/StateTreeExecutionContext.h b/Engine/Plugins/Runtime/StateTree/Source/StateTreeModule/Public/StateTreeExecutionContext.h
--- a/Engine/Plugins/Runtime/StateTree/Source/StateTreeModule/Public/StateTreeExecutionContext.h
+++ b/Engine/Plugins/Runtime/StateTree/Source/StateTreeModule/Public/StateTreeExecutionContext.h
@@ -1278,13 +1278,13 @@ PRAGMA_ENABLE_DEPRECATION_WARNINGS
 		const FStateTreeStateHandle TargetStateHandle);
 
 	/** @return StateTree execution state from the instance storage. */
-	FStateTreeExecutionState& GetExecState()
+	[[nodiscard]] FStateTreeExecutionState& GetExecState()
 	{
 		return Storage.GetMutableExecutionState();
 	}
 
 	/** @return const StateTree execution state from the instance storage. */
-	const FStateTreeExecutionState& GetExecState() const
+	[[nodiscard]] const FStateTreeExecutionState& GetExecState() const
 	{
 		return Storage.GetExecutionState();
 	}
@@ -1293,40 +1293,40 @@ PRAGMA_ENABLE_DEPRECATION_WARNINGS
 	UE_API void SetUpdatePhaseInExecutionState(FStateTreeExecutionState& ExecutionState, EStateTreeUpdatePhase UpdatePhase) const;
 
 	/** @return String describing state status for logging and debug. */
-	UE_API FString GetStateStatusString(const FStateTreeExecutionState& ExecState) const;
+	[[nodiscard]] UE_API FString GetStateStatusString(const FStateTreeExecutionState& ExecState) const;
 
 	/** @return String describing state name for logging and debug. */
-	UE_API FString GetSafeStateName(const FStateTreeExecutionFrame& CurrentFrame, const FStateTreeStateHandle State) const;
+	[[nodiscard]] UE_API FString GetSafeStateName(const FStateTreeExecutionFrame& CurrentFrame, const FStateTreeStateHandle State) const;
 
 	/** @return String describing state name for logging and debug. */
-	UE_API FString GetSafeStateName(const UStateTree* StateTree, const FStateTreeStateHandle State) const;
+	[[nodiscard]] UE_API FString GetSafeStateName(const UStateTree* StateTree, const FStateTreeStateHandle State) const;
 
 	/** @return String describing full path of an activate state for logging and debug. */
-	UE_API FString DebugGetStatePath(TConstArrayView<FStateTreeExecutionFrame> ActiveFrames, const FStateTreeExecutionFrame* CurrentFrame = nullptr, const int32 ActiveStateIndex = INDEX_NONE) const;
+	[[nodiscard]] UE_API FString DebugGetStatePath(TConstArrayView<FStateTreeExecutionFrame> ActiveFrames, const FStateTreeExecutionFrame* CurrentFrame = nullptr, const int32 ActiveStateIndex = INDEX_NONE) const;
 
 	/** @return String describing all events that are currently being processed  for logging and debug. */
-	UE_API FString DebugGetEventsAsString() const;
+	[[nodiscard]] UE_API FString DebugGetEventsAsString() const;
 
 	/** @return data view of the specified handle relative to given frame. */
-	UE_API FStateTreeDataView GetDataView(const FStateTreeExecutionFrame* ParentFrame, const FStateTreeExecutionFrame& CurrentFrame, const FStateTreeDataHandle Handle);
-	UE_API FStateTreeDataView GetDataView(const FStateTreeExecutionFrame* ParentFrame, const FStateTreeExecutionFrame& CurrentFrame, const FPropertyBindingCopyInfo& CopyInfo);
+	[[nodiscard]] UE_API FStateTreeDataView GetDataView(const FStateTreeExecutionFrame* ParentFrame, const FStateTreeExecutionFrame& CurrentFrame, const FStateTreeDataHandle Handle);
+	[[nodiscard]] UE_API FStateTreeDataView GetDataView(const FStateTreeExecutionFrame* ParentFrame, const FStateTreeExecutionFrame& CurrentFrame, const FPropertyBindingCopyInfo& CopyInfo);
 
 	/** @return true if handle source is valid cified handle relative to given frame. */
-	UE_API bool IsHandleSourceValid(const FStateTreeExecutionFrame* ParentFrame, const FStateTreeExecutionFrame& CurrentFrame, const FStateTreeDataHandle Handle) const;
-	UE_API bool IsHandleSourceValid(const FStateTreeExecutionFrame* ParentFrame, const FStateTreeExecutionFrame& CurrentFrame, const FPropertyBindingCopyInfo& CopyInfo) const;
+	[[nodiscard]] UE_API bool IsHandleSourceValid(const FStateTreeExecutionFrame* ParentFrame, const FStateTreeExecutionFrame& CurrentFrame, const FStateTreeDataHandle Handle) const;
+	[[nodiscard]] UE_API bool IsHandleSourceValid(const FStateTreeExecutionFrame* ParentFrame, const FStateTreeExecutionFrame& CurrentFrame, const FPropertyBindingCopyInfo& CopyInfo) const;
 
 	/** @return data view of the specified handle relative to the given frame, or tries to find a matching temporary instance. */
-	UE_API FStateTreeDataView GetDataViewOrTemporary(const FStateTreeExecutionFrame* ParentFrame, const FStateTreeExecutionFrame& CurrentFrame, const FStateTreeDataHandle Handle);
-	UE_API FStateTreeDataView GetDataViewOrTemporary(const FStateTreeExecutionFrame* ParentFrame, const FStateTreeExecutionFrame& CurrentFrame, const FPropertyBindingCopyInfo& CopyInfo);
+	[[nodiscard]] UE_API FStateTreeDataView GetDataViewOrTemporary(const FStateTreeExecutionFrame* ParentFrame, const FStateTreeExecutionFrame& CurrentFrame, const FStateTreeDataHandle Handle);
+	[[nodiscard]] UE_API FStateTreeDataView GetDataViewOrTemporary(const FStateTreeExecutionFrame* ParentFrame, const FStateTreeExecutionFrame& CurrentFrame, const FPropertyBindingCopyInfo& CopyInfo);
 
 	/** @return data view of the specified handle from temporary instance. */
-	UE_API FStateTreeDataView GetTemporaryDataView(const FStateTreeExecutionFrame* ParentFrame, const FStateTreeExecutionFrame& CurrentFrame, const FStateTreeDataHandle Handle);
+	[[nodiscard]] UE_API FStateTreeDataView GetTemporaryDataView(const FStateTreeExecutionFrame* ParentFrame, const FStateTreeExecutionFrame& CurrentFrame, const FStateTreeDataHandle Handle);
 
 	/**
 	 * Adds a temporary instance that can be located using frame and data handle later.
 	 * @returns view to the newly added instance. If NewInstanceData is Object wrapper, the new object is returned.
 	 */
-	UE_API FStateTreeDataView AddTemporaryInstance(const FStateTreeExecutionFrame& Frame, const FStateTreeIndex16 OwnerNodeIndex, const FStateTreeDataHandle DataHandle, FConstStructView NewInstanceData);
+	[[nodiscard]] UE_API FStateTreeDataView AddTemporaryInstance(const FStateTreeExecutionFrame& Frame, const FStateTreeIndex16 OwnerNodeIndex, const FStateTreeDataHandle DataHandle, FConstStructView NewInstanceData);
 
 	enum class ECopyBindings : uint8
 	{
