[StateTree] Fix bad indentation that cause the condition to be evaluated for each tasks. Introduced with CL 43425491.

diff --git a/Engine/Plugins/Runtime/StateTree/Source/StateTreeModule/Private/StateTreeExecutionContext.cpp b/Engine/Plugins/Runtime/StateTree/Source/StateTreeModule/Private/StateTreeExecutionContext.cpp
--- a/Engine/Plugins/Runtime/StateTree/Source/StateTreeModule/Private/StateTreeExecutionContext.cpp
+++ b/Engine/Plugins/Runtime/StateTree/Source/StateTreeModule/Private/StateTreeExecutionContext.cpp
@@ -3447,37 +3447,37 @@ void FStateTreeExecutionContext::ExitState(const FSelectStateResult& Args, const
 					STATETREE_LOG(VeryVerbose, TEXT("%*sSkipped 'ExitState' for disabled Task: '%s'"), UE::StateTree::Debug::IndentSize, TEXT(""), *Task.Name.ToString());
 				}
 				CurrentFrame.ActiveTaskNodeIndex = FStateTreeIndex16(TaskIndex - 1);
+			}
 
-				// Call state change events on conditions if needed.
-				if (CurrentState.bHasStateChangeConditions)
+			// Call state change events on conditions if needed.
+			if (CurrentState.bHasStateChangeConditions)
+			{
+				for (int32 ConditionIndex = (CurrentState.EnterConditionsBegin + CurrentState.EnterConditionsNum) - 1; ConditionIndex >= CurrentState.EnterConditionsBegin; --ConditionIndex)
 				{
-					for (int32 ConditionIndex = (CurrentState.EnterConditionsBegin + CurrentState.EnterConditionsNum) - 1; ConditionIndex >= CurrentState.EnterConditionsBegin; --ConditionIndex)
+					const FStateTreeConditionBase& Cond = CurrentFrame.StateTree->Nodes[ConditionIndex].Get<const FStateTreeConditionBase>();
+					if (Cond.bHasShouldCallStateChangeEvents)
 					{
-						const FStateTreeConditionBase& Cond = CurrentFrame.StateTree->Nodes[ConditionIndex].Get<const FStateTreeConditionBase>();
-						if (Cond.bHasShouldCallStateChangeEvents)
+						const bool bShouldCallStateChange = CurrentTransition.ChangeType == EStateTreeStateChangeType::Changed
+							|| (CurrentTransition.ChangeType == EStateTreeStateChangeType::Sustained && Cond.bShouldStateChangeOnReselect);
+
+						if (bShouldCallStateChange)
 						{
-							const bool bShouldCallStateChange = CurrentTransition.ChangeType == EStateTreeStateChangeType::Changed
-								|| (CurrentTransition.ChangeType == EStateTreeStateChangeType::Sustained && Cond.bShouldStateChangeOnReselect);
+							const FStateTreeDataView ConditionInstanceView = GetDataView(CurrentParentFrame, CurrentFrame, Cond.InstanceDataHandle);
+							FNodeInstanceDataScope DataScope(*this, &Cond, ConditionIndex, Cond.InstanceDataHandle, ConditionInstanceView);
 
-							if (bShouldCallStateChange)
+							if (Cond.BindingsBatch.IsValid())
 							{
-								const FStateTreeDataView ConditionInstanceView = GetDataView(CurrentParentFrame, CurrentFrame, Cond.InstanceDataHandle);
-								FNodeInstanceDataScope DataScope(*this, &Cond, ConditionIndex, Cond.InstanceDataHandle, ConditionInstanceView);
-
-								if (Cond.BindingsBatch.IsValid())
-								{
-									// Use validated copy, since we test in situations where the sources are not always valid (e.g. enter conditions may try to access inactive parent state).
-									CopyBatchOnActiveInstances(CurrentParentFrame, CurrentFrame, ConditionInstanceView, Cond.BindingsBatch);
-								}
+								// Use validated copy, since we test in situations where the sources are not always valid (e.g. enter conditions may try to access inactive parent state).
+								CopyBatchOnActiveInstances(CurrentParentFrame, CurrentFrame, ConditionInstanceView, Cond.BindingsBatch);
+							}
 
-								UE_STATETREE_DEBUG_CONDITION_EXIT_STATE(this, CurrentFrame.StateTree, FStateTreeIndex16(ConditionIndex));
-								Cond.ExitState(*this, CurrentTransition);
+							UE_STATETREE_DEBUG_CONDITION_EXIT_STATE(this, CurrentFrame.StateTree, FStateTreeIndex16(ConditionIndex));
+							Cond.ExitState(*this, CurrentTransition);
 
-								// Reset copied properties that might contain object references.
-								if (Cond.BindingsBatch.IsValid())
-								{
-									CurrentFrame.StateTree->PropertyBindings.Super::ResetObjects(Cond.BindingsBatch, ConditionInstanceView);
-								}
+							// Reset copied properties that might contain object references.
+							if (Cond.BindingsBatch.IsValid())
+							{
+								CurrentFrame.StateTree->PropertyBindings.Super::ResetObjects(Cond.BindingsBatch, ConditionInstanceView);
 							}
 						}
 					}
