[StateTree] Hide the completion icon when the schema doesn't use the task completion feature.

diff --git a/Engine/Plugins/Runtime/StateTree/Source/StateTreeEditorModule/Private/Customizations/StateTreeEditorNodeDetails.cpp b/Engine/Plugins/Runtime/StateTree/Source/StateTreeEditorModule/Private/Customizations/StateTreeEditorNodeDetails.cpp
--- a/Engine/Plugins/Runtime/StateTree/Source/StateTreeEditorModule/Private/Customizations/StateTreeEditorNodeDetails.cpp
+++ b/Engine/Plugins/Runtime/StateTree/Source/StateTreeEditorModule/Private/Customizations/StateTreeEditorNodeDetails.cpp
@@ -701,11 +701,18 @@ void FStateTreeEditorNodeDetails::CustomizeHeader(TSharedRef<class IPropertyHand
 		.PasteAction(FUIAction(FExecuteAction::CreateSP(this, &FStateTreeEditorNodeDetails::OnPasteNode)));
 
 	// Task completion
-	if (const FStateTreeEditorNode* Node = UE::StateTreeEditor::EditorNodeUtils::GetCommonNode(StructProperty))
+	bool bShowCompletion = true;
+	if (const UStateTreeEditorData* EditorDataPtr = EditorData.Get())
+	{
+		bShowCompletion = EditorDataPtr->Schema ? EditorDataPtr->Schema->AllowTasksCompletion() : true;
+	}
+	if (bShowCompletion)
 	{
-		if (const FStateTreeTaskBase* TaskBase = Node->Node.GetPtr<FStateTreeTaskBase>())
+		if (const FStateTreeEditorNode* Node = UE::StateTreeEditor::EditorNodeUtils::GetCommonNode(StructProperty))
 		{
-			DescriptionBox->InsertSlot(0)
+			if (const FStateTreeTaskBase* TaskBase = Node->Node.GetPtr<FStateTreeTaskBase>())
+			{
+				DescriptionBox->InsertSlot(0)
 				.AutoWidth()
 				.VAlign(VAlign_Center)
 				.Padding(0.0f, 0.0f, 0.0f, 0.0f)
@@ -723,6 +730,7 @@ void FStateTreeEditorNodeDetails::CustomizeHeader(TSharedRef<class IPropertyHand
 					.IsEnabled(UE::StateTreeEditor::EditorNodeUtils::CanEditTaskConsideredForCompletion(*Node))
 					.Visibility(this, &FStateTreeEditorNodeDetails::GetToggleCompletionTaskVisibility)
 				];
+			}
 		}
 	}
 
