[StateTree] Fix android compilation. Introduced with CL 43425491.

#rnx

diff --git a/Engine/Plugins/Runtime/StateTree/Source/StateTreeModule/Private/StateTreeExecutionContext.cpp b/Engine/Plugins/Runtime/StateTree/Source/StateTreeModule/Private/StateTreeExecutionContext.cpp
--- a/Engine/Plugins/Runtime/StateTree/Source/StateTreeModule/Private/StateTreeExecutionContext.cpp
+++ b/Engine/Plugins/Runtime/StateTree/Source/StateTreeModule/Private/StateTreeExecutionContext.cpp
@@ -2545,7 +2545,7 @@ EStateTreeRunStatus FStateTreeExecutionContext::ForceTransition(const FRecordedS
 			{
 				if (Args.SelectedStates.IsValidIndex(SelectedStateIndex))
 				{
-					Args.SelectionEvents.Emplace(Args.SelectedStates[SelectedStateIndex], Event);
+					Args.SelectionEvents.Add(FSelectionEventWithID{.State = Args.SelectedStates[SelectedStateIndex], .Event = Event});
 				}
 			}
 		}
@@ -6326,7 +6326,8 @@ bool FStateTreeExecutionContext::SelectStateInternal(
 		bool bRemoveSelectionEvents = false;
 		if (StateSelectionEvent.IsValid())
 		{
-			OutSelectionResult.SelectionEvents.Emplace(OutSelectionResult.SelectedStates.Last(), StateSelectionEvent);
+			
+			OutSelectionResult.SelectionEvents.Add(FSelectionEventWithID{.State = OutSelectionResult.SelectedStates.Last(), .Event = StateSelectionEvent});
 			bRemoveSelectionEvents = true;
 		}
 
