[StateTree] Remove test for enabled state. Disabled states cannot be selected and are never in the selected list or in the active list.

#rnx

diff --git a/Engine/Plugins/Runtime/StateTree/Source/StateTreeModule/Private/StateTreeExecutionContext.cpp b/Engine/Plugins/Runtime/StateTree/Source/StateTreeModule/Private/StateTreeExecutionContext.cpp
--- a/Engine/Plugins/Runtime/StateTree/Source/StateTreeModule/Private/StateTreeExecutionContext.cpp
+++ b/Engine/Plugins/Runtime/StateTree/Source/StateTreeModule/Private/StateTreeExecutionContext.cpp
@@ -383,21 +383,20 @@ FStateTreeScheduledTick FStateTreeReadOnlyExecutionContext::GetNextScheduledTick
 			{
 				const FStateTreeStateHandle CurrentHandle = CurrentFrame.ActiveStates[StateIndex];
 				const FCompactStateTreeState& State = CurrentStateTree->States[CurrentHandle.Index];
-				if (State.bEnabled)
+				checkf(State.bEnabled, TEXT("Only enabled states are in ActiveStates."));
+
+				if (State.bHasCustomTickRate)
 				{
-					if (State.bHasCustomTickRate)
-					{
-						CustomTickRate = CustomTickRate.IsSet() ? FMath::Min(CustomTickRate.GetValue(), State.CustomTickRate) : State.CustomTickRate;
-					}
-					else if (!CustomTickRate.IsSet())
+					CustomTickRate = CustomTickRate.IsSet() ? FMath::Min(CustomTickRate.GetValue(), State.CustomTickRate) : State.CustomTickRate;
+				}
+				else if (!CustomTickRate.IsSet())
+				{
+					if (State.DoesRequestTickTasks(bHasEvents)
+						|| State.ShouldTickTransitions(bHasEvents, bHasBroadcastedDelegates)
+						)
 					{
-						if (State.DoesRequestTickTasks(bHasEvents)
-							|| State.ShouldTickTransitions(bHasEvents, bHasBroadcastedDelegates)
-							)
-						{
-							// todo: ShouldTickTransitions has onevent or ontick. both can be already triggered and we are waiting for the delay
-							bHasTaskWithEveryFramesTick = true;
-						}
+						// todo: ShouldTickTransitions has onevent or ontick. both can be already triggered and we are waiting for the delay
+						bHasTaskWithEveryFramesTick = true;
 					}
 				}
 			}
@@ -3133,17 +3132,7 @@ EStateTreeRunStatus FStateTreeExecutionContext::EnterState(FSelectStateResult& A
 			const FCompactStateTreeState& CurrentState = CurrentStateTree->States[CurrentStateHandle.Index];
 			const UE::StateTree::FActiveStateID CurrentStateID = SelectedState.GetStateID();
 
-			// Add only enabled States to the list of active States
-			if (!CurrentState.bEnabled)
-			{
-				ensureMsgf(false, TEXT("%hs: A selected state is not enabled. '%s' should not be in the selected list. '%s' using StateTree '%s'."),
-					__FUNCTION__,
-					*GetStateStatusString(Exec),
-					*GetNameSafe(&Owner),
-					*GetFullNameSafe(&RootStateTree)
-				);
-				break;
-			}
+			checkf(CurrentState.bEnabled, TEXT("Only enabled states are in SelectedStates."));
 
 			// New state. Add it
 			if (!bSustainedState)
@@ -5021,11 +5010,7 @@ bool FStateTreeExecutionContext::TriggerTransitions()
 				const UE::StateTree::FActiveStateID StateID = CurrentFrame.ActiveStates.StateIDs[StateIndex];
 				const FCompactStateTreeState& State = CurrentStateTree->States[StateHandle.Index];
 
-				// Do not process any transitions from a disabled state
-				if (!State.bEnabled)
-				{
-					continue;
-				}
+				checkf(State.bEnabled, TEXT("Only enabled states are in ActiveStates."));
 
 				// Transition tasks.
 				if (State.bHasTransitionTasks)
