[StateTree] Fix bad copy paste that prevent completion transition to trigger when it's the only transition on the state.

#rnx

diff --git a/Engine/Plugins/Runtime/StateTree/Source/StateTreeEditorModule/Private/StateTreeCompiler.cpp b/Engine/Plugins/Runtime/StateTree/Source/StateTreeEditorModule/Private/StateTreeCompiler.cpp
--- a/Engine/Plugins/Runtime/StateTree/Source/StateTreeEditorModule/Private/StateTreeCompiler.cpp
+++ b/Engine/Plugins/Runtime/StateTree/Source/StateTreeEditorModule/Private/StateTreeCompiler.cpp
@@ -1421,7 +1421,7 @@ bool FStateTreeCompiler::CreateStateTransitions()
 				CompactState.bHasTickTriggerTransitions |= EnumHasAnyFlags(Transition.Trigger, EStateTreeTransitionTrigger::OnTick);
 				CompactState.bHasEventTriggerTransitions |= EnumHasAnyFlags(Transition.Trigger, EStateTreeTransitionTrigger::OnEvent);
 				CompactState.bHasDelegateTriggerTransitions |= EnumHasAnyFlags(Transition.Trigger, EStateTreeTransitionTrigger::OnDelegate);
-				CompactState.bHasSucceededTriggerTransitions |= EnumHasAnyFlags(Transition.Trigger, EStateTreeTransitionTrigger::OnStateCompleted);
+				CompactState.bHasCompletedTriggerTransitions |= EnumHasAnyFlags(Transition.Trigger, EStateTreeTransitionTrigger::OnStateCompleted);
 				CompactState.bHasSucceededTriggerTransitions |= EnumHasAnyFlags(Transition.Trigger, EStateTreeTransitionTrigger::OnStateSucceeded);
 				CompactState.bHasFailedTriggerTransitions |= EnumHasAnyFlags(Transition.Trigger, EStateTreeTransitionTrigger::OnStateFailed);
 			}
