[StateTree] Fix GetSharedInstanceData that access the thread local storage in a forked application.

diff --git a/Engine/Plugins/Runtime/StateTree/Source/StateTreeModule/Private/StateTree.cpp b/Engine/Plugins/Runtime/StateTree/Source/StateTreeModule/Private/StateTree.cpp
--- a/Engine/Plugins/Runtime/StateTree/Source/StateTreeModule/Private/StateTree.cpp
+++ b/Engine/Plugins/Runtime/StateTree/Source/StateTreeModule/Private/StateTree.cpp
@@ -106,7 +106,7 @@ UE_AUTORTFM_ALWAYS_OPEN
 static int32 GetThreadIndexForSharedInstanceData()
 {
 	// Create a unique index for each thread.
-	static std::atomic_int ThreadIndexCounter {0};
+	static std::atomic_int ThreadIndexCounter {1};
 	static thread_local int32 ThreadIndex = INDEX_NONE; // Cannot init directly on WinRT
 	if (ThreadIndex == INDEX_NONE)
 	{
@@ -118,7 +118,10 @@ static int32 GetThreadIndexForSharedInstanceData()
 
 TSharedPtr<FStateTreeInstanceData> UStateTree::GetSharedInstanceData() const
 {
-	int32 ThreadIndex = GetThreadIndexForSharedInstanceData();
+	// Only access the thread_local when not on the game thread. This is to prevent issues with a forked process.
+	//@todo It might create an additional FStateTreeInstanceData when all accesses are not on the game thread.
+	//We also have that problem if you mix assets that are single and multi-threaded in the same running application (because the atomic is static and not per asset).
+	const int32 ThreadIndex = IsInGameThread() ? 0 : GetThreadIndexForSharedInstanceData();
 
 	// If shared instance data for this thread exists, return it.
 	{
