[StateTree] fix broken FActiveStatePath::Describe

diff --git a/Engine/Plugins/Runtime/StateTree/Source/StateTreeModule/Private/StateTreeStatePath.cpp b/Engine/Plugins/Runtime/StateTree/Source/StateTreeModule/Private/StateTreeStatePath.cpp
--- a/Engine/Plugins/Runtime/StateTree/Source/StateTreeModule/Private/StateTreeStatePath.cpp
+++ b/Engine/Plugins/Runtime/StateTree/Source/StateTreeModule/Private/StateTreeStatePath.cpp
@@ -188,7 +188,7 @@ namespace UE::StateTree
 					bNewTree = false;
 				}
 
-				if (CurrentTree->GetStates().IsValidIndex(Element.GetStateHandle().Index))
+				if (!CurrentTree->GetStates().IsValidIndex(Element.GetStateHandle().Index))
 				{
 					return MakeError();
 				}
