[StateTree] Update ForceTransition with the new TransitionID. Confirms that the recorded transition data is valid before processing it. In the future, we can reuse the same functionality to debug state (while debugging force the instance to a specific state).

diff --git a/Engine/Plugins/Runtime/StateTree/Source/StateTreeModule/Private/StateTree.cpp b/Engine/Plugins/Runtime/StateTree/Source/StateTreeModule/Private/StateTree.cpp
--- a/Engine/Plugins/Runtime/StateTree/Source/StateTreeModule/Private/StateTree.cpp
+++ b/Engine/Plugins/Runtime/StateTree/Source/StateTreeModule/Private/StateTree.cpp
@@ -155,7 +155,12 @@ TSharedPtr<FStateTreeInstanceData> UStateTree::GetSharedInstanceData() const
 
 bool UStateTree::HasCompatibleContextData(const UStateTree& Other) const
 {
-	if (ContextDataDescs.Num() != Other.ContextDataDescs.Num())
+	return HasCompatibleContextData(&Other);
+}
+
+bool UStateTree::HasCompatibleContextData(TNotNull<const UStateTree*> Other) const
+{
+	if (ContextDataDescs.Num() != Other->ContextDataDescs.Num())
 	{
 		return false;
 	}
@@ -164,7 +169,7 @@ bool UStateTree::HasCompatibleContextData(const UStateTree& Other) const
 	for (int32 Index = 0; Index < Num; Index++)
 	{
 		const FStateTreeExternalDataDesc& Desc = ContextDataDescs[Index];
-		const FStateTreeExternalDataDesc& OtherDesc = Other.ContextDataDescs[Index];
+		const FStateTreeExternalDataDesc& OtherDesc = Other->ContextDataDescs[Index];
 		
 		if (!OtherDesc.Struct 
 			|| !OtherDesc.Struct->IsChildOf(Desc.Struct))
diff --git a/Engine/Plugins/Runtime/StateTree/Source/StateTreeModule/Private/StateTreeExecutionContext.cpp b/Engine/Plugins/Runtime/StateTree/Source/StateTreeModule/Private/StateTreeExecutionContext.cpp
--- a/Engine/Plugins/Runtime/StateTree/Source/StateTreeModule/Private/StateTreeExecutionContext.cpp
+++ b/Engine/Plugins/Runtime/StateTree/Source/StateTreeModule/Private/StateTreeExecutionContext.cpp
@@ -58,6 +58,14 @@ namespace UE::StateTree::ExecutionContext::Private
 		TEXT("You should upgrade your asset. This option is to help migrate to the new behavior.")
 	);
 
+	bool bSetDeprecatedTransitionResultProperties = false;
+	FAutoConsoleVariableRef CVarSetDeprecatedTransitionResultProperties(
+		TEXT("StateTree.SetDeprecatedTransitionResultProperties"),
+		bSetDeprecatedTransitionResultProperties,
+		TEXT("If true, set the FStateTreeTransitionResult deprecated properties when a transition occurs.\n")
+		TEXT("The information is not relevant and not valid but can be needed for backward compatibility.")
+	);
+
 	const FLazyName Name_Tick = "Tick";
 	const FLazyName Name_Start = "Start";
 	const FLazyName Name_Stop = "Stop";
@@ -2468,101 +2476,209 @@ FStateTreeDataView FStateTreeExecutionContext::GetDataView(const FStateTreeExecu
 	return GetDataView(ParentFrame, CurrentFrame, Handle);
 }
 
-//@todo should return TOptional or TValueOrError. We should not trust what is coming from FRecordedStateTreeTransitionResult.
-//@todo recorded should use something else to identify the frame and state. We should not thrust what is coming from FRecordedStateTreeTransitionResult.
-EStateTreeRunStatus FStateTreeExecutionContext::ForceTransition(const FRecordedStateTreeTransitionResult& Transition)
+EStateTreeRunStatus FStateTreeExecutionContext::ForceTransition(const FRecordedStateTreeTransitionResult& RecordedTransition)
 {
 	using namespace UE::StateTree;
 	using namespace UE::StateTree::ExecutionContext;
 	using namespace UE::StateTree::ExecutionContext::Private;
 
-	if (!IsValid())
+	TArray<FStateHandleContext> StateContexts;
+	StateContexts.Reserve(RecordedTransition.States.Num());
+	for (const FRecordedActiveState& RecordedActiveState : RecordedTransition.States)
 	{
-		STATETREE_LOG(Warning, TEXT("%hs: StateTree context is not initialized properly ('%s' using StateTree '%s')"),
-			__FUNCTION__, *GetNameSafe(&Owner), *GetFullNameSafe(&RootStateTree));
-		return EStateTreeRunStatus::Failed;
+		if (RecordedActiveState.StateTree == nullptr)
+		{
+			return EStateTreeRunStatus::Unset;
+		}
+		StateContexts.Add(FStateHandleContext(RecordedActiveState.StateTree, RecordedActiveState.State));
 	}
 
-	const FStateTreeExecutionState& Exec = GetExecState();
-
-	// A reentrant call to ForceTransition or a call from Start, Tick or Stop must be deferred.
-	if (Exec.CurrentPhase != EStateTreeUpdatePhase::Unset)
+	FSelectStateResult SelectStateResult;
+	if (!ForceTransitionInternal(StateContexts, SelectStateResult))
 	{
-		UE_STATETREE_DEBUG_LOG_EVENT(this, Warning, TEXT("Can't force a transition while %s"), *UEnum::GetDisplayValueAsText(Exec.CurrentPhase).ToString());
+		UE_STATETREE_DEBUG_LOG_EVENT(this, Verbose, TEXT("The force transition failed."));
 		return EStateTreeRunStatus::Unset;
 	}
 
-	if (!CollectActiveExternalData())
+	if (SelectStateResult.SelectedStates.Num() != StateContexts.Num())
 	{
-		STATETREE_LOG(Warning, TEXT("%hs: Failed to collect external data ('%s' using StateTree '%s')"),
-			__FUNCTION__, *GetNameSafe(&Owner), *GetFullNameSafe(&RootStateTree));
-		return EStateTreeRunStatus::Failed;
+		UE_STATETREE_DEBUG_LOG_EVENT(this, Verbose, TEXT("The force transition failed to find the target."));
+		return EStateTreeRunStatus::Unset;
 	}
 
-	TOptional<FStateTreeTransitionResult> FoundTransitionResult = MakeTransitionResult(Transition);
-	if (!FoundTransitionResult.IsSet())
+	TArray<FStateTreeSharedEvent> SharedEvents;
+	SharedEvents.Reserve(RecordedTransition.Events.Num());
+	for (const FStateTreeEvent& SelectedEvent : RecordedTransition.Events)
 	{
-		return EStateTreeRunStatus::Unset;
+		SharedEvents.Add(FStateTreeSharedEvent(SelectedEvent));
 	}
-	FStateTreeTransitionResult& TransitionResult = FoundTransitionResult.GetValue();
+	check(RecordedTransition.Events.Num() == SharedEvents.Num());
+	for (int32 RecordedStateIndex = 0; RecordedStateIndex < RecordedTransition.States.Num(); ++RecordedStateIndex)
+	{
+		const FRecordedActiveState& RecordedActiveState = RecordedTransition.States[RecordedStateIndex];
+		if (SharedEvents.IsValidIndex(RecordedActiveState.EventIndex))
+		{
+			const FActiveState& SelectedState = SelectStateResult.SelectedStates[RecordedStateIndex];
+			SelectStateResult.SelectionEvents.Add(FSelectionEventWithID{.State = SelectedState, .Event= SharedEvents[RecordedActiveState.EventIndex]});
+		}
+	}
+
+	FStateTreeTransitionResult TransitionResult;
+	// The SourceFrameID and SourceStateID are left uninitialized on purpose.
+	TransitionResult.TargetState = StateContexts.Last().StateHandle;
+	TransitionResult.CurrentState = FStateTreeStateHandle::Invalid;
+	TransitionResult.CurrentRunStatus = GetExecState().LastTickStatus;
+	TransitionResult.ChangeType = EStateTreeStateChangeType::Changed;
+	TransitionResult.Priority = RecordedTransition.Priority;
 
-	// Create the state argument
-	FSelectStateResult Args;
+	ExitState(SelectStateResult, TransitionResult);
+	return EnterState(SelectStateResult, TransitionResult);
+}
 
-	int32 CurrentActiveFrameIndex = 0;
-	for (FStateTreeExecutionFrame& Frame : TransitionResult.NextActiveFrames)
+namespace UE::StateTree::ExecutionContext::Private
+{
+	bool TestStateContextPath(TNotNull<const FStateTreeExecutionContext*> ExecutionContext, const TArrayView<const FStateHandleContext> StateContexts)
 	{
-		// reuse or new temporary
-		if (Exec.ActiveFrames.IsValidIndex(CurrentActiveFrameIndex)
-			&& Exec.ActiveFrames[CurrentActiveFrameIndex].HasSameRoot(Frame))
+		const UStateTree* PreviousStateTree = ExecutionContext->GetStateTree();
+		bool bNewTree = true;
+		bool bNewFrame = true;
+		for (int32 Index = 0; Index < StateContexts.Num(); ++Index)
 		{
-			Frame.FrameID = Exec.ActiveFrames[CurrentActiveFrameIndex].FrameID;
-			++CurrentActiveFrameIndex;
+			const FStateHandleContext& StateContext = StateContexts[Index];
+			if (StateContext.StateTree == nullptr || StateContext.StateTree != PreviousStateTree)
+			{
+				// Child state has to have the same state tree.
+				return false;
+			}
+
+			const FCompactStateTreeState* State = StateContext.StateTree->GetStateFromHandle(StateContext.StateHandle);
+			if (State == nullptr)
+			{
+				// The handle do not exist in the asset.
+				return false;
+			}
+			if (bNewFrame || bNewTree)
+			{
+				if (StateContext.StateTree->GetFrameFromHandle(StateContext.StateHandle) == nullptr)
+				{
+					// The state is a new frame (from a linked or linked asset) but is not compiled as a frame.
+					return false;
+				}
+			}
+			if (!bNewFrame)
+			{
+				if (State->Parent != StateContexts[Index - 1].StateHandle)
+				{
+					// Current state is not a child of previous state.
+					return false;
+				}
+			}
+			bNewFrame = false;
+			bNewTree = false;
+
+			if (State->Type == EStateTreeStateType::LinkedAsset)
+			{
+				bNewTree = true;
+				bNewFrame = true;
+				if (Index + 1 >= StateContexts.Num())
+				{
+					// The path cannot end with a linked asset state.
+					return false;
+				}
+				const FStateHandleContext& NextStateContext = StateContexts[Index + 1];
+				if (NextStateContext.StateTree == nullptr
+					|| !NextStateContext.StateTree->IsReadyToRun()
+					|| !NextStateContext.StateTree->HasCompatibleContextData(ExecutionContext->GetStateTree())
+					|| NextStateContext.StateTree->GetSchema()->GetClass() != ExecutionContext->GetStateTree()->GetSchema()->GetClass())
+				{
+					// The trees have to be compatible.
+					return false;
+				}
+				PreviousStateTree = NextStateContext.StateTree;
+			}
+			if (State->Type == EStateTreeStateType::Linked)
+			{
+				bNewFrame = true;
+			}
 		}
-		else
+		if (bNewFrame || bNewTree)
 		{
-			Frame.FrameID = FActiveFrameID(Storage.GenerateUniqueId());
-			Args.TemporaryFrames.Add_GetRef(Frame).ActiveStates = FStateTreeActiveStates();
-			CurrentActiveFrameIndex = INDEX_NONE;
+			// The path cannot end with a linked or linked asset state.
+			return false;
 		}
 
-		Args.SelectedFrames.Add(Frame.FrameID);
+		return true;
 	}
+}
+
+bool FStateTreeExecutionContext::ForceTransitionInternal(const TArrayView<const UE::StateTree::ExecutionContext::FStateHandleContext> StateContexts, FSelectStateResult& OutSelectionResult)
+{
+	using namespace UE::StateTree;
+	using namespace UE::StateTree::ExecutionContext;
+	using namespace UE::StateTree::ExecutionContext::Private;
+
+	if (!IsValid())
 	{
-		TArray<FActiveState, TInlineAllocator<ExpectedActiveStatePathLength>> NewActiveStates;
-		GetActiveStatePath(TransitionResult.NextActiveFrames, NewActiveStates);
-		Args.SelectedStates.Append(NewActiveStates);
+		STATETREE_LOG(Warning, TEXT("%hs: StateTree context is not initialized properly ('%s' using StateTree '%s')"),
+			__FUNCTION__, *GetNameSafe(&Owner), *GetFullNameSafe(&RootStateTree));
+		return false;
 	}
 
-	int32 SelectedStateIndex = 0;
-	for (FStateTreeFrameStateSelectionEvents& SelectionEvent : TransitionResult.NextActiveFrameEvents)
+	if (StateContexts.IsEmpty())
 	{
-		for (FStateTreeSharedEvent& Event : SelectionEvent.Events)
-		{
-			// Event should be valid but they might come from replication and we don't thrust anything that comes from the network.
-			if (Event.IsValid())
-			{
-				if (Args.SelectedStates.IsValidIndex(SelectedStateIndex))
-				{
-					Args.SelectionEvents.Add(FSelectionEventWithID{.State = Args.SelectedStates[SelectedStateIndex], .Event = Event});
-				}
-			}
-		}
-		++SelectedStateIndex;
+		UE_STATETREE_DEBUG_LOG_EVENT(this, Verbose, TEXT("There are no states in the transition."));
+		return false;
 	}
 
-	const FActiveState* FoundActiveState = Args.SelectedStates.FindByPredicate([&TransitionResult](const FActiveState& ActiveState)
-		{
-			return (ActiveState.GetFrameID() == TransitionResult.SourceFrameID
-				&& ActiveState.GetStateHandle() == TransitionResult.TargetState);
-		});
-	if (FoundActiveState != nullptr)
+	if (!TestStateContextPath(this, StateContexts))
 	{
-		Args.TargetState = *FoundActiveState;
+		UE_STATETREE_DEBUG_LOG_EVENT(this, Verbose, TEXT("The StateContexts is invalid."));
+		return false;
 	}
 
-	ExitState(Args, TransitionResult);
-	return EnterState(Args, TransitionResult);
+	const FStateTreeExecutionState& Exec = GetExecState();
+
+	// A reentrant call to ForceTransition or a call from Start, Tick or Stop must be deferred.
+	if (Exec.CurrentPhase != EStateTreeUpdatePhase::Unset)
+	{
+		UE_STATETREE_DEBUG_LOG_EVENT(this, Warning, TEXT("Can't force a transition while %s"), *UEnum::GetDisplayValueAsText(Exec.CurrentPhase).ToString());
+		return false;
+	}
+
+	if (!CollectActiveExternalData())
+	{
+		STATETREE_LOG(Warning, TEXT("%hs: Failed to collect external data ('%s' using StateTree '%s')"),
+			__FUNCTION__, *GetNameSafe(&Owner), *GetFullNameSafe(&RootStateTree));
+		return false;
+	}
+
+	TArray<FActiveState, TInlineAllocator<ExpectedActiveStatePathLength>> CurrentActiveStatePath;
+	GetActiveStatePath(Exec.ActiveFrames, CurrentActiveStatePath);
+	if (CurrentActiveStatePath.Num() == 0)
+	{
+		UE_STATETREE_DEBUG_LOG_EVENT(this, Verbose, TEXT("The tree needs active states to transition from."));
+		return false;
+	}
+
+	FSelectStateArguments SelectStateArgs;
+	SelectStateArgs.ActiveStates = MakeConstArrayView(CurrentActiveStatePath);
+	SelectStateArgs.SourceState = SelectStateArgs.ActiveStates[0];
+	SelectStateArgs.TargetState = StateContexts.Last();
+	SelectStateArgs.Behavior = ESelectStateBehavior::Forced;
+
+	FSelectStateInternalArguments InternalArgs;
+	InternalArgs.MissingActiveStates = MakeConstArrayView(CurrentActiveStatePath);
+	InternalArgs.MissingSourceFrameID = Exec.ActiveFrames[0].FrameID;
+	InternalArgs.MissingSourceStates = MakeConstArrayView(CurrentActiveStatePath).Left(1);
+	InternalArgs.MissingStatesToReachTarget = MakeConstArrayView(StateContexts);
+
+	FSelectStateResult SelectStateResult;
+	if (!SelectStateInternal(SelectStateArgs, InternalArgs, OutSelectionResult))
+	{
+		UE_STATETREE_DEBUG_LOG_EVENT(this, Verbose, TEXT("The force selection to target failed."));
+		return false;
+	}
+	return true;
 }
 
 const FStateTreeExecutionFrame* FStateTreeExecutionContext::FindFrame(const UStateTree* StateTree, FStateTreeStateHandle RootState, TConstArrayView<FStateTreeExecutionFrame> Frames, const FStateTreeExecutionFrame*& OutParentFrame)
@@ -2938,9 +3054,11 @@ bool FStateTreeExecutionContext::SetGlobalParameters(const FInstancedPropertyBag
 }
 
 // Deprecated
+PRAGMA_DISABLE_DEPRECATION_WARNINGS
 void FStateTreeExecutionContext::CaptureNewStateEvents(TConstArrayView<FStateTreeExecutionFrame> PrevFrames, TConstArrayView<FStateTreeExecutionFrame> NewFrames, TArrayView<FStateTreeFrameStateSelectionEvents> FramesStateSelectionEvents)
 {
 }
+PRAGMA_ENABLE_DEPRECATION_WARNINGS
 
 void FStateTreeExecutionContext::CaptureNewStateEvents(FSelectStateResult& Args)
 {
@@ -3038,7 +3156,7 @@ EStateTreeRunStatus FStateTreeExecutionContext::EnterState(FSelectStateResult& A
 
 	if (bRecordTransitions)
 	{
-		RecordedTransitions.Add(MakeRecordedTransitionResult(Transition));
+		RecordedTransitions.Add(MakeRecordedTransitionResult(Args, Transition));
 	}
 
 	UpdateInstanceData(Args);
@@ -4746,9 +4864,10 @@ bool FStateTreeExecutionContext::RequestTransitionInternal(
 		SetupNextTransition(SourceFrame, SourceStateID, TargetState, Priority);
 		RequestedTransition->Selection = MoveTemp(StateSelectionResult);
 
-		//@todo remove this when NextActiveFrames & NextActiveFrameEvents are deprecated
-		// Fill NextActiveFrames & NextActiveFrameEvents for backward compatibility with ForceTransition
+		// Fill NextActiveFrames & NextActiveFrameEvents for backward compatibility
+		if (bSetDeprecatedTransitionResultProperties)
 		{
+PRAGMA_DISABLE_DEPRECATION_WARNINGS
 			RequestedTransition->Transition.NextActiveFrames.Reset();
 			RequestedTransition->Transition.NextActiveFrameEvents.Reset();
 			FActiveFrameID PreviousFrameID;
@@ -4788,6 +4907,7 @@ bool FStateTreeExecutionContext::RequestTransitionInternal(
 				}
 			}
 		}
+PRAGMA_ENABLE_DEPRECATION_WARNINGS
 
 		//@todo In TriggerTransitions, the events are processed bottom up. If an event with a higher priority is process after, the event will already be consumed.
 		// Consume events from states, if required. The transition might also want to consume the event. (Not a bug.)
@@ -4804,7 +4924,6 @@ bool FStateTreeExecutionContext::RequestTransitionInternal(
 			}
 		}
 
-PRAGMA_DISABLE_DEPRECATION_WARNINGS
 		UE_SUPPRESS(LogStateTree, Verbose,
 		{
 			check(RequestedTransition->Selection.SelectedStates.Num() > 0);
@@ -4817,7 +4936,6 @@ PRAGMA_DISABLE_DEPRECATION_WARNINGS
 				*GetSafeStateName(*SelectedFrame, LastSelectedState.GetStateHandle())
 			);
 		})
-PRAGMA_ENABLE_DEPRECATION_WARNINGS
 
 		return true;
 	}
@@ -4856,35 +4974,37 @@ void FStateTreeExecutionContext::SetupNextTransition(
 
 	OutTransitionResult.SourceFrameID = SourceFrame.FrameID;
 	OutTransitionResult.SourceStateID = SourceStateID;
-PRAGMA_DISABLE_DEPRECATION_WARNINGS
-	if (SourceStateID.IsValid())
-	{
-		OutTransitionResult.SourceState = SourceFrame.ActiveStates.FindStateHandle(SourceStateID);
-	}
-	OutTransitionResult.SourceStateTree = SourceFrame.StateTree;
-	OutTransitionResult.SourceRootState = SourceFrame.RootState;
-PRAGMA_ENABLE_DEPRECATION_WARNINGS
-
 	OutTransitionResult.TargetState = TargetState.StateHandle;
 	OutTransitionResult.CurrentState = FStateTreeStateHandle::Invalid;
 	OutTransitionResult.CurrentRunStatus = Exec.LastTickStatus;
 	OutTransitionResult.ChangeType = EStateTreeStateChangeType::Changed;
 	OutTransitionResult.Priority = Priority;
 
-	//@todo remove this when NextActiveFrames is deprecated
-	//For backward compatibility with ForceTransition
-	FStateTreeExecutionFrame& NewFrame = OutTransitionResult.NextActiveFrames.AddDefaulted_GetRef();
-	NewFrame.StateTree = SourceFrame.StateTree;
-	NewFrame.RootState = SourceFrame.RootState;
-	NewFrame.ActiveTasksStatus = SourceFrame.ActiveTasksStatus;
-	if (TargetState.StateHandle == FStateTreeStateHandle::Invalid)
-	{
-		NewFrame.ActiveStates = {};
-	}
-	else
+	// Fill NextActiveFrames & NextActiveFrameEvents for backward compatibility
+	if (UE::StateTree::ExecutionContext::Private::bSetDeprecatedTransitionResultProperties)
 	{
-		NewFrame.ActiveStates = FStateTreeActiveStates(TargetState.StateHandle, UE::StateTree::FActiveStateID::Invalid);
+PRAGMA_DISABLE_DEPRECATION_WARNINGS
+		if (SourceStateID.IsValid())
+		{
+			OutTransitionResult.SourceState = SourceFrame.ActiveStates.FindStateHandle(SourceStateID);
+		}
+		OutTransitionResult.SourceStateTree = SourceFrame.StateTree;
+		OutTransitionResult.SourceRootState = SourceFrame.RootState;
+
+		FStateTreeExecutionFrame& NewFrame = OutTransitionResult.NextActiveFrames.AddDefaulted_GetRef();
+		NewFrame.StateTree = SourceFrame.StateTree;
+		NewFrame.RootState = SourceFrame.RootState;
+		NewFrame.ActiveTasksStatus = SourceFrame.ActiveTasksStatus;
+		if (TargetState.StateHandle == FStateTreeStateHandle::Invalid)
+		{
+			NewFrame.ActiveStates = {};
+		}
+		else
+		{
+			NewFrame.ActiveStates = FStateTreeActiveStates(TargetState.StateHandle, UE::StateTree::FActiveStateID::Invalid);
+		}
 	}
+PRAGMA_ENABLE_DEPRECATION_WARNINGS
 }
 
 bool FStateTreeExecutionContext::TriggerTransitions()
@@ -5507,10 +5627,13 @@ PRAGMA_ENABLE_DEPRECATION_WARNINGS
 	return RequestedTransition.IsValid();
 }
 
+// Deprecated
 TOptional<FStateTreeTransitionResult> FStateTreeExecutionContext::MakeTransitionResult(const FRecordedStateTreeTransitionResult& RecordedTransition) const
 {
 	FStateTreeTransitionResult Result;
 
+#if WITH_EDITOR
+PRAGMA_DISABLE_DEPRECATION_WARNINGS
 	for (int32 RecordedFrameIndex = 0; RecordedFrameIndex < RecordedTransition.NextActiveFrames.Num(); RecordedFrameIndex++)
 	{
 		const FRecordedStateTreeExecutionFrame& RecordedExecutionFrame = RecordedTransition.NextActiveFrames[RecordedFrameIndex];
@@ -5584,22 +5707,22 @@ TOptional<FStateTreeTransitionResult> FStateTreeExecutionContext::MakeTransition
 	}
 	Result.TargetState = RecordedTransition.TargetState;
 	Result.Priority = RecordedTransition.Priority;
-
-PRAGMA_DISABLE_DEPRECATION_WARNINGS
 	Result.SourceState = RecordedTransition.SourceState;
 	Result.SourceStateTree = RecordedTransition.SourceStateTree;
 	Result.SourceRootState = RecordedTransition.SourceRootState;
 PRAGMA_ENABLE_DEPRECATION_WARNINGS
+#endif //WITH_EDITOR
 
 	return Result;
 }
 
+// Deprecated
 FRecordedStateTreeTransitionResult FStateTreeExecutionContext::MakeRecordedTransitionResult(const FStateTreeTransitionResult& Transition) const
 {
-	check(Transition.NextActiveFrames.Num() == Transition.NextActiveFrameEvents.Num());
-
 	FRecordedStateTreeTransitionResult Result;
 
+#if WITH_EDITOR
+PRAGMA_DISABLE_DEPRECATION_WARNINGS
 	for (int32 FrameIndex = 0; FrameIndex < Transition.NextActiveFrames.Num(); FrameIndex++)
 	{
 		const FStateTreeExecutionFrame& ExecutionFrame = Transition.NextActiveFrames[FrameIndex];
@@ -5635,18 +5758,58 @@ FRecordedStateTreeTransitionResult FStateTreeExecutionContext::MakeRecordedTrans
 	}
 	else
 	{
-PRAGMA_DISABLE_DEPRECATION_WARNINGS
 		Result.SourceStateTree = Transition.SourceStateTree;
 		Result.SourceRootState = Transition.SourceRootState;
 		Result.SourceState = Transition.SourceState;
-PRAGMA_ENABLE_DEPRECATION_WARNINGS
 	}
 	Result.TargetState = Transition.TargetState;
 	Result.Priority = Transition.Priority;
+PRAGMA_ENABLE_DEPRECATION_WARNINGS
+#endif //WITH_EDITOR
 
 	return Result;
 }
 
+FRecordedStateTreeTransitionResult FStateTreeExecutionContext::MakeRecordedTransitionResult(const FSelectStateResult& SelectStateResult, const FStateTreeTransitionResult& TransitionResult) const
+{
+	using namespace UE::StateTree;
+	using namespace UE::StateTree::ExecutionContext;
+	using namespace UE::StateTree::ExecutionContext::Private;
+
+	const FStateTreeExecutionState& Exec = GetExecState();
+
+	FRecordedStateTreeTransitionResult Recorded;
+	Recorded.States.Reserve(SelectStateResult.SelectedStates.Num());
+	for (const UE::StateTree::FActiveState& SelectedState : SelectStateResult.SelectedStates)
+	{
+		const FStateTreeExecutionFrame* ProcessFrame = FindExecutionFrame(SelectedState.GetFrameID(), MakeConstArrayView(Exec.ActiveFrames), MakeConstArrayView(SelectStateResult.TemporaryFrames));
+		if (ProcessFrame == nullptr)
+		{
+			return FRecordedStateTreeTransitionResult();
+		}
+		FRecordedActiveState& RecordedState = Recorded.States.AddDefaulted_GetRef();
+		RecordedState.StateTree = ProcessFrame->StateTree;
+		RecordedState.State = SelectedState.GetStateHandle();
+	}
+
+	Recorded.Events.Reserve(SelectStateResult.SelectionEvents.Num());
+	for (const FSelectionEventWithID& SelectedEvent : SelectStateResult.SelectionEvents)
+	{
+		if (SelectedEvent.Event.IsValid())
+		{
+			const int32 RecordedEventIndex = Recorded.Events.Add(*SelectedEvent.Event.Get());
+			if (const int32 SeletedStateIndex = FActiveStatePath::IndexOf(SelectStateResult.SelectedStates, SelectedEvent.State); SeletedStateIndex != INDEX_NONE)
+			{
+				Recorded.States[SeletedStateIndex].EventIndex = RecordedEventIndex;
+			}
+		}
+	}
+
+	Recorded.Priority = TransitionResult.Priority;
+
+	return Recorded;
+}
+
 // Deprecated
 bool FStateTreeExecutionContext::SelectState(const FStateTreeExecutionFrame& CurrentFrame,
 	const FStateTreeStateHandle NextState,
@@ -6694,7 +6857,9 @@ bool FStateTreeExecutionContext::SelectStateInternal_LinkedAsset(
 
 	FStateTreeExecutionState& Exec = GetExecState();
 
-	if (NextLinkedStateAsset == nullptr || NextLinkedStateAsset->States.Num() == 0)
+	if (NextLinkedStateAsset == nullptr
+		|| NextLinkedStateAsset->States.Num() == 0
+		|| !NextLinkedStateAsset->IsReadyToRun())
 	{
 		STATETREE_LOG(Error, TEXT("%hs: The linked State Tree is invalid. '%s' using StateTree '%s'."),
 			__FUNCTION__,
diff --git a/Engine/Plugins/Runtime/StateTree/Source/StateTreeModule/Private/StateTreeExecutionTypes.cpp b/Engine/Plugins/Runtime/StateTree/Source/StateTreeModule/Private/StateTreeExecutionTypes.cpp
--- a/Engine/Plugins/Runtime/StateTree/Source/StateTreeModule/Private/StateTreeExecutionTypes.cpp
+++ b/Engine/Plugins/Runtime/StateTree/Source/StateTreeModule/Private/StateTreeExecutionTypes.cpp
@@ -176,17 +176,21 @@ void FStateTreeExecutionState::CacheScheduledTickRequest()
 // FStateTreeExecutionFrame
 //----------------------------------------------------------------------//
 //Deprecated
+PRAGMA_DISABLE_DEPRECATION_WARNINGS
 FStateTreeExecutionFrame::FStateTreeExecutionFrame(const FRecordedStateTreeExecutionFrame& RecordedExecutionFrame)
 {
 }
+PRAGMA_ENABLE_DEPRECATION_WARNINGS
 
 //----------------------------------------------------------------------//
 // FRecordedStateTreeExecutionFrame
 //----------------------------------------------------------------------//
 //Deprecated
+PRAGMA_DISABLE_DEPRECATION_WARNINGS
 FRecordedStateTreeExecutionFrame::FRecordedStateTreeExecutionFrame(const FStateTreeExecutionFrame& ExecutionFrame)
 {
 }
+PRAGMA_ENABLE_DEPRECATION_WARNINGS
 
 //Deprecated
 //----------------------------------------------------------------------//
diff --git a/Engine/Plugins/Runtime/StateTree/Source/StateTreeModule/Public/StateTree.h b/Engine/Plugins/Runtime/StateTree/Source/StateTreeModule/Public/StateTree.h
--- a/Engine/Plugins/Runtime/StateTree/Source/StateTreeModule/Public/StateTree.h
+++ b/Engine/Plugins/Runtime/StateTree/Source/StateTreeModule/Public/StateTree.h
@@ -137,6 +137,9 @@ public:
 
 	/** @return true if the other StateTree has compatible context data. */
 	UE_API bool HasCompatibleContextData(const UStateTree& Other) const;
+
+	/** @return true if the other StateTree has compatible context data. */
+	UE_API bool HasCompatibleContextData(TNotNull<const UStateTree*> Other) const;
 	
 	/** @return List of default parameters of the state tree. Default parameter values can be overridden at runtime by the execution context. */
 	const FInstancedPropertyBag& GetDefaultParameters() const
diff --git a/Engine/Plugins/Runtime/StateTree/Source/StateTreeModule/Public/StateTreeExecutionContext.h b/Engine/Plugins/Runtime/StateTree/Source/StateTreeModule/Public/StateTreeExecutionContext.h
--- a/Engine/Plugins/Runtime/StateTree/Source/StateTreeModule/Public/StateTreeExecutionContext.h
+++ b/Engine/Plugins/Runtime/StateTree/Source/StateTreeModule/Public/StateTreeExecutionContext.h
@@ -95,7 +95,7 @@ public:
 	}
 
 	/** @return The owner of the context */
-	UObject* GetOwner() const
+	TNotNull<UObject*> GetOwner() const
 	{
 		return &Owner;
 	}
@@ -107,7 +107,7 @@ public:
 	}
 
 	/** @return the StateTree asset in use by the instance. It is the root asset. */
-	const UStateTree* GetStateTree() const
+	TNotNull<const UStateTree*> GetStateTree() const
 	{
 		return &RootStateTree;
 	}
@@ -733,13 +733,14 @@ PRAGMA_ENABLE_DEPRECATION_WARNINGS
 	}
 
 	/**
-	* Forces transition to a state from a previously recorded state tree transition result.
-	* Primarily used for replication purposes so that a client state tree stay in sync with its server counterpart.
-	* @param Recorded state transition to run on the state tree.
-	* @return The new run status for the state tree.
-	*/
+	 * Forces transition to a state. It will skip all conditions.
+	 * Primarily used for replication purposes so that a client state tree stays in sync with its server counterpart.
+	 * It has to be a running instance. It will not work if you didn't call Start or if the execution previously failed.
+	 * @param Recorded state transition to run on the state tree.
+	 * @return The new run status for the state tree.
+	 */
 	UE_API EStateTreeRunStatus ForceTransition(const FRecordedStateTreeTransitionResult& Transition);
-
+	
 	/** Returns the recorded transitions for this context. */
 	TConstArrayView<FRecordedStateTreeTransitionResult> GetRecordedTransitions() const
 	{
@@ -787,6 +788,7 @@ protected:
 		UE_API FStateTreeExecutionFrame& MakeAndAddTemporaryFrameWithNewRoot(UE::StateTree::FActiveFrameID FrameID, const UE::StateTree::FExecutionFrameKey& FrameKey, FStateTreeExecutionFrame& OtherFrame);
 	};
 
+PRAGMA_DISABLE_DEPRECATION_WARNINGS
 	struct
 	UE_DEPRECATED(5.7, "The selection result changed to include the state ID and frame ID")
 	FStateSelectionResult
@@ -847,6 +849,7 @@ protected:
 		TArray<FStateTreeExecutionFrame, TFixedAllocator<MaxExecutionFrames>> SelectedFrames;
 		TArray<FStateTreeFrameStateSelectionEvents, TFixedAllocator<MaxExecutionFrames>> FramesStateSelectionEvents;
 	};
+PRAGMA_ENABLE_DEPRECATION_WARNINGS
 
 	/** The result of the RequestTransition function. */
 	struct FRequestTransitionResult
@@ -1114,12 +1117,17 @@ protected:
 	/** Tick the states transitions. */
 	UE_API bool TriggerTransitions();
 
+	UE_DEPRECATED(5.7, "Use ForceTransition")
 	/** Create a new transition result from a recorded transition result. It will fail if the recorded transition is malformed. */
 	UE_API TOptional<FStateTreeTransitionResult> MakeTransitionResult(const FRecordedStateTreeTransitionResult& Transition) const;
 
+	UE_DEPRECATED(5.7, "Use MakeRecordedTransitionResult with the FSelectStateResult argument.")
 	/** Create a new recorded transition from a transition result. */
 	UE_API FRecordedStateTreeTransitionResult MakeRecordedTransitionResult(const FStateTreeTransitionResult& Transition) const;
 
+	/** Create a new recorded transition from a select state result. */
+	UE_API FRecordedStateTreeTransitionResult MakeRecordedTransitionResult(const FSelectStateResult& Args, const FStateTreeTransitionResult& Transition) const;
+
 PRAGMA_DISABLE_DEPRECATION_WARNINGS
 	UE_DEPRECATED(5.6, "Replaced with FStateTreeTasksCompletionStatus")
 	/** Confirms that the frame and state ID are valid and the task index is correct. */
@@ -1147,6 +1155,9 @@ PRAGMA_ENABLE_DEPRECATION_WARNINGS
 	/** In the list of completed states, returns the run status when it's completed by a global task. */
 	UE_API EStateTreeRunStatus GetGlobalTasksCompletedStatesStatus() const;
 
+	/** Forces transition to a state. It will skip all conditions. */
+	UE_API bool ForceTransitionInternal(const TArrayView<const UE::StateTree::ExecutionContext::FStateHandleContext> States, FSelectStateResult& OutSelectionResult);
+
 	UE_DEPRECATED(5.7, "Use the SelectState with the FSelectStateArguments argument.")
 	UE_API bool SelectState(const FStateTreeExecutionFrame& CurrentFrame, const FStateTreeStateHandle NextState, FSelectStateResult& OutSelectionResult, const FStateTreeSharedEvent* TransitionEvent = nullptr, const EStateTreeSelectionFallback Fallback = EStateTreeSelectionFallback::None);
 
@@ -1372,8 +1383,10 @@ PRAGMA_ENABLE_DEPRECATION_WARNINGS
 	 */
 	UE_API bool SetGlobalParameters(const FInstancedPropertyBag& Parameters);
 
+PRAGMA_DISABLE_DEPRECATION_WARNINGS
 	UE_DEPRECATED(5.7, "Use the CaptureNewStateEvents with the FSelectStateResult argument.")
 	void CaptureNewStateEvents(TConstArrayView<FStateTreeExecutionFrame> PrevFrames, TConstArrayView<FStateTreeExecutionFrame> NewFrames, TArrayView<FStateTreeFrameStateSelectionEvents> FramesStateSelectionEvents);
+PRAGMA_ENABLE_DEPRECATION_WARNINGS
 
 	/**
 	 * Captures StateTree events used during the state selection.
diff --git a/Engine/Plugins/Runtime/StateTree/Source/StateTreeModule/Public/StateTreeExecutionTypes.h b/Engine/Plugins/Runtime/StateTree/Source/StateTreeModule/Public/StateTreeExecutionTypes.h
--- a/Engine/Plugins/Runtime/StateTree/Source/StateTreeModule/Public/StateTreeExecutionTypes.h
+++ b/Engine/Plugins/Runtime/StateTree/Source/StateTreeModule/Public/StateTreeExecutionTypes.h
@@ -1189,9 +1189,10 @@ PRAGMA_ENABLE_DEPRECATION_WARNINGS
 #endif //WITH_EDITORONLY_DATA
 };
 
-//@todo deprecate once ForceTransition is fixed
 /** Contains StateTree events used during State Selection for a single execution frame. */
-struct FStateTreeFrameStateSelectionEvents
+struct
+UE_DEPRECATED(5.7, "Use the StateTreeExecutionContext.RequestTransitionResult.Selection.SelectionEvents instead.")
+FStateTreeFrameStateSelectionEvents
 {
 	TStaticArray<FStateTreeSharedEvent, FStateTreeActiveStates::MaxStates> Events;
 };
@@ -1222,15 +1223,6 @@ PRAGMA_ENABLE_DEPRECATION_WARNINGS
 		*this = FStateTreeTransitionResult();
 	}
 
-	//@todo deprecate once ForceTransition is fixed
-	/** States selected as result of the transition. */
-	UPROPERTY(EditDefaultsOnly, Category = "Default", BlueprintReadOnly)
-	TArray<FStateTreeExecutionFrame> NextActiveFrames;
-
-	//@todo deprecate once ForceTransition is fixed
-	/** Events used in state selection. */
-	TArray<FStateTreeFrameStateSelectionEvents> NextActiveFrameEvents;
-
 	/** Frame that was active when the transition was requested. */
 	UE::StateTree::FActiveFrameID SourceFrameID;
 
@@ -1240,11 +1232,6 @@ PRAGMA_ENABLE_DEPRECATION_WARNINGS
 	 */
 	UE::StateTree::FActiveStateID SourceStateID;
 
-	UE_DEPRECATED(5.6, "Use SourceStateID instead. The value is valid for Blueprint backward compatibility.")
-	/** Transition source state. */
-	UPROPERTY(EditDefaultsOnly, Category = "Default", BlueprintReadOnly)
-	FStateTreeStateHandle SourceState = FStateTreeStateHandle::Invalid;
-
 	/** Transition target state. It can be a completion state. */
 	UPROPERTY(EditDefaultsOnly, Category = "Default", BlueprintReadOnly)
 	FStateTreeStateHandle TargetState = FStateTreeStateHandle::Invalid;
@@ -1265,15 +1252,31 @@ PRAGMA_ENABLE_DEPRECATION_WARNINGS
 	UPROPERTY(EditDefaultsOnly, Category = "Default", BlueprintReadOnly)
 	EStateTreeTransitionPriority Priority = EStateTreeTransitionPriority::None;
 
-	UE_DEPRECATED(5.6, "Use SourceFrameID instead. The value is valid for Blueprint backward compatibility.")
+PRAGMA_DISABLE_DEPRECATION_WARNINGS
+	UE_DEPRECATED(5.7, "Use the StateTreeExecutionContext.RequestTransitionResult.Selection.SelectedState instead.")
+	/** States selected as result of the transition. */
+	UPROPERTY(EditDefaultsOnly, Category = "Default", BlueprintReadOnly, meta = (DeprecatedProperty))
+	TArray<FStateTreeExecutionFrame> NextActiveFrames;
+
+	UE_DEPRECATED(5.7, "Use the StateTreeExecutionContext.RequestTransitionResult.Selection.SelectionEvents instead.")
+	/** Events used in state selection. */
+	TArray<FStateTreeFrameStateSelectionEvents> NextActiveFrameEvents;
+
+	UE_DEPRECATED(5.6, "Use SourceFrameID instead")
 	/** StateTree asset that was active when the transition was requested. */
-	UPROPERTY(EditDefaultsOnly, Category = "Default", BlueprintReadOnly)
+	UPROPERTY(EditDefaultsOnly, Category = "Default", BlueprintReadOnly, meta = (DeprecatedProperty))
 	TObjectPtr<const UStateTree> SourceStateTree = nullptr;
 
-	UE_DEPRECATED(5.6, "Use SourceFrameID instead. The value is valid for Blueprint backward compatibility.")
+	UE_DEPRECATED(5.6, "Use SourceFrameID instead.")
 	/** Root state the execution frame where the transition was requested. */
-	UPROPERTY(EditDefaultsOnly, Category = "Default", BlueprintReadOnly)
+	UPROPERTY(EditDefaultsOnly, Category = "Default", BlueprintReadOnly, meta = (DeprecatedProperty))
 	FStateTreeStateHandle SourceRootState = FStateTreeStateHandle::Invalid;
+
+	UE_DEPRECATED(5.6, "Use SourceStateID instead")
+	/** Transition source state. */
+	UPROPERTY(EditDefaultsOnly, Category = "Default", BlueprintReadOnly, meta=(DeprecatedProperty))
+	FStateTreeStateHandle SourceState = FStateTreeStateHandle::Invalid;
+PRAGMA_ENABLE_DEPRECATION_WARNINGS
 };
 
 /*
@@ -1288,11 +1291,13 @@ enum class EStateTreeRecordTransitions : uint8
 };
 
 /*
-* Captured state tree execution frame that can be cached for recording purposes.
-* Held in FRecordedStateTreeTransitionResult for its NextActiveFrames.
-*/
+ * Captured state tree execution frame that can be cached for recording purposes.
+ * Held in FRecordedStateTreeTransitionResult for its NextActiveFrames.
+ */
 USTRUCT()
-struct FRecordedStateTreeExecutionFrame
+struct
+UE_DEPRECATED(5.7, "FRecordedStateTreeExecutionFrame is not used anymore.")
+FRecordedStateTreeExecutionFrame
 {
 	GENERATED_BODY()
 
@@ -1320,9 +1325,30 @@ struct FRecordedStateTreeExecutionFrame
 	TStaticArray<uint8, FStateTreeActiveStates::MaxStates> EventIndices;
 };
 
+/** Captured state cached for recording purposes. */
+USTRUCT()
+struct FRecordedActiveState
+{
+	GENERATED_BODY()
+
+	FRecordedActiveState() = default;
+
+	/** The state tree that owns the state handle. */
+	UPROPERTY()
+	TObjectPtr<const UStateTree> StateTree;
+
+	/** The active state. */
+	UPROPERTY()
+	FStateTreeStateHandle State;
+
+	/** Captured events from the transition that we've recorded */
+	UPROPERTY()
+	int32 EventIndex = INDEX_NONE;
+};
+
 /*
  * Captured state tree transition result that can be cached for recording purposes.
- * Primarily, when transitions are recorded through this structure, we can replicate them down
+ * when transitions are recorded through this structure, we can replicate them down
  * to clients to keep our state tree in sync.
  */
 USTRUCT()
@@ -1330,37 +1356,62 @@ struct FRecordedStateTreeTransitionResult
 {
 	GENERATED_BODY()
 
+PRAGMA_DISABLE_DEPRECATION_WARNINGS
 	FRecordedStateTreeTransitionResult() = default;
+	FRecordedStateTreeTransitionResult(const FRecordedStateTreeTransitionResult&) = default;
+	FRecordedStateTreeTransitionResult(FRecordedStateTreeTransitionResult&&) = default;
+	FRecordedStateTreeTransitionResult& operator=(const FRecordedStateTreeTransitionResult&) = default;
+	FRecordedStateTreeTransitionResult& operator=(FRecordedStateTreeTransitionResult&&) = default;
+PRAGMA_ENABLE_DEPRECATION_WARNINGS
+
 	UE_DEPRECATED(5.6, "Use FStateTreeExecutionContext::MakeRecordedTransitionResult to create a new recorded transition.")
 	UE_API FRecordedStateTreeTransitionResult(const FStateTreeTransitionResult& Transition);
+	
+	/** The selected states. */
+	UPROPERTY()
+	TArray<FRecordedActiveState> States;
 
-   	/** States selected as result of the transition. */
+	/** The selected states. */
+	UPROPERTY()
+	TArray<FStateTreeEvent> Events;
+	
+	/** Priority of the transition that caused the state change. */
+	UPROPERTY()
+	EStateTreeTransitionPriority Priority = EStateTreeTransitionPriority::None;
+
+#if WITH_EDITORONLY_DATA
+PRAGMA_DISABLE_DEPRECATION_WARNINGS
+	UE_DEPRECATED(5.7, "The recorded transition doesn't record the frame. ForceTransition will recreated them if needed.")
+	/** States selected as result of the transition. */
 	UPROPERTY()
 	TArray<FRecordedStateTreeExecutionFrame> NextActiveFrames;
 
+	UE_DEPRECATED(5.7, "The event are recorded in SelectedStates.")
 	/** Captured events from the transition that we've recorded */
 	UPROPERTY()
 	TArray<FStateTreeEvent> NextActiveFrameEvents;
 
+	UE_DEPRECATED(5.7, "We can't assumed that the recorded has the same active states. The source won't be set.")
 	/** Transition source state. */
 	UPROPERTY()
 	FStateTreeStateHandle SourceState = FStateTreeStateHandle::Invalid;
 
+	UE_DEPRECATED(5.7, "The TargetState is always the last item. We can't assumed that the recorded has the same active states.")
 	/** Transition target state. */
 	UPROPERTY()
 	FStateTreeStateHandle TargetState = FStateTreeStateHandle::Invalid;
-	
-	/** Priority of the transition that caused the state change. */
-	UPROPERTY()
-	EStateTreeTransitionPriority Priority = EStateTreeTransitionPriority::None;
 
+	UE_DEPRECATED(5.7, "We can't assumed that the recorded has the same active states. The source won't be set.")
 	/** StateTree asset that was active when the transition was requested. */
 	UPROPERTY()
 	TObjectPtr<const UStateTree> SourceStateTree = nullptr;
 
+	UE_DEPRECATED(5.7, "We can't assumed that the recorded has the same active states. The source won't be set.")
 	/** Root state the execution frame where the transition was requested. */
 	UPROPERTY()
 	FStateTreeStateHandle SourceRootState = FStateTreeStateHandle::Invalid;
+PRAGMA_ENABLE_DEPRECATION_WARNINGS
+#endif //WITH_EDITORONLY_DATA
 };
 
 #undef UE_API
