[StateTree] Remove 5.5 deprecation.

#rnx

diff --git a/Engine/Plugins/Runtime/StateTree/Source/StateTreeEditorModule/Private/StateTreeEditorPropertyBindings.cpp b/Engine/Plugins/Runtime/StateTree/Source/StateTreeEditorModule/Private/StateTreeEditorPropertyBindings.cpp
--- a/Engine/Plugins/Runtime/StateTree/Source/StateTreeEditorModule/Private/StateTreeEditorPropertyBindings.cpp
+++ b/Engine/Plugins/Runtime/StateTree/Source/StateTreeEditorModule/Private/StateTreeEditorPropertyBindings.cpp
@@ -164,16 +164,6 @@ int32 FStateTreeEditorPropertyBindings::GetNumBindings() const
 	return PropertyBindings.Num();
 }
 
-PRAGMA_DISABLE_DEPRECATION_WARNINGS
-void FStateTreeEditorPropertyBindings::GetPropertyBindingsFor(const FGuid StructID, TArray<FStateTreePropertyPathBinding>& OutBindings) const
-{
-	TArray<const FPropertyBindingBinding*> NodeBindings;
-	FPropertyBindingBindingCollection::GetBindingsFor(StructID, NodeBindings);
-	Algo::Transform(NodeBindings, OutBindings, [](const FPropertyBindingBinding* BindingPtr) { return *static_cast<const FStateTreePropertyPathBinding*>(BindingPtr); });
-}
-PRAGMA_ENABLE_DEPRECATION_WARNINGS
-
-
 //////////////////////////////////////////////////////////////////////////
 
 FStateTreeBindingLookup::FStateTreeBindingLookup(const IStateTreeEditorPropertyBindingsOwner* InBindingOwner)
diff --git a/Engine/Plugins/Runtime/StateTree/Source/StateTreeEditorModule/Public/StateTreeEditorNode.h b/Engine/Plugins/Runtime/StateTree/Source/StateTreeEditorModule/Public/StateTreeEditorNode.h
--- a/Engine/Plugins/Runtime/StateTree/Source/StateTreeEditorModule/Public/StateTreeEditorNode.h
+++ b/Engine/Plugins/Runtime/StateTree/Source/StateTreeEditorModule/Public/StateTreeEditorNode.h
@@ -68,23 +68,6 @@ struct FStateTreeEditorNode
 
 	UPROPERTY(EditDefaultsOnly, Category = Node)
 	EStateTreeExpressionOperand ExpressionOperand = EStateTreeExpressionOperand::And;
-
-#if WITH_EDITOR
-	UE_DEPRECATED(5.5, "Use ExpressionIndent instead.")
-	uint8 ConditionIndent = 0;
-
-PRAGMA_DISABLE_DEPRECATION_WARNINGS
-	UE_DEPRECATED(5.5, "Use ExpressionOperand instead.")
-	EStateTreeConditionOperand ConditionOperand = EStateTreeConditionOperand::And;
-
-	FStateTreeEditorNode() = default;
-	FStateTreeEditorNode(const FStateTreeEditorNode&) = default;
-	FStateTreeEditorNode(FStateTreeEditorNode&&) = default;
-	FStateTreeEditorNode& operator=(const FStateTreeEditorNode&) = default;
-	FStateTreeEditorNode& operator=(FStateTreeEditorNode&&) = default;
-	~FStateTreeEditorNode() = default;
-PRAGMA_ENABLE_DEPRECATION_WARNINGS
-#endif //WITH_EDITOR
 };
 
 template <typename T>
diff --git a/Engine/Plugins/Runtime/StateTree/Source/StateTreeEditorModule/Public/StateTreeEditorPropertyBindings.h b/Engine/Plugins/Runtime/StateTree/Source/StateTreeEditorModule/Public/StateTreeEditorPropertyBindings.h
--- a/Engine/Plugins/Runtime/StateTree/Source/StateTreeEditorModule/Public/StateTreeEditorPropertyBindings.h
+++ b/Engine/Plugins/Runtime/StateTree/Source/StateTreeEditorModule/Public/StateTreeEditorPropertyBindings.h
@@ -142,9 +142,6 @@ PRAGMA_DISABLE_DEPRECATION_WARNINGS
 	void RemoveUnusedBindings(const TMap<FGuid, const FStateTreeDataView>& ValidStructs)
 	{
 	}
-
-	UE_DEPRECATED(5.5, "Use GetPropertyBindingsFor returning pointers instead.")
-	UE_API void GetPropertyBindingsFor(const FGuid StructID, TArray<FStateTreePropertyPathBinding>& OutBindings) const;
 PRAGMA_ENABLE_DEPRECATION_WARNINGS
 
 protected:
diff --git a/Engine/Plugins/Runtime/StateTree/Source/StateTreeEditorModule/Public/StateTreePropertyBindingCompiler.h b/Engine/Plugins/Runtime/StateTree/Source/StateTreeEditorModule/Public/StateTreePropertyBindingCompiler.h
--- a/Engine/Plugins/Runtime/StateTree/Source/StateTreeEditorModule/Public/StateTreePropertyBindingCompiler.h
+++ b/Engine/Plugins/Runtime/StateTree/Source/StateTreeEditorModule/Public/StateTreePropertyBindingCompiler.h
@@ -89,16 +89,10 @@ struct FStateTreePropertyBindingCompiler
 	{
 		return SourceStructs.FindByPredicate([ID](const FStateTreeBindableStructDesc& Structs) { return (Structs.ID == ID); });
 	}
-PRAGMA_DISABLE_DEPRECATION_WARNINGS
 
+PRAGMA_DISABLE_DEPRECATION_WARNINGS
 	UE_DEPRECATED(5.6, "Use the version taking FPropertyBindingPath instead")
 	UE_API FStateTreeIndex16 GetDispatcherIDFromPath(const FStateTreePropertyPath& PathToDispatcher) const;
-
-	UE_DEPRECATED(5.5, "Use CompileBatch with PropertyFuncsBegin and PropertyFuncsEnd instead.")
-	[[nodiscard]] bool CompileBatch(const FStateTreeBindableStructDesc& InTargetStruct, TConstArrayView<FStateTreePropertyPathBinding> InPropertyBindings, int32& OutBatchIndex) 
-	{
-		return CompileBatch(InTargetStruct, InPropertyBindings, FStateTreeIndex16::Invalid, FStateTreeIndex16::Invalid, OutBatchIndex);
-	}
 PRAGMA_ENABLE_DEPRECATION_WARNINGS
 
 protected:
diff --git a/Engine/Plugins/Runtime/StateTree/Source/StateTreeEditorModule/Public/StateTreeState.h b/Engine/Plugins/Runtime/StateTree/Source/StateTreeEditorModule/Public/StateTreeState.h
--- a/Engine/Plugins/Runtime/StateTree/Source/StateTreeEditorModule/Public/StateTreeState.h
+++ b/Engine/Plugins/Runtime/StateTree/Source/StateTreeEditorModule/Public/StateTreeState.h
@@ -161,7 +161,7 @@ PRAGMA_ENABLE_DEPRECATION_WARNINGS
 	bool bTransitionEnabled = true;
 
 #if WITH_EDITORONLY_DATA
-	UE_DEPRECATED(5.5, "Use RequiredEvent.Tag instead.")
+	UE_DEPRECATED(all, "Use RequiredEvent.Tag instead.")
 	UPROPERTY()
 	FGameplayTag EventTag_DEPRECATED;
 #endif // WITH_EDITORONLY_DATA
diff --git a/Engine/Plugins/Runtime/StateTree/Source/StateTreeModule/Private/StateTreePropertyBindings.cpp b/Engine/Plugins/Runtime/StateTree/Source/StateTreeModule/Private/StateTreePropertyBindings.cpp
--- a/Engine/Plugins/Runtime/StateTree/Source/StateTreeModule/Private/StateTreePropertyBindings.cpp
+++ b/Engine/Plugins/Runtime/StateTree/Source/StateTreeModule/Private/StateTreePropertyBindings.cpp
@@ -348,13 +348,6 @@ const FStateTreePropertyAccess* FStateTreePropertyBindings::GetPropertyAccess(co
 	return &PropertyAccesses[InPropertyReference.GetRefAccessIndex().Get()];
 }
 
-void FStateTreePropertyBindings::DebugPrintInternalLayout(FString& OutString) const
-{
-#if WITH_EDITOR || WITH_PROPERTYBINDINGUTILS_DEBUG
-	OutString += DebugAsString();
-#endif
-}
-
 PRAGMA_DISABLE_DEPRECATION_WARNINGS
 int32 FStateTreePropertyBindings::GetSourceStructNum() const
 {
diff --git a/Engine/Plugins/Runtime/StateTree/Source/StateTreeModule/Public/StateTreeEvents.h b/Engine/Plugins/Runtime/StateTree/Source/StateTreeModule/Public/StateTreeEvents.h
--- a/Engine/Plugins/Runtime/StateTree/Source/StateTreeModule/Public/StateTreeEvents.h
+++ b/Engine/Plugins/Runtime/StateTree/Source/StateTreeModule/Public/StateTreeEvents.h
@@ -215,9 +215,6 @@ struct FStateTreeEventQueue
 		}
 	}
 
-	UE_DEPRECATED(5.5, "Use GetEventsView() instead.")
-	TConstArrayView<FStateTreeEvent> GetEvents() const { return {};	}
-
 protected:
 	// Used by FStateTreeExecutionState to implement deprecated functionality.
 	TArray<FStateTreeSharedEvent>& GetEventsArray() { return SharedEvents; };
diff --git a/Engine/Plugins/Runtime/StateTree/Source/StateTreeModule/Public/StateTreeExecutionContext.h b/Engine/Plugins/Runtime/StateTree/Source/StateTreeModule/Public/StateTreeExecutionContext.h
--- a/Engine/Plugins/Runtime/StateTree/Source/StateTreeModule/Public/StateTreeExecutionContext.h
+++ b/Engine/Plugins/Runtime/StateTree/Source/StateTreeModule/Public/StateTreeExecutionContext.h
@@ -539,9 +539,6 @@ PRAGMA_ENABLE_DEPRECATION_WARNINGS
 	/** Consumes and removes the specified event from the event queue. */
 	UE_API void ConsumeEvent(const FStateTreeSharedEvent& Event);
 
-	UE_DEPRECATED(5.5, "Use GetEventsToProcessView() instead.")
-	TConstArrayView<FStateTreeEvent> GetEventsToProcess() const { return {}; }
-
 	FStateTreeIndex16 GetCurrentlyProcessedNodeIndex() const
 	{
 		return FStateTreeIndex16(CurrentNodeIndex);
@@ -726,12 +723,6 @@ PRAGMA_ENABLE_DEPRECATION_WARNINGS
 	/** Looks for a frame in provided list of frames. */
 	static UE_API const FStateTreeExecutionFrame* FindFrame(const UStateTree* StateTree, FStateTreeStateHandle RootState, TConstArrayView<FStateTreeExecutionFrame> Frames, const FStateTreeExecutionFrame*& OutParentFrame);
 
-	UE_DEPRECATED(5.5, "Use UE::StateTree::InstanceData::GetDataView instead.")
-	static FStateTreeDataView GetDataView(FStateTreeInstanceStorage& InstanceDataStorage, FStateTreeInstanceStorage* CurrentlyProcessedSharedInstanceStorage, const FStateTreeExecutionFrame* ParentFrame, const FStateTreeExecutionFrame& CurrentFrame, TConstArrayView<FStateTreeDataView> ContextAndExternalDataViews, const FStateTreeDataHandle Handle)
-	{
-		return UE::StateTree::InstanceData::GetDataView(InstanceDataStorage, CurrentlyProcessedSharedInstanceStorage, ParentFrame, CurrentFrame, Handle);
-	}
-
 	/**
 	 * Forces transition to a state. It will skip all conditions.
 	 * Primarily used for replication purposes so that a client state tree stays in sync with its server counterpart.
diff --git a/Engine/Plugins/Runtime/StateTree/Source/StateTreeModule/Public/StateTreeExecutionTypes.h b/Engine/Plugins/Runtime/StateTree/Source/StateTreeModule/Public/StateTreeExecutionTypes.h
--- a/Engine/Plugins/Runtime/StateTree/Source/StateTreeModule/Public/StateTreeExecutionTypes.h
+++ b/Engine/Plugins/Runtime/StateTree/Source/StateTreeModule/Public/StateTreeExecutionTypes.h
@@ -1039,9 +1039,6 @@ public:
 	/** @return the unique path of all the active states of all the active execution frames. */
 	UE_API UE::StateTree::FActiveStatePath GetActiveStatePath() const;
 
-	UE_DEPRECATED(5.5, "Use FindAndRemoveExpiredDelayedTransitions() instead or search through DelayedTransitions directly.")
-	FStateTreeTransitionDelayedState* FindDelayedTransition(const UStateTree* OwnerStateTree, const FStateTreeIndex16 TransitionIndex) { return nullptr; }
-
 	UE_DEPRECATED(5.6, "FindAndRemoveExpiredDelayedTransitions is not used anymore.")
 	/** Finds all delayed transition states for a specific transition and removes them. Returns their copies. */
 	TArray<FStateTreeTransitionDelayedState, TInlineAllocator<8>> FindAndRemoveExpiredDelayedTransitions(const UStateTree* OwnerStateTree, const FStateTreeIndex16 TransitionIndex)
diff --git a/Engine/Plugins/Runtime/StateTree/Source/StateTreeModule/Public/StateTreeNodeRef.h b/Engine/Plugins/Runtime/StateTree/Source/StateTreeModule/Public/StateTreeNodeRef.h
--- a/Engine/Plugins/Runtime/StateTree/Source/StateTreeModule/Public/StateTreeNodeRef.h
+++ b/Engine/Plugins/Runtime/StateTree/Source/StateTreeModule/Public/StateTreeNodeRef.h
@@ -46,7 +46,7 @@ private:
 
 
 /**
- * A reference to a task that can be retreive. Similar to StrongPtr and WeakPtr.
+ * A reference to a task that can be retrieve. Similar to StrongPtr and WeakPtr.
  */
 struct UE_DEPRECATED(5.6, "FStateTreeWeakTaskRef is deprecated. We now use TaskIndex in WeakExecutionContext.") FStateTreeWeakTaskRef
 {
diff --git a/Engine/Plugins/Runtime/StateTree/Source/StateTreeModule/Public/StateTreePropertyBindings.h b/Engine/Plugins/Runtime/StateTree/Source/StateTreeModule/Public/StateTreePropertyBindings.h
--- a/Engine/Plugins/Runtime/StateTree/Source/StateTreeModule/Public/StateTreePropertyBindings.h
+++ b/Engine/Plugins/Runtime/StateTree/Source/StateTreeModule/Public/StateTreePropertyBindings.h
@@ -143,14 +143,6 @@ struct FStateTreeBindableStructDesc : public FPropertyBindingBindableStructDescr
 	{
 	}
 
-	UE_DEPRECATED(5.5, "Use constructor with StatePath instead.")
-	FStateTreeBindableStructDesc(const FName InName, const UStruct* InStruct, const FStateTreeDataHandle InDataHandle, const EStateTreeBindableStructSource InDataSource, const FGuid InGuid)
-		: FPropertyBindingBindableStructDescriptor(InName, InStruct, InGuid)
-		, DataHandle(InDataHandle)
-		, DataSource(InDataSource)
-	{
-	}
-
 	virtual FString GetSection() const override
 	{
 		return StatePath;
@@ -538,9 +530,6 @@ PRAGMA_DISABLE_DEPRECATION_WARNINGS
 	UE_DEPRECATED(5.6, "Use the overload taking FPropertyBindingPathIndirection instead")
 	[[nodiscard]] static UE_API bool ResolveCopyType(const FStateTreePropertyPathIndirection& SourceIndirection, const FStateTreePropertyPathIndirection& TargetIndirection, FStateTreePropertyCopy& OutCopy);
 
-	UE_DEPRECATED(5.5, "DebugPrintInternalLayout is deprecated. Use DebugAsString instead.")
-	UE_API void DebugPrintInternalLayout(FString& OutString) const;
-
 private:
 	UE_DEPRECATED(5.6, "Use the version taking FPropertyBindingPropertyIndirection instead.")
 	[[nodiscard]] UE_API bool ResolvePath(const UStruct* Struct, const FPropertyBindingPath& Path, FStateTreePropertyIndirection& OutFirstIndirection, FPropertyBindingPathIndirection& OutLeafIndirection);
diff --git a/Engine/Plugins/Runtime/StateTree/Source/StateTreeModule/Public/StateTreeTaskBase.h b/Engine/Plugins/Runtime/StateTree/Source/StateTreeModule/Public/StateTreeTaskBase.h
--- a/Engine/Plugins/Runtime/StateTree/Source/StateTreeModule/Public/StateTreeTaskBase.h
+++ b/Engine/Plugins/Runtime/StateTree/Source/StateTreeModule/Public/StateTreeTaskBase.h
@@ -35,25 +35,6 @@ struct FStateTreeTaskBase : public FStateTreeNodeBase
 #endif
 	{
 	}
-	
-	UE_DEPRECATED(5.5, "Use EnterState without the EStateTreeStateChangeType.")
-	/**
-	 * Note: The API has been deprecated. ChangeType is moved into FStateTreeTransitionResult.
-	 * You can configure the task to be only called on state changes (that is, never call sustained changes) by setting bShouldStateChangeOnReselect to true.
-	 */
-	virtual EStateTreeRunStatus EnterState(FStateTreeExecutionContext& Context, const EStateTreeStateChangeType ChangeType, const FStateTreeTransitionResult& Transition) const final
-	{
-		return EStateTreeRunStatus::Running;
-	}
-
-	UE_DEPRECATED(5.5, "Use ExitState without the EStateTreeStateChangeType.")
-	/**
-	 * Note: The API has been deprecated. ChangeType is moved into FStateTreeTransitionResult.
-	 * You can configure the task to be only called on state changes (that is, never call sustained changes) by setting bShouldStateChangeOnReselect to true.
-	 */
-	virtual void ExitState(FStateTreeExecutionContext& Context, const EStateTreeStateChangeType ChangeType, const FStateTreeTransitionResult& Transition) const final
-	{
-	}
 
 	/**
 	 * Called when a new state is entered and task is part of active states.
diff --git a/Engine/Plugins/Runtime/StateTree/Source/StateTreeModule/Public/StateTreeTypes.h b/Engine/Plugins/Runtime/StateTree/Source/StateTreeModule/Public/StateTreeTypes.h
--- a/Engine/Plugins/Runtime/StateTree/Source/StateTreeModule/Public/StateTreeTypes.h
+++ b/Engine/Plugins/Runtime/StateTree/Source/StateTreeModule/Public/StateTreeTypes.h
@@ -25,11 +25,6 @@ struct FStateTreeEvent;
 
 namespace UE::StateTree
 {
-#if WITH_EDITOR
-	UE_DEPRECATED(5.5, "Use MaxExpressionIndent instead.")
-	inline constexpr int32 MaxConditionIndent = 4;
-#endif //WITH_EDITOR
-
 	inline constexpr int32 MaxExpressionIndent = 4;
 
 	inline const FName SchemaTag(TEXT("Schema"));
@@ -126,10 +121,6 @@ enum class EStateTreeExpressionOperand : uint8
 	Or,
 };
 
-#if WITH_EDITOR
-	using EStateTreeConditionOperand UE_DEPRECATED(5.5, "Use EStateTreeExpressionOperand instead.") = EStateTreeExpressionOperand;
-#endif //WITH_EDITOR
-
 UENUM()
 enum class EStateTreeStateType : uint8
 {
@@ -174,8 +165,8 @@ enum class EStateTreeStateSelectionBehavior : uint8
 	TryFollowTransitions UMETA(DisplayName = "Try Follow Transitions"),
 
 	// Olds names that needs to be kept forever to ensure asset serialization to work correctly when UENUM() switched from serializing int to names.
-	TrySelectChildrenAtUniformRandom UE_DEPRECATED(5.5, "Use TrySelectChildrenAtRandom Instead") = TrySelectChildrenAtRandom UMETA(Hidden),
-	TrySelectChildrenBasedOnRelativeUtility UE_DEPRECATED(5.5, "Use TrySelectChildrenAtRandomWeightedByUtility Instead") = TrySelectChildrenAtRandomWeightedByUtility UMETA(Hidden)
+	TrySelectChildrenAtUniformRandom UE_DEPRECATED(all, "Use TrySelectChildrenAtRandom Instead") = TrySelectChildrenAtRandom UMETA(Hidden),
+	TrySelectChildrenBasedOnRelativeUtility UE_DEPRECATED(all, "Use TrySelectChildrenAtRandomWeightedByUtility Instead") = TrySelectChildrenAtRandomWeightedByUtility UMETA(Hidden)
 };
 
 
@@ -524,13 +515,6 @@ struct FStateTreeRandomTimeDuration
 		return static_cast<decltype(Scale)>(RandomStream.RandRange(MinVal, MaxVal)) / Scale;
 	}
 
-	UE_DEPRECATED(5.5, "Use the version with random stream instead.")
-	float GetRandomDuration() const
-	{
-		const FRandomStream RandomStream((int32)FPlatformTime::Cycles());
-		return GetRandomDuration(RandomStream);
-	}
-
 protected:
 
 	static constexpr float Scale = 100.0f;
