[StateTree] Allowing state selection behaviors that are not supported by the schema if they're predefined for certain state types.
Adding an option to hide state types per schema.

diff --git a/Engine/Plugins/Runtime/StateTree/Source/StateTreeEditorModule/Private/Customizations/StateTreeStateDetails.cpp b/Engine/Plugins/Runtime/StateTree/Source/StateTreeEditorModule/Private/Customizations/StateTreeStateDetails.cpp
--- a/Engine/Plugins/Runtime/StateTree/Source/StateTreeEditorModule/Private/Customizations/StateTreeStateDetails.cpp
+++ b/Engine/Plugins/Runtime/StateTree/Source/StateTreeEditorModule/Private/Customizations/StateTreeStateDetails.cpp
@@ -170,6 +170,20 @@ void FStateTreeStateDetails::CustomizeDetails(IDetailLayoutBuilder& DetailBuilde
 
 	// Type
 	TypeProperty->MarkHiddenByCustomization();
+	{
+		// Restrict IsStateTypeAllowed
+		TSharedRef<FPropertyRestriction> SchemaSupported = MakeShared<FPropertyRestriction>(LOCTEXT("StateSelectionAllowedRestriction", "The schema restricts the state type."));
+		const UEnum* EnumPtr = StaticEnum<EStateTreeStateType>();
+		for (int32 Index = 0; Index < EnumPtr->NumEnums(); ++Index)
+		{
+			const EStateTreeStateType EnumValue = static_cast<EStateTreeStateType>(EnumPtr->GetValueByIndex(Index));
+			if (!Schema->IsStateTypeAllowed(EnumValue))
+			{
+				SchemaSupported->AddHiddenValue(EnumPtr->GetNameStringByIndex(Index));
+			}
+		}
+		TypeProperty->AddRestriction(SchemaSupported);
+	}
 	StateCategory.AddProperty(TypeProperty);
 
 	// Per state type properties
diff --git a/Engine/Plugins/Runtime/StateTree/Source/StateTreeEditorModule/Private/StateTreeCompiler.cpp b/Engine/Plugins/Runtime/StateTree/Source/StateTreeEditorModule/Private/StateTreeCompiler.cpp
--- a/Engine/Plugins/Runtime/StateTree/Source/StateTreeEditorModule/Private/StateTreeCompiler.cpp
+++ b/Engine/Plugins/Runtime/StateTree/Source/StateTreeEditorModule/Private/StateTreeCompiler.cpp
@@ -630,7 +630,20 @@ bool FStateTreeCompiler::CreateStateRecursive(UStateTreeState& State, const FSta
 	CompactState.Type = State.Type;
 	CompactState.SelectionBehavior = State.SelectionBehavior;
 
-	if (!Schema->IsStateSelectionAllowed(CompactState.SelectionBehavior))
+	if (!Schema->IsStateTypeAllowed(CompactState.Type))
+	{
+		Log.Reportf(EMessageSeverity::Warning,
+			TEXT("The State '%s' has a restricted type for the schema."),
+			*CompactState.Name.ToString());
+		return false;
+	}
+
+	const bool bHasPredefinedSelectionBehavior = (CompactState.Type == EStateTreeStateType::Group || CompactState.Type == EStateTreeStateType::Linked || CompactState.Type == EStateTreeStateType::LinkedAsset);
+	if (bHasPredefinedSelectionBehavior)
+	{
+		CompactState.SelectionBehavior = EStateTreeStateSelectionBehavior::TryEnterState;
+	}
+	else if (!Schema->IsStateSelectionAllowed(CompactState.SelectionBehavior))
 	{
 		Log.Reportf(EMessageSeverity::Warning,
 			TEXT("The State '%s' has a restricted selection behavior for the schema."),
diff --git a/Engine/Plugins/Runtime/StateTree/Source/StateTreeEditorModule/Private/StateTreeState.cpp b/Engine/Plugins/Runtime/StateTree/Source/StateTreeEditorModule/Private/StateTreeState.cpp
--- a/Engine/Plugins/Runtime/StateTree/Source/StateTreeEditorModule/Private/StateTreeState.cpp
+++ b/Engine/Plugins/Runtime/StateTree/Source/StateTreeEditorModule/Private/StateTreeState.cpp
@@ -154,15 +154,12 @@ void UStateTreeState::PostEditChangeChainProperty(FPropertyChangedChainEvent& Pr
 	
 	if (ChangePropertyPath.IsPathExact(StateTypePath))
 	{
-		// Reset Selection Behavior back to Try Enter State for group and linked types
-		if (Type == EStateTreeStateType::Group || Type == EStateTreeStateType::Linked || Type == EStateTreeStateType::LinkedAsset)
+		const bool bHasPredefinedSelectionBehavior = Type == EStateTreeStateType::Group || Type == EStateTreeStateType::Linked || Type == EStateTreeStateType::LinkedAsset;
+		if (bHasPredefinedSelectionBehavior)
 		{
+			// Reset Selection Behavior back to Try Enter State for group and linked types
 			SelectionBehavior = EStateTreeStateSelectionBehavior::TryEnterState;
-		}
-
-		// Remove any tasks when they are not used.
-		if (Type == EStateTreeStateType::Group || Type == EStateTreeStateType::Linked || Type == EStateTreeStateType::LinkedAsset)
-		{
+			// Remove any tasks when they are not used.
 			Tasks.Reset();
 		}
 
diff --git a/Engine/Plugins/Runtime/StateTree/Source/StateTreeModule/Public/StateTreeSchema.h b/Engine/Plugins/Runtime/StateTree/Source/StateTreeModule/Public/StateTreeSchema.h
--- a/Engine/Plugins/Runtime/StateTree/Source/StateTreeModule/Public/StateTreeSchema.h
+++ b/Engine/Plugins/Runtime/StateTree/Source/StateTreeModule/Public/StateTreeSchema.h
@@ -13,6 +13,7 @@ enum class EStateTreeDataSourceType : uint8;
 enum class EStateTreeParameterDataType : uint8;
 enum class EStateTreeStateSelectionBehavior : uint8;
 enum class EStateTreeStateSelectionRules : uint32;
+enum class EStateTreeStateType : uint8;
 
 
 /**
@@ -61,6 +62,12 @@ public:
 		return true;
 	}
 
+	/** @return True if the state type is supported. */
+	virtual bool IsStateTypeAllowed(EStateTreeStateType InStateType) const
+	{
+		return true;
+	}
+
 	/**
 	 * Helper function to check if a class is any of the Blueprint extendable item classes (Eval, Task, Condition).
 	 * Can be used to quickly accept all of those classes in IsClassAllowed().
