[StateTree] Fixed checking global tasks status on non-global frames.

diff --git a/Engine/Plugins/Runtime/StateTree/Source/StateTreeEditorModule/Private/StateTreeCompiler.cpp b/Engine/Plugins/Runtime/StateTree/Source/StateTreeEditorModule/Private/StateTreeCompiler.cpp
--- a/Engine/Plugins/Runtime/StateTree/Source/StateTreeEditorModule/Private/StateTreeCompiler.cpp
+++ b/Engine/Plugins/Runtime/StateTree/Source/StateTreeEditorModule/Private/StateTreeCompiler.cpp
@@ -1065,9 +1065,8 @@ bool FStateTreeCompiler::CreateStateTasksAndParameters()
 			else
 			{
 				// Frame need an extra buffer for global tasks.
-				//Subtree do not contains global task.
-				const bool bFrameWithoutGlobalTasks = CompactState.Type == EStateTreeStateType::Subtree && StateIndex != 0;
-				StartBitIndex = bFrameWithoutGlobalTasks ? 0 : GlobalTaskEndBit;
+				//Linked sub-frames do not support global tasks but they can be use to replace root (old behavior).
+				StartBitIndex = GlobalTaskEndBit;
 			}
 
 			const UE::StateTree::Compiler::FCompletionTasksMaskResult MaskResult = UE::StateTree::Compiler::MakeCompletionTasksMask(StartBitIndex, Tasks, ValidTaskNodeIndex);
diff --git a/Engine/Plugins/Runtime/StateTree/Source/StateTreeModule/Private/StateTreeExecutionContext.cpp b/Engine/Plugins/Runtime/StateTree/Source/StateTreeModule/Private/StateTreeExecutionContext.cpp
--- a/Engine/Plugins/Runtime/StateTree/Source/StateTreeModule/Private/StateTreeExecutionContext.cpp
+++ b/Engine/Plugins/Runtime/StateTree/Source/StateTreeModule/Private/StateTreeExecutionContext.cpp
@@ -5408,7 +5408,7 @@ PRAGMA_ENABLE_DEPRECATION_WARNINGS
 			const FStateTreeExecutionFrame& CurrentFrame = Exec.ActiveFrames[FrameIndex];
 			const UStateTree* CurrentStateTree = CurrentFrame.StateTree;
 			const ETaskCompletionStatus FrameTasksStatus = CurrentFrame.ActiveTasksStatus.GetStatus(CurrentStateTree).GetCompletionStatus();
-			if (FrameTasksStatus != ETaskCompletionStatus::Running)
+			if (CurrentFrame.bIsGlobalFrame && FrameTasksStatus != ETaskCompletionStatus::Running)
 			{
 				if (FrameIndex == 0)
 				{
