[StateTree] fixed SCompactTreeView construction when view model is not provided

diff --git a/Engine/Plugins/Runtime/StateTree/Source/StateTreeEditorModule/Private/Customizations/StateTreeStateLinkDetails.cpp b/Engine/Plugins/Runtime/StateTree/Source/StateTreeEditorModule/Private/Customizations/StateTreeStateLinkDetails.cpp
--- a/Engine/Plugins/Runtime/StateTree/Source/StateTreeEditorModule/Private/Customizations/StateTreeStateLinkDetails.cpp
+++ b/Engine/Plugins/Runtime/StateTree/Source/StateTreeEditorModule/Private/Customizations/StateTreeStateLinkDetails.cpp
@@ -14,7 +14,6 @@
 #include "StateTreePropertyHelpers.h"
 #include "TextStyleDecorator.h"
 #include "Widgets/SCompactTreeEditorView.h"
-#include "Widgets/SCompactTreeView.h"
 #include "Widgets/Text/SRichTextBlock.h"
 
 #define LOCTEXT_NAMESPACE "StateTreeEditor"
diff --git a/Engine/Plugins/Runtime/StateTree/Source/StateTreeEditorModule/Private/Customizations/Widgets/SCompactTreeEditorView.cpp b/Engine/Plugins/Runtime/StateTree/Source/StateTreeEditorModule/Private/Customizations/Widgets/SCompactTreeEditorView.cpp
--- a/Engine/Plugins/Runtime/StateTree/Source/StateTreeEditorModule/Private/Customizations/Widgets/SCompactTreeEditorView.cpp
+++ b/Engine/Plugins/Runtime/StateTree/Source/StateTreeEditorModule/Private/Customizations/Widgets/SCompactTreeEditorView.cpp
@@ -35,12 +35,22 @@ void SCompactTreeEditorView::Construct(const FArguments& InArgs, const TSharedPt
 	bSelectableStatesOnly = InArgs._SelectableStatesOnly;
 	bShowLinkedStates = InArgs._ShowLinkedStates;
 
-	SCompactTreeView::Construct(
-		SCompactTreeView::FArguments()
-		.SelectionMode(InArgs._SelectionMode)
-		.OnSelectionChanged(InArgs._OnSelectionChanged)
-		.OnContextMenuOpening(InArgs._OnContextMenuOpening)
-		, InViewModel->GetStateTree());
+	const UStateTree* StateTreeAsset = InViewModel
+		? InViewModel->GetStateTree()
+		: InArgs._StateTreeEditorData
+			? InArgs._StateTreeEditorData->GetTypedOuter<UStateTree>()
+			: nullptr;
+
+	if (ensureMsgf(StateTreeAsset != nullptr
+		, TEXT("Expecting either a valid StateTreeViewModel or StateTreeEditorData to construct the view")))
+	{
+		SCompactTreeView::Construct(
+			SCompactTreeView::FArguments()
+			.SelectionMode(InArgs._SelectionMode)
+			.OnSelectionChanged(InArgs._OnSelectionChanged)
+			.OnContextMenuOpening(InArgs._OnContextMenuOpening)
+			, StateTreeAsset);
+	}
 }
 
 void SCompactTreeEditorView::CacheStatesInternal()
