[StateTreeDebugger] some tweaks to event messages and added new logs

diff --git a/Engine/Plugins/Runtime/StateTree/Source/StateTreeModule/Private/Debugger/StateTreeTraceTypes.cpp b/Engine/Plugins/Runtime/StateTree/Source/StateTreeModule/Private/Debugger/StateTreeTraceTypes.cpp
--- a/Engine/Plugins/Runtime/StateTree/Source/StateTreeModule/Private/Debugger/StateTreeTraceTypes.cpp
+++ b/Engine/Plugins/Runtime/StateTree/Source/StateTreeModule/Private/Debugger/StateTreeTraceTypes.cpp
@@ -123,7 +123,7 @@ FString FStateTreeTracePropertyEvent::GetTypeString(const UStateTree& StateTree)
 //----------------------------------------------------------------------//
 FString FStateTreeTraceTransitionEvent::ToFullString(const UStateTree& StateTree) const
 {
-	return FString::Printf(TEXT("%s %s"), *GetTypeString(StateTree), *GetValueString(StateTree));
+	return FString::Printf(TEXT("%s '%s': %s"), *UEnum::GetDisplayValueAsText(EventType).ToString(), *GetTypeString(StateTree), *GetValueString(StateTree));
 }
 
 FString FStateTreeTraceTransitionEvent::GetValueString(const UStateTree& StateTree) const
@@ -256,7 +256,26 @@ FString FStateTreeTraceStateEvent::GetValueString(const UStateTree& StateTree) c
 	const FStateTreeStateHandle StateHandle(Index.Get());
 	if (const FCompactStateTreeState* CompactState = StateTree.GetStateFromHandle(StateHandle))
 	{
-		return UE::StateTreeTrace::GetStateName(StateTree, CompactState);
+		TArray<FString, TInlineAllocator<FStateTreeActiveStates::MaxStates>> StateHierarchy;
+		StateHierarchy.Push(UE::StateTreeTrace::GetStateName(StateTree, CompactState));
+
+		// Climb hierarchy
+		for (FStateTreeStateHandle ParentHandle = CompactState->Parent; ParentHandle.IsValid(); ParentHandle = CompactState->Parent)
+		{
+			if (CompactState = StateTree.GetStateFromHandle(ParentHandle); CompactState != nullptr)
+			{
+				StateHierarchy.Push(UE::StateTreeTrace::GetStateName(StateTree, CompactState));
+			}
+			else
+			{
+				break;
+			}
+		}
+
+		// Reverse to output from the root
+		Algo::Reverse(StateHierarchy);
+
+		return FString::Join(StateHierarchy, TEXT("."));
 	}
 
 	return FString::Printf(TEXT("Invalid State Index %s for '%s'"), *StateHandle.Describe(), *StateTree.GetFullName());
diff --git a/Engine/Plugins/Runtime/StateTree/Source/StateTreeModule/Private/StateTreeExecutionContext.cpp b/Engine/Plugins/Runtime/StateTree/Source/StateTreeModule/Private/StateTreeExecutionContext.cpp
--- a/Engine/Plugins/Runtime/StateTree/Source/StateTreeModule/Private/StateTreeExecutionContext.cpp
+++ b/Engine/Plugins/Runtime/StateTree/Source/StateTreeModule/Private/StateTreeExecutionContext.cpp
@@ -4528,8 +4528,10 @@ bool FStateTreeExecutionContext::TestAllConditionsInternal(const FStateTreeExecu
 				if (!bBatchCopied)
 				{
 					// If the source data cannot be accessed, the whole expression evaluates to false.
+					static constexpr TCHAR Message[] = TEXT("Evaluation forced to false: source data cannot be accessed (e.g. enter conditions trying to access inactive parent state)");
 					UE_STATETREE_DEBUG_CONDITION_EVENT(this, ConditionIndex, ConditionInstanceView, EStateTreeTraceEventType::InternalForcedFailure);
-					UE_STATETREE_DEBUG_LOG_EVENT(this, Warning, TEXT("Evaluation forced to false: source data cannot be accessed (e.g. enter conditions trying to access inactive parent state)"));
+					UE_STATETREE_DEBUG_LOG_EVENT(this, Warning, Message);
+					STATETREE_LOG(Warning, TEXT("%s"), Message);
 					Values[0] = false;
 					break;
 				}
@@ -5551,7 +5553,7 @@ PRAGMA_ENABLE_DEPRECATION_WARNINGS
 			if (!RequestedTransition.IsValid())
 			{
 				STATETREE_LOG(Verbose, TEXT("Could not trigger completion transition, jump back to root state."));
-				UE_STATETREE_DEBUG_LOG_EVENT(this, Warning, TEXT("Could not trigger completion transition, jump back to root state."));
+				UE_STATETREE_DEBUG_LOG_EVENT(this, Log, TEXT("Could not trigger completion transition, jump back to root state."));
 
 				check(!Exec.ActiveFrames.IsEmpty());
 				FStateTreeExecutionFrame& RootFrame = Exec.ActiveFrames[0];
@@ -6652,6 +6654,9 @@ TOptional<bool> FStateTreeExecutionContext::SelectStateFromSourceInternal(
 						*GetNameSafe(&Owner),
 						*GetFullNameSafe(NextFrame.StateTree)
 					);
+
+					UE_STATETREE_DEBUG_LOG_EVENT(this, Log, TEXT("Selection fails because parent state '%s' is completed."), *NextState.Name.ToString());
+
 					bSucceededToSelectState = false;
 				}
 			}
diff --git a/Engine/Plugins/Runtime/StateTree/Source/StateTreeModule/Public/Debugger/StateTreeTraceTypes.h b/Engine/Plugins/Runtime/StateTree/Source/StateTreeModule/Public/Debugger/StateTreeTraceTypes.h
--- a/Engine/Plugins/Runtime/StateTree/Source/StateTreeModule/Public/Debugger/StateTreeTraceTypes.h
+++ b/Engine/Plugins/Runtime/StateTree/Source/StateTreeModule/Public/Debugger/StateTreeTraceTypes.h
@@ -42,7 +42,7 @@ enum class EStateTreeTraceEventType : uint8
 	InternalForcedFailure	UMETA(DisplayName = "Internal Forced Failure"),
 	OnRequesting			UMETA(DisplayName = "Requesting"),
 	OnEvaluating			UMETA(DisplayName = "Evaluating"),
-	OnTransition			UMETA(DisplayName = "Transition"),
+	OnTransition			UMETA(DisplayName = "Applying Transition"),
 	OnTreeStarted			UMETA(DisplayName = "Tree Started"),
 	OnTreeStopped			UMETA(DisplayName = "Tree Stopped")
 };
