[StateTree] Add transition when setting the completion task flag.

diff --git a/Engine/Plugins/Runtime/StateTree/Source/StateTreeEditorModule/Private/Customizations/StateTreeEditorNodeDetails.cpp b/Engine/Plugins/Runtime/StateTree/Source/StateTreeEditorModule/Private/Customizations/StateTreeEditorNodeDetails.cpp
--- a/Engine/Plugins/Runtime/StateTree/Source/StateTreeEditorModule/Private/Customizations/StateTreeEditorNodeDetails.cpp
+++ b/Engine/Plugins/Runtime/StateTree/Source/StateTreeEditorModule/Private/Customizations/StateTreeEditorNodeDetails.cpp
@@ -1663,14 +1663,20 @@ void FStateTreeEditorNodeDetails::HandleNameCommitted(const FText& NewText, ETex
 
 FReply FStateTreeEditorNodeDetails::HandleToggleCompletionTaskClicked()
 {
-	if (FStateTreeEditorNode* Node = UE::StateTreeEditor::EditorNodeUtils::GetMutableCommonNode(StructProperty))
-	{
-		if (UE::StateTreeEditor::EditorNodeUtils::IsTaskEnabled(*Node))
+	UE::StateTreeEditor::EditorNodeUtils::ModifyNodeInTransaction(LOCTEXT("OnCompletionTaskToggled", "Toggled Completion Task"),
+		StructProperty,
+		[](const TSharedPtr<IPropertyHandle>& StructProperty)
 		{
-			const bool bCurrentValue = UE::StateTreeEditor::EditorNodeUtils::IsTaskConsideredForCompletion(*Node);
-			UE::StateTreeEditor::EditorNodeUtils::SetTaskConsideredForCompletion(*Node, !bCurrentValue);
-		}
-	}
+			if (FStateTreeEditorNode* Node = UE::StateTreeEditor::EditorNodeUtils::GetMutableCommonNode(StructProperty))
+			{
+				if (UE::StateTreeEditor::EditorNodeUtils::IsTaskEnabled(*Node))
+				{
+					const bool bCurrentValue = UE::StateTreeEditor::EditorNodeUtils::IsTaskConsideredForCompletion(*Node);
+					UE::StateTreeEditor::EditorNodeUtils::SetTaskConsideredForCompletion(*Node, !bCurrentValue);
+				}
+			}
+		});
+
 	return FReply::Handled();
 }
 
