[StateTree] Remove deprecated variables from the detail view and the shipping target.

diff --git a/Engine/Plugins/Runtime/StateTree/Source/StateTreeModule/Public/StateTreeExecutionTypes.h b/Engine/Plugins/Runtime/StateTree/Source/StateTreeModule/Public/StateTreeExecutionTypes.h
--- a/Engine/Plugins/Runtime/StateTree/Source/StateTreeModule/Public/StateTreeExecutionTypes.h
+++ b/Engine/Plugins/Runtime/StateTree/Source/StateTreeModule/Public/StateTreeExecutionTypes.h
@@ -1067,7 +1067,7 @@ PRAGMA_ENABLE_DEPRECATION_WARNINGS
 	/** @return the active frame index that matches the unique frame ID. */
 	UE_API int32 IndexOfActiveFrame(UE::StateTree::FActiveFrameID FrameID) const;
 
-	/** @return whether it contains any shcedule tick requests. */
+	/** @return whether it contains any scheduled tick requests. */
 	bool HasScheduledTickRequests() const
 	{
 		return ScheduledTickRequests.Num() > 0;
@@ -1252,7 +1252,7 @@ PRAGMA_ENABLE_DEPRECATION_WARNINGS
 PRAGMA_DISABLE_DEPRECATION_WARNINGS
 	UE_DEPRECATED(5.7, "Use the StateTreeExecutionContext.RequestTransitionResult.Selection.SelectedState instead.")
 	/** States selected as result of the transition. */
-	UPROPERTY(EditDefaultsOnly, Category = "Default", BlueprintReadOnly, meta = (DeprecatedProperty))
+	UPROPERTY(Category = "Default", BlueprintReadOnly, meta = (DeprecatedProperty))
 	TArray<FStateTreeExecutionFrame> NextActiveFrames;
 
 	UE_DEPRECATED(5.7, "Use the StateTreeExecutionContext.RequestTransitionResult.Selection.SelectionEvents instead.")
@@ -1261,17 +1261,17 @@ PRAGMA_DISABLE_DEPRECATION_WARNINGS
 
 	UE_DEPRECATED(5.6, "Use SourceFrameID instead")
 	/** StateTree asset that was active when the transition was requested. */
-	UPROPERTY(EditDefaultsOnly, Category = "Default", BlueprintReadOnly, meta = (DeprecatedProperty))
+	UPROPERTY(Category = "Default", BlueprintReadOnly, meta = (DeprecatedProperty))
 	TObjectPtr<const UStateTree> SourceStateTree = nullptr;
 
 	UE_DEPRECATED(5.6, "Use SourceFrameID instead.")
 	/** Root state the execution frame where the transition was requested. */
-	UPROPERTY(EditDefaultsOnly, Category = "Default", BlueprintReadOnly, meta = (DeprecatedProperty))
+	UPROPERTY(Category = "Default", BlueprintReadOnly, meta = (DeprecatedProperty))
 	FStateTreeStateHandle SourceRootState = FStateTreeStateHandle::Invalid;
 
 	UE_DEPRECATED(5.6, "Use SourceStateID instead")
 	/** Transition source state. */
-	UPROPERTY(EditDefaultsOnly, Category = "Default", BlueprintReadOnly, meta=(DeprecatedProperty))
+	UPROPERTY(Category = "Default", BlueprintReadOnly, meta=(DeprecatedProperty))
 	FStateTreeStateHandle SourceState = FStateTreeStateHandle::Invalid;
 PRAGMA_ENABLE_DEPRECATION_WARNINGS
 };
