[StateTree] prevent access to null schema struct type

diff --git a/Engine/Plugins/Runtime/StateTree/Source/StateTreeEditorModule/Private/StateTreeCompiler.cpp b/Engine/Plugins/Runtime/StateTree/Source/StateTreeEditorModule/Private/StateTreeCompiler.cpp
--- a/Engine/Plugins/Runtime/StateTree/Source/StateTreeEditorModule/Private/StateTreeCompiler.cpp
+++ b/Engine/Plugins/Runtime/StateTree/Source/StateTreeEditorModule/Private/StateTreeCompiler.cpp
@@ -365,7 +365,9 @@ bool FStateTreeCompiler::Compile(UStateTree& InStateTree)
 		StateTree->ContextDataDescs = StateTree->Schema->GetContextDataDescs();
 		for (FStateTreeExternalDataDesc& Desc : StateTree->ContextDataDescs)
 		{
-			const FStateTreeBindableStructDesc ExtDataDesc = {
+			if (Desc.Struct)
+			{
+				const FStateTreeBindableStructDesc ExtDataDesc = {
 					UE::StateTree::Editor::GlobalStateName,
 					Desc.Name,
 					Desc.Struct,
@@ -373,14 +375,15 @@ bool FStateTreeCompiler::Compile(UStateTree& InStateTree)
 					EStateTreeBindableStructSource::Context,
 					Desc.ID
 				};
-			BindingsCompiler.AddSourceStruct(ExtDataDesc);
-			if (const auto Validation = UE::StateTree::Compiler::IsValidIndex16(ContextDataIndex); Validation.DidFail())
-			{
-				Validation.Log(Log, TEXT("ExternalStructIndex"), ExtDataDesc);
-				return false;
+				BindingsCompiler.AddSourceStruct(ExtDataDesc);
+				if (const auto Validation = UE::StateTree::Compiler::IsValidIndex16(ContextDataIndex); Validation.DidFail())
+				{
+					Validation.Log(Log, TEXT("ExternalStructIndex"), ExtDataDesc);
+					return false;
+				}
+				Desc.Handle.DataHandle = ExtDataDesc.DataHandle;
 			}
-			Desc.Handle.DataHandle = ExtDataDesc.DataHandle;
-		} 
+		}
 	}
 
 	if (const UE::StateTree::Compiler::FValidationResult Validation = UE::StateTree::Compiler::IsValidIndex16(ContextDataIndex); Validation.DidFail())
diff --git a/Engine/Plugins/Runtime/StateTree/Source/StateTreeEditorModule/Private/StateTreeEditorData.cpp b/Engine/Plugins/Runtime/StateTree/Source/StateTreeEditorModule/Private/StateTreeEditorData.cpp
--- a/Engine/Plugins/Runtime/StateTree/Source/StateTreeEditorModule/Private/StateTreeEditorData.cpp
+++ b/Engine/Plugins/Runtime/StateTree/Source/StateTreeEditorModule/Private/StateTreeEditorData.cpp
@@ -507,7 +507,8 @@ FStateTreeBindableStructDesc UStateTreeEditorData::FindContextData(const UStruct
 	TArray<FStateTreeBindableStructDesc> Candidates;
 	for (const FStateTreeExternalDataDesc& Desc : Schema->GetContextDataDescs())
 	{
-		if (Desc.Struct->IsChildOf(ObjectType))
+		if (Desc.Struct != nullptr
+			&& Desc.Struct->IsChildOf(ObjectType))
 		{
 			Candidates.Emplace(UE::StateTree::Editor::GlobalStateName, Desc.Name, Desc.Struct, FStateTreeDataHandle(), EStateTreeBindableStructSource::Context, Desc.ID);
 		}
@@ -1531,17 +1532,20 @@ EStateTreeVisitor UStateTreeEditorData::VisitGlobalNodes(TFunctionRef<EStateTree
 	{
 		for (const FStateTreeExternalDataDesc& ContextDesc : Schema->GetContextDataDescs())
 		{
-			FStateTreeBindableStructDesc Desc;
-			Desc.StatePath = UE::StateTree::Editor::GlobalStateName;
-			Desc.Struct = ContextDesc.Struct;
-			Desc.Name = ContextDesc.Name;
-			Desc.ID = ContextDesc.ID;
-			Desc.DataSource = EStateTreeBindableStructSource::Context;
-
-			// We don't have value for the external objects, but return the type and null value so that users of GetAllStructValues() can use the type. 
-			if (InFunc(nullptr, Desc, FStateTreeDataView(Desc.Struct, nullptr)) == EStateTreeVisitor::Break)
+			if (ContextDesc.Struct)
 			{
-				return EStateTreeVisitor::Break;
+				FStateTreeBindableStructDesc Desc;
+				Desc.StatePath = UE::StateTree::Editor::GlobalStateName;
+				Desc.Struct = ContextDesc.Struct;
+				Desc.Name = ContextDesc.Name;
+				Desc.ID = ContextDesc.ID;
+				Desc.DataSource = EStateTreeBindableStructSource::Context;
+
+				// We don't have value for the external objects, but return the type and null value so that users of GetAllStructValues() can use the type.
+				if (InFunc(nullptr, Desc, FStateTreeDataView(Desc.Struct, nullptr)) == EStateTreeVisitor::Break)
+				{
+					return EStateTreeVisitor::Break;
+				}
 			}
 		}
 	}
