[StateTreeDebugger] fixed timeline after an instance lifetime is over
- optim to not rebuild events when scrubtime didn't change
- optim to reuse active states per assets while processing state recursively

diff --git a/Engine/Plugins/Runtime/StateTree/Source/StateTreeDeveloper/Internal/Debugger/SCompactTreeDebuggerView.h b/Engine/Plugins/Runtime/StateTree/Source/StateTreeDeveloper/Internal/Debugger/SCompactTreeDebuggerView.h
--- a/Engine/Plugins/Runtime/StateTree/Source/StateTreeDeveloper/Internal/Debugger/SCompactTreeDebuggerView.h
+++ b/Engine/Plugins/Runtime/StateTree/Source/StateTreeDeveloper/Internal/Debugger/SCompactTreeDebuggerView.h
@@ -69,7 +69,10 @@ private:
 		uint16 StateIdx;
 	};
 
-	void CacheStateRecursive(TNotNull<const UStateTree*> InStateTree, TSharedPtr<FStateItem> InParentItem, const uint16 InStateIdx, TArray<FProcessedState>& OutProcessedStates);
+	void CacheStateRecursive(TNotNull<const FStateTreeTraceActiveStates::FAssetActiveStates*> InAssetActiveStates
+		, TSharedPtr<FStateItem> InParentItem
+		, const uint16 InStateIdx
+		, TArray<FProcessedState>& OutProcessedStates);
 
 	TAttribute<FStateTreeTraceActiveStates> AllActiveStates;
 };
diff --git a/Engine/Plugins/Runtime/StateTree/Source/StateTreeDeveloper/Private/Debugger/SCompactTreeDebuggerView.cpp b/Engine/Plugins/Runtime/StateTree/Source/StateTreeDeveloper/Private/Debugger/SCompactTreeDebuggerView.cpp
--- a/Engine/Plugins/Runtime/StateTree/Source/StateTreeDeveloper/Private/Debugger/SCompactTreeDebuggerView.cpp
+++ b/Engine/Plugins/Runtime/StateTree/Source/StateTreeDeveloper/Private/Debugger/SCompactTreeDebuggerView.cpp
@@ -34,18 +34,19 @@ void SCompactTreeDebuggerView::CacheStatesInternal()
 		if (PerAssetActiveStates && PerAssetActiveStates->ActiveStates.Num() > 0)
 		{
 			TArray<FProcessedState> ProcessedStates;
-			CacheStateRecursive(StateTree, RootItem, PerAssetActiveStates->ActiveStates[0].Index, ProcessedStates);
+			CacheStateRecursive(PerAssetActiveStates, RootItem, PerAssetActiveStates->ActiveStates[0].Index, ProcessedStates);
 		}
 	}
 }
 
 void SCompactTreeDebuggerView::CacheStateRecursive(
-	const TNotNull<const UStateTree*> InStateTree
+	TNotNull<const FStateTreeTraceActiveStates::FAssetActiveStates*> InAssetActiveStates
 	, TSharedPtr<FStateItem> InParentItem
 	, const uint16 InStateIdx
 	, TArray<FProcessedState>& OutProcessedStates)
 {
-	const FProcessedState ProcessedState(InStateTree, InStateIdx);
+	TNotNull<const UStateTree*> StateTree = InAssetActiveStates->WeakStateTree.Get();
+	const FProcessedState ProcessedState(StateTree, InStateIdx);
 	if (OutProcessedStates.Contains(ProcessedState))
 	{
 		return;
@@ -53,20 +54,14 @@ void SCompactTreeDebuggerView::CacheStateRecursive(
 
 	OutProcessedStates.Add(ProcessedState);
 
-	const FCompactStateTreeState& State = InStateTree->GetStates()[InStateIdx];
+	const FCompactStateTreeState& State = StateTree->GetStates()[InStateIdx];
 	FString StringForDebug = State.Name.ToString();
 
-	const FStateTreeTraceActiveStates::FAssetActiveStates* AssetActiveStates = AllActiveStates.Get().PerAssetStates.FindByPredicate(
-		[StateTree = InStateTree](const FStateTreeTraceActiveStates::FAssetActiveStates& AssetActiveStates)
-		{
-			return AssetActiveStates.WeakStateTree == StateTree;
-		});
-	const bool bIsActiveState = AssetActiveStates != nullptr ? AssetActiveStates->ActiveStates.Contains(FStateTreeStateHandle(InStateIdx)) : false;
-
-	auto AddStateItem = [&InParentItem, State, InStateTree, InStateIdx](const TSharedRef<FStateItem>& StateItem, const bool bIsActiveState)
+	const bool bIsActiveState = InAssetActiveStates->ActiveStates.Contains(FStateTreeStateHandle(InStateIdx));
+	auto AddStateItem = [&InParentItem, State, StateTree, InStateIdx](const TSharedRef<FStateItem>& StateItem, const bool bIsActiveState)
 		{
 			StateItem->Desc = FText::FromName(State.Name);
-			StateItem->StateID = InStateTree->GetStateIdFromHandle(FStateTreeStateHandle(InStateIdx));
+			StateItem->StateID = StateTree->GetStateIdFromHandle(FStateTreeStateHandle(InStateIdx));
 			StateItem->TooltipText = FText::FromName(State.Name);
 			StateItem->bIsEnabled = State.bEnabled;
 			StateItem->CustomData.GetMutable<UE::StateTree::CompactTreeView::FStateItemDebuggerData>().bIsActive = bIsActiveState;
@@ -85,7 +80,7 @@ void SCompactTreeDebuggerView::CacheStateRecursive(
 		// Recurse to linked tree only for active states
 		if (bIsActiveState)
 		{
-			CacheStateRecursive(InStateTree, InParentItem, State.LinkedState.Index, OutProcessedStates);
+			CacheStateRecursive(InAssetActiveStates, InParentItem, State.LinkedState.Index, OutProcessedStates);
 		}
 	}
 	// Add external subtree
@@ -107,7 +102,7 @@ void SCompactTreeDebuggerView::CacheStateRecursive(
 				if (LinkedAssetActiveStates && LinkedAssetActiveStates->ActiveStates.Num() > 0)
 				{
 					TArray<FProcessedState> ProcessedStates;
-					CacheStateRecursive(State.LinkedAsset, InParentItem, LinkedAssetActiveStates->ActiveStates[0].Index, OutProcessedStates);
+					CacheStateRecursive(LinkedAssetActiveStates, InParentItem, LinkedAssetActiveStates->ActiveStates[0].Index, OutProcessedStates);
 				}
 			}
 		}
@@ -121,7 +116,7 @@ void SCompactTreeDebuggerView::CacheStateRecursive(
 		{
 			for (uint16 ChildIdx = State.ChildrenBegin; ChildIdx < State.ChildrenEnd; ChildIdx++)
 			{
-				CacheStateRecursive(InStateTree, InParentItem, ChildIdx, OutProcessedStates);
+				CacheStateRecursive(InAssetActiveStates, InParentItem, ChildIdx, OutProcessedStates);
 			}
 		}
 	}
diff --git a/Engine/Plugins/Runtime/StateTree/Source/StateTreeEditorModule/Private/Debugger/SStateTreeDebuggerView.cpp b/Engine/Plugins/Runtime/StateTree/Source/StateTreeEditorModule/Private/Debugger/SStateTreeDebuggerView.cpp
--- a/Engine/Plugins/Runtime/StateTree/Source/StateTreeEditorModule/Private/Debugger/SStateTreeDebuggerView.cpp
+++ b/Engine/Plugins/Runtime/StateTree/Source/StateTreeEditorModule/Private/Debugger/SStateTreeDebuggerView.cpp
@@ -670,13 +670,16 @@ void SStateTreeDebuggerView::Construct(const FArguments& InArgs, const UStateTre
 		]
 	];
 	
-	// Auto-select session if there is only one available and that we are trac
+	// Auto-select session if there is only one available and that we are tracing
 	// Do that after creating all our widgets in case we receive a callback
-	TArray<FStateTreeDebugger::FTraceDescriptor> TraceDescriptors;
-	Debugger->GetLiveTraces(TraceDescriptors);
-	if (TraceDescriptors.Num() == 1 && bRecording)
+	if (bRecording)
 	{
-		Debugger->RequestSessionAnalysis(TraceDescriptors.Last());
+		TArray<FStateTreeDebugger::FTraceDescriptor> TraceDescriptors;
+		Debugger->GetLiveTraces(TraceDescriptors);
+		if (TraceDescriptors.Num() == 1)
+		{
+			Debugger->RequestSessionAnalysis(TraceDescriptors.Last());
+		}
 	}
 }
 
diff --git a/Engine/Plugins/Runtime/StateTree/Source/StateTreeEditorModule/Private/Debugger/StateTreeDebuggerTrack.cpp b/Engine/Plugins/Runtime/StateTree/Source/StateTreeEditorModule/Private/Debugger/StateTreeDebuggerTrack.cpp
--- a/Engine/Plugins/Runtime/StateTree/Source/StateTreeEditorModule/Private/Debugger/StateTreeDebuggerTrack.cpp
+++ b/Engine/Plugins/Runtime/StateTree/Source/StateTreeEditorModule/Private/Debugger/StateTreeDebuggerTrack.cpp
@@ -174,6 +174,11 @@ bool FStateTreeDebuggerInstanceTrack::UpdateInternal()
 		// otherwise we build up to the last time the track was active.
 		InstanceLastActiveTime = Descriptor->IsInstanceActive() ? StaleTime : Descriptor.Get()->Lifetime.GetUpperBoundValue();
 	}
+	else if (!Descriptor->IsInstanceActive())
+	{
+		// If the instance is no longer active we don't want to build data past that time
+		InstanceLastActiveTime = Descriptor.Get()->Lifetime.GetUpperBoundValue();
+	}
 
 	const bool bChanged = InstanceTrackHelper.Get().RebuildEventData(Debugger.GetAsset()
 		, Debugger.GetEventCollection(InstanceTrackHelper.Get().GetInstanceId())
diff --git a/Engine/Plugins/Runtime/StateTree/Source/StateTreeEditorModule/Private/Debugger/StateTreeRewindDebuggerTrack.cpp b/Engine/Plugins/Runtime/StateTree/Source/StateTreeEditorModule/Private/Debugger/StateTreeRewindDebuggerTrack.cpp
--- a/Engine/Plugins/Runtime/StateTree/Source/StateTreeEditorModule/Private/Debugger/StateTreeRewindDebuggerTrack.cpp
+++ b/Engine/Plugins/Runtime/StateTree/Source/StateTreeEditorModule/Private/Debugger/StateTreeRewindDebuggerTrack.cpp
@@ -168,24 +168,30 @@ bool FRewindDebuggerTrack::UpdateInternal()
 			, FTraceFilter(InstanceTrackHelper.Get().GetInstanceId())
 			, LastTraceReadTime);
 
-		FStateTreeTraceActiveStates ActiveStates;
-		const bool bTimelineDataChanged = InstanceTrackHelper.Get().RebuildEventData(StateTree
-			, EventCollection
-			, RewindDebugger->GetRecordingDuration()
-			, RewindDebugger->GetScrubTime()
-			, &ActiveStates);
-
-		const bool bActivateStatesChanged = ActiveStates != InstanceTrackHelper.Get().GetActiveStates();
-		if (bActivateStatesChanged)
+		if (!FMath::IsNearlyEqual(LastUpdateScrubTime, RewindDebugger->GetScrubTime()))
 		{
-			InstanceTrackHelper.Get().SetActiveStates(ActiveStates);
-		}
+			const double InstanceLastActiveTime = Descriptor->IsInstanceActive() ? RewindDebugger->GetRecordingDuration() : Descriptor->Lifetime.GetUpperBoundValue();
+			LastUpdateScrubTime = RewindDebugger->GetScrubTime();
+
+			FStateTreeTraceActiveStates ActiveStates;
+			const bool bTimelineDataChanged = InstanceTrackHelper.Get().RebuildEventData(StateTree
+				, EventCollection
+				, InstanceLastActiveTime
+				, RewindDebugger->GetScrubTime()
+				, &ActiveStates);
+
+			const bool bActivateStatesChanged = ActiveStates != InstanceTrackHelper.Get().GetActiveStates();
+			if (bActivateStatesChanged)
+			{
+				InstanceTrackHelper.Get().SetActiveStates(ActiveStates);
+			}
 
-		if (bTimelineDataChanged
-			|| bActivateStatesChanged)
-		{
-			DebugTreeView->Refresh();
-			return true;
+			if (bTimelineDataChanged
+				|| bActivateStatesChanged)
+			{
+				DebugTreeView->Refresh();
+				return true;
+			}
 		}
 	}
 
diff --git a/Engine/Plugins/Runtime/StateTree/Source/StateTreeEditorModule/Private/Debugger/StateTreeRewindDebuggerTrack.h b/Engine/Plugins/Runtime/StateTree/Source/StateTreeEditorModule/Private/Debugger/StateTreeRewindDebuggerTrack.h
--- a/Engine/Plugins/Runtime/StateTree/Source/StateTreeEditorModule/Private/Debugger/StateTreeRewindDebuggerTrack.h
+++ b/Engine/Plugins/Runtime/StateTree/Source/StateTreeEditorModule/Private/Debugger/StateTreeRewindDebuggerTrack.h
@@ -65,6 +65,9 @@ private:
 
 	/** Last time in the recording that we used to fetch events, and that we will use for the next read. */
 	double LastTraceReadTime = 0;
+
+	/** Last scrub time used to rebuild the event data. */
+	double LastUpdateScrubTime = 0;
 };
 
 } // UE::StateTreeDebugger
\ No newline at end of file
