[StateTreeDebugger] initialize debugger view with existing debugger information when opening the view after an analysis was started.
Also converted Tracks shared references to weak ones for the View and Debugger. This allow the closed debugger view tab to properly removed the active states information in the main StateTreeView

diff --git a/Engine/Plugins/Runtime/StateTree/Source/StateTreeEditorModule/Private/Debugger/SStateTreeDebuggerView.cpp b/Engine/Plugins/Runtime/StateTree/Source/StateTreeEditorModule/Private/Debugger/SStateTreeDebuggerView.cpp
--- a/Engine/Plugins/Runtime/StateTree/Source/StateTreeEditorModule/Private/Debugger/SStateTreeDebuggerView.cpp
+++ b/Engine/Plugins/Runtime/StateTree/Source/StateTreeEditorModule/Private/Debugger/SStateTreeDebuggerView.cpp
@@ -669,11 +669,25 @@ void SStateTreeDebuggerView::Construct(const FArguments& InArgs, const UStateTre
 			]
 		]
 	];
-	
-	// Auto-select session if there is only one available and that we are tracing
-	// Do that after creating all our widgets in case we receive a callback
-	if (bRecording)
+
+	if (Debugger->IsAnalysisSessionActive())
+	{
+		// Creating the view for an already active analysis session
+		// so rebuild track views for all instances and scrub to the analysis current time 
+
+		TArray<const TSharedRef<const UE::StateTreeDebugger::FInstanceDescriptor>> InstanceDescriptors;
+		Debugger->GetSessionInstanceDescriptors(InstanceDescriptors);
+		for (const TSharedRef<const UE::StateTreeDebugger::FInstanceDescriptor>& InstanceDescriptor : InstanceDescriptors)
+		{
+			OnNewInstance(InstanceDescriptor->Id);
+		}
+
+		Debugger->SetScrubTime(Debugger->GetScrubTime());
+	}
+	else if (bRecording)
 	{
+		// Auto-select session if there is only one available and that we are tracing
+		// Do that after creating all our widgets in case we receive a callback
 		TArray<FStateTreeDebugger::FTraceDescriptor> TraceDescriptors;
 		Debugger->GetLiveTraces(TraceDescriptors);
 		if (TraceDescriptors.Num() == 1)
diff --git a/Engine/Plugins/Runtime/StateTree/Source/StateTreeEditorModule/Private/Debugger/StateTreeDebuggerTrack.cpp b/Engine/Plugins/Runtime/StateTree/Source/StateTreeEditorModule/Private/Debugger/StateTreeDebuggerTrack.cpp
--- a/Engine/Plugins/Runtime/StateTree/Source/StateTreeEditorModule/Private/Debugger/StateTreeDebuggerTrack.cpp
+++ b/Engine/Plugins/Runtime/StateTree/Source/StateTreeEditorModule/Private/Debugger/StateTreeDebuggerTrack.cpp
@@ -159,14 +159,23 @@ FStateTreeDebuggerInstanceTrack::FStateTreeDebuggerInstanceTrack(
 
 void FStateTreeDebuggerInstanceTrack::OnSelected()
 {
-	StateTreeDebugger.Get().SelectInstance(InstanceTrackHelper.Get().GetInstanceId());
+	if (const TSharedPtr<FStateTreeDebugger> Debugger = StateTreeDebugger.Pin())
+	{
+		Debugger->SelectInstance(InstanceTrackHelper.Get().GetInstanceId());
+	}
 }
 
 bool FStateTreeDebuggerInstanceTrack::UpdateInternal()
 {
-	const FStateTreeDebugger& Debugger = StateTreeDebugger.Get();
+	const TSharedPtr<FStateTreeDebugger> Debugger = StateTreeDebugger.Pin();
+	const TSharedPtr<SStateTreeDebuggerView> DebuggerView = StateTreeDebuggerView.Pin();
+	if (!Debugger.IsValid()
+		|| !DebuggerView.IsValid())
+	{
+		return false;
+	}
 
-	double InstanceLastActiveTime = StateTreeDebuggerView.Get().GetExtrapolatedRecordedWorldTime();
+	double InstanceLastActiveTime = DebuggerView->GetExtrapolatedRecordedWorldTime();
 	const bool bIsStaleTrack = IsStale();
 	if (bIsStaleTrack)
 	{
@@ -180,10 +189,10 @@ bool FStateTreeDebuggerInstanceTrack::UpdateInternal()
 		InstanceLastActiveTime = Descriptor.Get()->Lifetime.GetUpperBoundValue();
 	}
 
-	const bool bChanged = InstanceTrackHelper.Get().RebuildEventData(Debugger.GetAsset()
-		, Debugger.GetEventCollection(InstanceTrackHelper.Get().GetInstanceId())
+	const bool bChanged = InstanceTrackHelper.Get().RebuildEventData(Debugger->GetAsset()
+		, Debugger->GetEventCollection(InstanceTrackHelper.Get().GetInstanceId())
 		, InstanceLastActiveTime
-		, StateTreeDebugger->GetScrubTime()
+		, Debugger->GetScrubTime()
 		, /*OutLastActiveStates*/ nullptr
 		, bIsStaleTrack);
 
diff --git a/Engine/Plugins/Runtime/StateTree/Source/StateTreeEditorModule/Private/Debugger/StateTreeDebuggerTrack.h b/Engine/Plugins/Runtime/StateTree/Source/StateTreeEditorModule/Private/Debugger/StateTreeDebuggerTrack.h
--- a/Engine/Plugins/Runtime/StateTree/Source/StateTreeEditorModule/Private/Debugger/StateTreeDebuggerTrack.h
+++ b/Engine/Plugins/Runtime/StateTree/Source/StateTreeEditorModule/Private/Debugger/StateTreeDebuggerTrack.h
@@ -165,8 +165,8 @@ protected:
 
 private:
 	
-	TSharedRef<SStateTreeDebuggerView> StateTreeDebuggerView;
-	TSharedRef<FStateTreeDebugger> StateTreeDebugger;
+	TWeakPtr<SStateTreeDebuggerView> StateTreeDebuggerView;
+	TWeakPtr<FStateTreeDebugger> StateTreeDebugger;
 	TSharedRef<UE::StateTreeDebugger::FInstanceTrackHelper> InstanceTrackHelper;
 	TSharedPtr<const UE::StateTreeDebugger::FInstanceDescriptor> Descriptor;
 	double StaleTime = UE::StateTreeDebugger::FInstanceDescriptor::ActiveInstanceEndTime;
