[StateTree] Output the frames and completed mask at compile time. Enabled with cvar StateTree.Compiler.LogResultOnCompilationCompleted

#rnx

diff --git a/Engine/Plugins/Runtime/StateTree/Source/StateTreeModule/Private/StateTree.cpp b/Engine/Plugins/Runtime/StateTree/Source/StateTreeModule/Private/StateTree.cpp
--- a/Engine/Plugins/Runtime/StateTree/Source/StateTreeModule/Private/StateTree.cpp
+++ b/Engine/Plugins/Runtime/StateTree/Source/StateTreeModule/Private/StateTree.cpp
@@ -1436,6 +1436,22 @@ FString UStateTree::DebugInternalLayoutAsString() const
 	DebugString << PropertyBindings.DebugAsString();
 #endif
 
+	// Frames
+	DebugString.Appendf(TEXT("\nFrames(%d)\n"), Frames.Num());
+	if (Frames.Num())
+	{
+		DebugString.Appendf(TEXT("  [ (Idx) | %-4s ]\n"),
+			TEXT("Root"), TEXT("NumTasks"));
+		for (int32 FrameIndex = 0; FrameIndex < Frames.Num(); ++FrameIndex)
+		{
+			const FCompactStateTreeFrame& Frame = Frames[FrameIndex];
+			DebugString.Appendf(TEXT("  | (%3d) | %-4d |\n"),
+				FrameIndex,
+				Frame.RootState.Index
+			);
+		}
+	}
+
 	// States
 	DebugString.Appendf(TEXT("\nStates(%d)\n"), States.Num());
 	if (States.Num())
@@ -1452,7 +1468,27 @@ FString UStateTree::DebugInternalLayoutAsString() const
 				*State.Parent.Describe(),
 				TEXT(" "), State.ChildrenBegin, State.ChildrenEnd,
 				TEXT(" "), State.EnterConditionsBegin, State.TransitionsBegin, State.TasksBegin, State.UtilityConsiderationsBegin,
-				TEXT(" "), State.EnterConditionsNum, State.TransitionsNum, State.TasksNum, State.UtilityConsiderationsNum);
+				TEXT(" "), State.EnterConditionsNum, State.TransitionsNum, State.TasksNum, State.UtilityConsiderationsNum
+			);
+		}
+
+		auto AppendBinary = [&DebugString](const uint32 Mask)
+			{
+				for (int32 Index = (sizeof(uint32) * 8) - 1; Index >= 0; --Index)
+				{
+					const TCHAR BinaryValue = ((Mask >> Index) & 1U) ? TEXT('1') : TEXT('0');
+					DebugString << BinaryValue;
+				}
+			};
+
+		DebugString.Appendf(TEXT("\n  [ (Idx) | %32s | %8s ]\n"),
+			TEXT("CompMask"), TEXT("CustTick"));
+		for (int32 StateIndex = 0; StateIndex < States.Num(); ++StateIndex)
+		{
+			const FCompactStateTreeState& State = States[StateIndex];
+			DebugString.Appendf(TEXT("  | (%3d) | "), StateIndex);
+			AppendBinary(State.CompletionTasksMask);
+			DebugString.Appendf(TEXT(" | %3f |\n"), State.CustomTickRate);
 		}
 	}
 
