[StateTree] Fix unittest. Explicitly state that we are waiting for a warning. Introduced with CL 43608536.

#rnx

diff --git a/Engine/Plugins/Runtime/StateTree/Source/StateTreeTestSuite/Private/StateTreeTest.cpp b/Engine/Plugins/Runtime/StateTree/Source/StateTreeTestSuite/Private/StateTreeTest.cpp
--- a/Engine/Plugins/Runtime/StateTree/Source/StateTreeTestSuite/Private/StateTreeTest.cpp
+++ b/Engine/Plugins/Runtime/StateTree/Source/StateTreeTestSuite/Private/StateTreeTest.cpp
@@ -596,6 +596,8 @@ struct FStateTreeTest_SubTree_Condition : FStateTreeTestBase
 		const TCHAR* ExitStateStr(TEXT("ExitState"));
 		const TCHAR* StateCompletedStr(TEXT("StateCompleted"));
 
+		GetTestRunner().AddExpectedMessage(TEXT("Evaluation forced to false"), ELogVerbosity::Warning, EAutomationExpectedErrorFlags::Contains, 1);
+
 		// Start and enter state
 		Status = Exec.Start();
 
