[StateTree] Use the temporary instance data of the global evaluator when the temporary frame stops.
#rnx

diff --git a/Engine/Plugins/Runtime/StateTree/Source/StateTreeModule/Private/StateTreeExecutionContext.cpp b/Engine/Plugins/Runtime/StateTree/Source/StateTreeModule/Private/StateTreeExecutionContext.cpp
--- a/Engine/Plugins/Runtime/StateTree/Source/StateTreeModule/Private/StateTreeExecutionContext.cpp
+++ b/Engine/Plugins/Runtime/StateTree/Source/StateTreeModule/Private/StateTreeExecutionContext.cpp
@@ -4145,7 +4145,9 @@ void FStateTreeExecutionContext::StopGlobalsForFrameInternal(const FStateTreeExe
 	for (int32 EvalIndex = EvaluatorsEnd - 1; EvalIndex >= CurrentStateTree->EvaluatorsBegin; --EvalIndex)
 	{
 		const FStateTreeEvaluatorBase& Eval = CurrentStateTree->Nodes[EvalIndex].Get<const FStateTreeEvaluatorBase>();
-		const FStateTreeDataView EvalInstanceView = GetDataView(ParentFrame, Frame, Eval.InstanceDataHandle);
+		const FStateTreeDataView EvalInstanceView = bOnActiveInstances
+			? GetDataView(ParentFrame, Frame, Eval.InstanceDataHandle)
+			: GetDataViewOrTemporary(ParentFrame, Frame, Eval.InstanceDataHandle);
 		FNodeInstanceDataScope DataScope(*this, &Eval, EvalIndex, Eval.InstanceDataHandle, EvalInstanceView);
 
 		STATETREE_LOG(Verbose, TEXT("  Stop: '%s'"), *Eval.Name.ToString());
