[StateTree] Add StateTreeCompilerManager. Move the callback for PIE and BP type from StateTree to StateTreeCompilerManager. This is the first step. The goal is to gather the dependencies and only compile/link when needed. Use the dependencies to sort the compilation.

diff --git a/Engine/Plugins/Runtime/StateTree/Source/StateTreeEditorModule/Private/StateTreeCompilerManager.cpp b/Engine/Plugins/Runtime/StateTree/Source/StateTreeEditorModule/Private/StateTreeCompilerManager.cpp
new file mode 100644
--- /dev/null
+++ b/Engine/Plugins/Runtime/StateTree/Source/StateTreeEditorModule/Private/StateTreeCompilerManager.cpp
@@ -0,0 +1,199 @@
+// Copyright Epic Games, Inc. All Rights Reserved.
+
+#include "StateTreeCompilerManager.h"
+#include "StateTree.h"
+#include "StateTreeCompilerLog.h"
+#include "StateTreeDelegates.h"
+#include "StateTreeEditingSubsystem.h"
+
+#include "Editor.h"
+#include "StructUtilsDelegates.h"
+#include "StructUtils/UserDefinedStruct.h"
+#include "UObject/ObjectKey.h"
+#include "UObject/UObjectGlobals.h"
+#include "UObject/UObjectIterator.h"
+
+namespace UE::StateTree::Compiler::Private
+{
+
+class FCompilerManagerImpl
+{
+public:
+	FCompilerManagerImpl();
+	~FCompilerManagerImpl();
+	FCompilerManagerImpl(const FCompilerManagerImpl&) = delete;
+	FCompilerManagerImpl& operator=(const FCompilerManagerImpl&) = delete;
+
+private:
+	bool HandleCompileStateTree(UStateTree& StateTree);
+
+#if WITH_EDITOR
+	using FReplacementObjectMap = TMap<UObject*, UObject*>;
+	void HandleObjectsReinstanced(const FReplacementObjectMap& ObjectMap);
+	void HandlePreBeginPIE(bool bIsSimulating);
+	void HandleUserDefinedStructReinstanced(const UUserDefinedStruct& UserDefinedStruct);
+#endif
+
+private:
+	FDelegateHandle ObjectsReinstancedHandle;
+	FDelegateHandle UserDefinedStructReinstancedHandle;
+	FDelegateHandle PreBeginPIEHandle;
+};
+static TUniquePtr<FCompilerManagerImpl> CompilerManagerImpl;
+
+FCompilerManagerImpl::FCompilerManagerImpl()
+{
+	UE::StateTree::Delegates::OnRequestCompile.BindRaw(this, &FCompilerManagerImpl::HandleCompileStateTree);
+	ObjectsReinstancedHandle = FCoreUObjectDelegates::OnObjectsReinstanced.AddRaw(this, &FCompilerManagerImpl::HandleObjectsReinstanced);
+	UserDefinedStructReinstancedHandle = UE::StructUtils::Delegates::OnUserDefinedStructReinstanced.AddRaw(this, &FCompilerManagerImpl::HandleUserDefinedStructReinstanced);
+	PreBeginPIEHandle = FEditorDelegates::PreBeginPIE.AddRaw(this, &FCompilerManagerImpl::HandlePreBeginPIE);
+}
+
+FCompilerManagerImpl::~FCompilerManagerImpl()
+{
+	FEditorDelegates::PreBeginPIE.Remove(PreBeginPIEHandle);
+	UE::StructUtils::Delegates::OnUserDefinedStructReinstanced.Remove(UserDefinedStructReinstancedHandle);
+	FCoreUObjectDelegates::OnObjectsReinstanced.Remove(ObjectsReinstancedHandle);
+	UE::StateTree::Delegates::OnRequestCompile.Unbind();
+}
+
+bool FCompilerManagerImpl::HandleCompileStateTree(UStateTree& StateTree)
+{
+	FStateTreeCompilerLog Log;
+	return UStateTreeEditingSubsystem::CompileStateTree(&StateTree, Log);
+}
+
+void FCompilerManagerImpl::HandlePreBeginPIE(const bool bIsSimulating)
+{
+	for (TObjectIterator<UStateTree> It; It; ++It)
+	{
+		check(!It->HasAnyFlags(RF_ClassDefaultObject));
+		It->CompileIfChanged();
+	}
+}
+
+void FCompilerManagerImpl::HandleObjectsReinstanced(const FReplacementObjectMap& ObjectMap)
+{
+	if (ObjectMap.IsEmpty())
+	{
+		return;
+	}
+
+	TArray<const UObject*> ObjectsToBeReplaced;
+	ObjectsToBeReplaced.Reserve(ObjectMap.Num());
+	for (TMap<UObject*, UObject*>::TConstIterator It(ObjectMap); It; ++It)
+	{
+		if (const UObject* ObjectToBeReplaced = It->Value)
+		{
+			ObjectsToBeReplaced.Add(ObjectToBeReplaced);
+		}
+	}
+
+	TSet<const UStruct*> Structs;
+	Structs.Reserve(ObjectsToBeReplaced.Num());
+	for (const UObject* ObjectToBeReplaced : ObjectsToBeReplaced)
+	{
+		// It's a UClass or a UScriptStruct
+		if (const UStruct* StructToReplaced = Cast<const UStruct>(ObjectToBeReplaced))
+		{
+			Structs.Add(StructToReplaced);
+		}
+		else
+		{
+			Structs.Add(ObjectToBeReplaced->GetClass());
+		}
+	}
+
+	for (TObjectIterator<UStateTree> It; It; ++It)
+	{
+		UStateTree* StateTree = *It;
+		check(!StateTree->HasAnyFlags(RF_ClassDefaultObject));
+		bool bShouldRelink = false;
+
+		// Relink if one of the out of date objects got reinstanced.
+		if (StateTree->OutOfDateStructs.Num() > 0)
+		{
+			for (const FObjectKey& OutOfDateObjectKey : StateTree->OutOfDateStructs)
+			{
+				if (const UObject* OutOfDateObject = OutOfDateObjectKey.ResolveObjectPtr())
+				{
+					if (ObjectMap.Contains(OutOfDateObject))
+					{
+						bShouldRelink = true;
+						break;
+					}
+				}
+			}
+		}
+
+		// If the asset is not linked yet (or has failed), no need to link.
+		if (!bShouldRelink && !StateTree->bIsLinked)
+		{
+			continue;
+		}
+
+		// Relink only if the reinstantiated object belongs to this asset,
+		// or anything from the property binding refers to the classes of the reinstantiated object.
+		if (!bShouldRelink)
+		{
+			for (const UObject* ObjectToBeReplaced : ObjectsToBeReplaced)
+			{
+				if (ObjectToBeReplaced->IsInOuter(StateTree))
+				{
+					bShouldRelink = true;
+					break;
+				}
+			}
+		}
+
+		if (!bShouldRelink)
+		{
+			bShouldRelink |= StateTree->PropertyBindings.ContainsAnyStruct(Structs);
+		}
+
+		if (bShouldRelink)
+		{
+			if (!StateTree->Link())
+			{
+				UE_LOG(LogStateTree, Error, TEXT("%s failed to link after Object reinstantiation. Take a look at the asset for any errors. Asset will not be usable at runtime."), *StateTree->GetPathName());
+			}
+		}
+	}
+}
+
+void FCompilerManagerImpl::HandleUserDefinedStructReinstanced(const UUserDefinedStruct& UserDefinedStruct)
+{
+	TSet<const UStruct*> Structs;
+	Structs.Add(&UserDefinedStruct);
+
+	for (TObjectIterator<UStateTree> It; It; ++It)
+	{
+		UStateTree* StateTree = *It;
+		if (StateTree->PropertyBindings.ContainsAnyStruct(Structs))
+		{
+			if (!StateTree->Link())
+			{
+				UE_LOG(LogStateTree, Error, TEXT("%s failed to link after Struct reinstantiation. Take a look at the asset for any errors. Asset will not be usable at runtime."), *StateTree->GetPathName());
+			}
+		}
+	}
+}
+
+} // namespace UE::StateTree::Compiler::Private
+
+
+namespace UE::StateTree::Compiler
+{
+
+void FCompilerManager::Startup()
+{
+	Private::CompilerManagerImpl = MakeUnique<Private::FCompilerManagerImpl>();
+}
+
+void FCompilerManager::Shutdown()
+{
+	Private::CompilerManagerImpl.Reset();
+}
+
+} // UE::StateTree::Compiler
+
diff --git a/Engine/Plugins/Runtime/StateTree/Source/StateTreeEditorModule/Private/StateTreeEditorModule.cpp b/Engine/Plugins/Runtime/StateTree/Source/StateTreeEditorModule/Private/StateTreeEditorModule.cpp
--- a/Engine/Plugins/Runtime/StateTree/Source/StateTreeEditorModule/Private/StateTreeEditorModule.cpp
+++ b/Engine/Plugins/Runtime/StateTree/Source/StateTreeEditorModule/Private/StateTreeEditorModule.cpp
@@ -27,6 +27,7 @@
 #include "PropertyEditorModule.h"
 #include "StateTree.h"
 #include "StateTreeCompilerLog.h"
+#include "StateTreeCompilerManager.h"
 #include "StateTreeDelegates.h"
 #include "StateTreeEditingSubsystem.h"
 #include "StateTreeEditor.h"
@@ -81,8 +82,8 @@ FStateTreeEditorModule* FStateTreeEditorModule::GetModulePtr()
 
 void FStateTreeEditorModule::StartupModule()
 {
-	UE::StateTree::Delegates::OnRequestCompile.BindStatic(&UE::StateTree::Editor::CompileStateTree);
 	UE::StateTree::Delegates::OnRequestEditorHash.BindLambda([](const UStateTree& InStateTree) -> uint32 { return UStateTreeEditingSubsystem::CalculateStateTreeHash(&InStateTree); });
+	UE::StateTree::Compiler::FCompilerManager::Startup();
 
 	OnPostEngineInitHandle = FCoreDelegates::OnPostEngineInit.AddLambda([this]()
 		{
@@ -166,8 +167,8 @@ void FStateTreeEditorModule::ShutdownModule()
 	FCoreDelegates::OnPostEngineInit.Remove(OnPostEngineInitHandle);
 	NodeClassCache.Reset();
 
+	UE::StateTree::Compiler::FCompilerManager::Shutdown();
 	UE::StateTree::Delegates::OnRequestEditorHash.Unbind();
-	UE::StateTree::Delegates::OnRequestCompile.Unbind();
 }
 
 TSharedRef<IStateTreeEditor> FStateTreeEditorModule::CreateStateTreeEditor(const EToolkitMode::Type Mode, const TSharedPtr< class IToolkitHost >& InitToolkitHost, UStateTree* StateTree)
diff --git a/Engine/Plugins/Runtime/StateTree/Source/StateTreeEditorModule/Public/StateTreeCompilerManager.h b/Engine/Plugins/Runtime/StateTree/Source/StateTreeEditorModule/Public/StateTreeCompilerManager.h
new file mode 100644
--- /dev/null
+++ b/Engine/Plugins/Runtime/StateTree/Source/StateTreeEditorModule/Public/StateTreeCompilerManager.h
@@ -0,0 +1,29 @@
+// Copyright Epic Games, Inc. All Rights Reserved.
+
+#pragma once
+
+#include "StateTreeCompiler.h"
+
+#define UE_API STATETREEEDITORMODULE_API
+
+namespace UE::StateTree::Compiler
+{
+
+/**
+ * 
+ */
+class FCompilerManager final
+{
+public:
+	static UE_API void Startup();
+	static UE_API void Shutdown();
+	
+private:
+	FCompilerManager() = delete;
+	FCompilerManager(const FCompilerManager&) = delete;
+	FCompilerManager& operator= (const FCompilerManager&) = delete;
+};
+
+} // UE::StateTree::Compiler
+
+#undef UE_API
diff --git a/Engine/Plugins/Runtime/StateTree/Source/StateTreeModule/Private/StateTree.cpp b/Engine/Plugins/Runtime/StateTree/Source/StateTreeModule/Private/StateTree.cpp
--- a/Engine/Plugins/Runtime/StateTree/Source/StateTreeModule/Private/StateTree.cpp
+++ b/Engine/Plugins/Runtime/StateTree/Source/StateTreeModule/Private/StateTree.cpp
@@ -178,20 +178,11 @@ bool UStateTree::HasCompatibleContextData(TNotNull<const UStateTree*> Other) con
 	return true;
 }
 
-#if WITH_EDITOR || WITH_STATETREE_DEBUG
+#if WITH_STATETREE_DEBUG
 void UStateTree::PostInitProperties()
 {
 	Super::PostInitProperties();
 
-#if WITH_EDITOR
-	if (!HasAnyFlags(RF_ClassDefaultObject))
-	{
-		OnObjectsReinstancedHandle = FStateTreeModule::OnObjectsReinstanced.AddUObject(this, &UStateTree::OnObjectsReinstanced);
-		OnUserDefinedStructReinstancedHandle = FStateTreeModule::OnUserDefinedStructReinstanced.AddUObject(this, &UStateTree::OnUserDefinedStructReinstanced);
-		OnPreBeginPIEHandle = FStateTreeModule::OnPreBeginPIE.AddUObject(this, &UStateTree::OnPreBeginPIE);
-	}
-#endif
-
 #if WITH_STATETREE_DEBUG
 	if (!HasAnyFlags(RF_ClassDefaultObject))
 	{
@@ -203,12 +194,6 @@ void UStateTree::PostInitProperties()
 
 void UStateTree::BeginDestroy()
 {
-#if WITH_EDITOR
-	FStateTreeModule::OnObjectsReinstanced.Remove(OnObjectsReinstancedHandle);
-	FStateTreeModule::OnUserDefinedStructReinstanced.Remove(OnUserDefinedStructReinstancedHandle);
-	FStateTreeModule::OnPreBeginPIE.Remove(OnPreBeginPIEHandle);
-#endif
-
 #if WITH_STATETREE_DEBUG
 	FStateTreeModule::OnPreRuntimeValidationInstanceData.Remove(PreGCHandle);
 	FStateTreeModule::OnPostRuntimeValidationInstanceData.Remove(PostGCHandle);
@@ -216,7 +201,7 @@ void UStateTree::BeginDestroy()
 
 	Super::BeginDestroy();
 }
-#endif //WITH_EDITOR || WITH_STATETREE_DEBUG
+#endif //WITH_STATETREE_DEBUG
 
 #if WITH_EDITOR
 namespace UE::StateTree::Compiler
@@ -284,102 +269,6 @@ void UStateTree::ResetCompiled()
 	}
 }
 
-void UStateTree::OnObjectsReinstanced(const FReplacementObjectMap& ObjectMap)
-{
-	if (ObjectMap.IsEmpty())
-	{
-		return;
-	}
-
-	bool bShouldRelink = false;
-
-	// Relink if one of the out of date objects got reinstanced.
-	if (OutOfDateStructs.Num() > 0)
-	{
-		for (FObjectKey OutOfDateObjectKey : OutOfDateStructs)
-		{
-			if (const UObject* OutOfDateObject = OutOfDateObjectKey.ResolveObjectPtr())
-			{
-				if (ObjectMap.Contains(OutOfDateObject))
-				{
-					bShouldRelink = true;
-					break;
-				}
-			}
-		}
-	}
-
-	// If the asset is not linked yet (or has failed), no need to link.
-	if (!bShouldRelink && !bIsLinked)
-	{
-		return;
-	}
-
-	// Relink only if the reinstantiated object belongs to this asset,
-	// or anything from the property binding refers to the classes of the reinstantiated object.
-	if (!bShouldRelink)
-	{
-		for (TMap<UObject*, UObject*>::TConstIterator It(ObjectMap); It; ++It)
-		{
-			if (const UObject* ObjectToBeReplaced = It->Value)
-			{
-				if (ObjectToBeReplaced->IsInOuter(this))
-				{
-					bShouldRelink = true;
-					break;
-				}
-			}
-		}
-	}
-
-	if (!bShouldRelink)
-	{
-		TSet<const UStruct*> Structs;
-		for (TMap<UObject*, UObject*>::TConstIterator It(ObjectMap); It; ++It)
-		{
-			if (const UObject* ObjectToBeReplaced = It->Value)
-			{
-				// It's a UClass or a UScriptStruct
-				if (const UStruct* StructToReplaced = Cast<const UStruct>(ObjectToBeReplaced))
-				{
-					Structs.Add(StructToReplaced);
-				}
-				else
-				{
-					Structs.Add(ObjectToBeReplaced->GetClass());
-				}
-			}
-		}
-
-		bShouldRelink |= PropertyBindings.ContainsAnyStruct(Structs);
-	}
-
-	if (bShouldRelink)
-	{
-		if (!Link())
-		{
-			UE_LOG(LogStateTree, Error, TEXT("%s failed to link after Object reinstantiation. Take a look at the asset for any errors. Asset will not be usable at runtime."), *GetPathName());
-		}
-	}
-}
-
-void UStateTree::OnUserDefinedStructReinstanced(const UUserDefinedStruct& UserDefinedStruct)
-{
-	// Struct utils handle reinstancing the struct values (instanced struct, property bag, etc).
-	// We will need to update the property binding.
-	
-	TSet<const UStruct*> Structs;
-	Structs.Add(&UserDefinedStruct);
-	
-	if (PropertyBindings.ContainsAnyStruct(Structs))
-	{
-		if (!Link())
-		{
-			UE_LOG(LogStateTree, Error, TEXT("%s failed to link after Struct reinstantiation. Take a look at the asset for any errors. Asset will not be usable at runtime."), *GetPathName());
-		}
-	}
-}
-
 void UStateTree::PostDuplicate(EDuplicateMode::Type DuplicateMode)
 {
 	Super::PostDuplicate(DuplicateMode);
@@ -1298,11 +1187,6 @@ TArray<FStateTreeMemoryUsage> UStateTree::CalculateEstimatedMemoryUsage() const
 	return MemoryUsages;
 }
 
-void UStateTree::OnPreBeginPIE(const bool bIsSimulating)
-{
-	CompileIfChanged();
-}
-
 void UStateTree::CompileIfChanged()
 {
 	if (UE::StateTree::Delegates::OnRequestCompile.IsBound() && UE::StateTree::Delegates::OnRequestEditorHash.IsBound())
@@ -1440,8 +1324,7 @@ FString UStateTree::DebugInternalLayoutAsString() const
 	DebugString.Appendf(TEXT("\nFrames(%d)\n"), Frames.Num());
 	if (Frames.Num())
 	{
-		DebugString.Appendf(TEXT("  [ (Idx) | %-4s ]\n"),
-			TEXT("Root"), TEXT("NumTasks"));
+		DebugString.Appendf(TEXT("  [ (Idx) | %-4s ]\n"), TEXT("Root"));
 		for (int32 FrameIndex = 0; FrameIndex < Frames.Num(); ++FrameIndex)
 		{
 			const FCompactStateTreeFrame& Frame = Frames[FrameIndex];
diff --git a/Engine/Plugins/Runtime/StateTree/Source/StateTreeModule/Private/StateTreeModule.cpp b/Engine/Plugins/Runtime/StateTree/Source/StateTreeModule/Private/StateTreeModule.cpp
--- a/Engine/Plugins/Runtime/StateTree/Source/StateTreeModule/Private/StateTreeModule.cpp
+++ b/Engine/Plugins/Runtime/StateTree/Source/StateTreeModule/Private/StateTreeModule.cpp
@@ -29,22 +29,12 @@
 #include "StateTreeInstanceData.h"
 #endif // WITH_EDITORONLY_DATA
 
-#if WITH_EDITOR
-#include "StructUtilsDelegates.h"
-#endif
-
 #if WITH_STATETREE_DEBUG
 #include "UObject/UObjectGlobals.h"
 #endif
 
 #define LOCTEXT_NAMESPACE "StateTree"
 
-#if WITH_EDITOR
-FStateTreeModule::FOnObjectsReinstanced FStateTreeModule::OnObjectsReinstanced;
-FStateTreeModule::FOnPIEEvent FStateTreeModule::OnPreBeginPIE;
-FStateTreeModule::FOnUserDefinedStructReinstanced FStateTreeModule::OnUserDefinedStructReinstanced;
-#endif
-
 #if WITH_STATETREE_DEBUG
 FTSSimpleMulticastDelegate FStateTreeModule::OnPreRuntimeValidationInstanceData;
 FTSSimpleMulticastDelegate FStateTreeModule::OnPostRuntimeValidationInstanceData;
@@ -128,11 +118,6 @@ void FStateTreeModule::StartupModule()
 #endif // WITH_EDITORONLY_DATA
 
 #if WITH_EDITOR
-	// Register thread safe delegates, which allows the StateTree objects to safely register to these delegates e.g. in PostInitProperties() which may be called from another thread.
-	OnObjectsReinstancedHandle = FCoreUObjectDelegates::OnObjectsReinstanced.AddRaw(this, &FStateTreeModule::HandleObjectsReinstanced);
-	OnUserDefinedStructReinstancedHandle = UE::StructUtils::Delegates::OnUserDefinedStructReinstanced.AddRaw(this, &FStateTreeModule::HandleUserDefinedStructReinstanced);
-	OnPreBeginPIEHandle = FEditorDelegates::PreBeginPIE.AddRaw(this, &FStateTreeModule::HandlePreBeginPIE);
-
 	 UE::PropertyBinding::PropertyBindingIndex16ConversionFuncList.Add([](const FPropertyTag& Tag, FStructuredArchive::FSlot Slot, TNotNull<FPropertyBindingIndex16*> Index)
 	 {
 		static FName FStateTreeIndex16TypeName = FStateTreeIndex16::StaticStruct()->GetFName();
@@ -161,12 +146,6 @@ void FStateTreeModule::ShutdownModule()
 	FCoreUObjectDelegates::GetPostPurgeGarbageDelegate().Remove(PostGCHandle);
 #endif
 
-#if WITH_EDITOR
-	FCoreUObjectDelegates::OnObjectsReinstanced.Remove(OnObjectsReinstancedHandle);
-	UE::StructUtils::Delegates::OnUserDefinedStructReinstanced.Remove(OnUserDefinedStructReinstancedHandle);
-	FEditorDelegates::PreBeginPIE.Remove(OnPreBeginPIEHandle);
-#endif //WITH_EDITOR
-
 #if WITH_STATETREE_TRACE
 	StopTraces();
 
@@ -311,23 +290,6 @@ void FStateTreeModule::StopTraces()
 #endif // WITH_STATETREE_TRACE
 }
 
-#if WITH_EDITOR
-void FStateTreeModule::HandleObjectsReinstanced(const FReplacementObjectMap& ObjectMap)
-{
-	OnObjectsReinstanced.Broadcast(ObjectMap);
-}
-
-void FStateTreeModule::HandlePreBeginPIE(const bool bIsSimulating)
-{
-	OnPreBeginPIE.Broadcast(bIsSimulating);
-}
-
-void FStateTreeModule::HandleUserDefinedStructReinstanced(const UUserDefinedStruct& UserDefinedStruct)
-{
-	OnUserDefinedStructReinstanced.Broadcast(UserDefinedStruct);
-}
-#endif
-
 #if WITH_STATETREE_DEBUG
 void FStateTreeModule::HandlePreGC()
 {
diff --git a/Engine/Plugins/Runtime/StateTree/Source/StateTreeModule/Private/StateTreeModuleImpl.h b/Engine/Plugins/Runtime/StateTree/Source/StateTreeModule/Private/StateTreeModuleImpl.h
--- a/Engine/Plugins/Runtime/StateTree/Source/StateTreeModule/Private/StateTreeModuleImpl.h
+++ b/Engine/Plugins/Runtime/StateTree/Source/StateTreeModule/Private/StateTreeModuleImpl.h
@@ -31,29 +31,12 @@ public:
 	virtual bool IsTracing() const override;
 	virtual void StopTraces() override;
 
-#if WITH_EDITOR
-	using FReplacementObjectMap = FCoreUObjectDelegates::FReplacementObjectMap;
-	DECLARE_TS_MULTICAST_DELEGATE_OneParam(FOnObjectsReinstanced, const FReplacementObjectMap&);
-	static FOnObjectsReinstanced OnObjectsReinstanced;
-
-	DECLARE_TS_MULTICAST_DELEGATE_OneParam(FOnPIEEvent, bool);
-	static FOnPIEEvent OnPreBeginPIE;
-
-	DECLARE_TS_MULTICAST_DELEGATE_OneParam(FOnUserDefinedStructReinstanced, const UUserDefinedStruct& /*UserDefinedStruct*/);
-	static FOnUserDefinedStructReinstanced OnUserDefinedStructReinstanced;
-#endif
-
 #if WITH_STATETREE_DEBUG
 	static FTSSimpleMulticastDelegate OnPreRuntimeValidationInstanceData;
 	static FTSSimpleMulticastDelegate OnPostRuntimeValidationInstanceData;
 #endif
 
 private:
-#if WITH_EDITOR
-	void HandleObjectsReinstanced(const FReplacementObjectMap& ObjectMap);
-	void HandlePreBeginPIE(const bool bIsSimulating);
-	void HandleUserDefinedStructReinstanced(const UUserDefinedStruct& UserDefinedStruct);
-#endif
 
 #if WITH_STATETREE_DEBUG
 	void HandlePreGC();
@@ -86,12 +69,6 @@ private:
 	FAutoConsoleCommand StopDebuggerTracesCommand;
 #endif // WITH_STATETREE_TRACE
 
-#if WITH_EDITOR
-	FDelegateHandle OnObjectsReinstancedHandle;
-	FDelegateHandle OnUserDefinedStructReinstancedHandle;
-	FDelegateHandle OnPreBeginPIEHandle;
-#endif
-
 #if WITH_STATETREE_DEBUG
 	FDelegateHandle PreGCHandle;
 	FDelegateHandle PostGCHandle;
diff --git a/Engine/Plugins/Runtime/StateTree/Source/StateTreeModule/Public/StateTree.h b/Engine/Plugins/Runtime/StateTree/Source/StateTreeModule/Public/StateTree.h
--- a/Engine/Plugins/Runtime/StateTree/Source/StateTreeModule/Public/StateTree.h
+++ b/Engine/Plugins/Runtime/StateTree/Source/StateTreeModule/Public/StateTree.h
@@ -76,6 +76,11 @@ private:
 
 
 #if WITH_EDITOR
+namespace UE::StateTree::Compiler::Private
+{
+	class FCompilerManagerImpl;
+}
+
 /** Struct containing information about the StateTree runtime memory usage. */
 struct FStateTreeMemoryUsage
 {
@@ -259,8 +264,11 @@ public:
 	/** Calculates runtime memory usage for different sections of the tree. */
 	UE_API TArray<FStateTreeMemoryUsage> CalculateEstimatedMemoryUsage() const;
 
+	UE_DEPRECATED(5.7, "Use the compiler manager.")
 	/** Called when the editor is preparing to start a pie session. */
-	UE_API void OnPreBeginPIE(const bool bIsSimulating);
+	UE_API void OnPreBeginPIE(const bool bIsSimulating)
+	{
+	}
 
 	/** Compile the state trees if the editor hash data as changed since the last compilation. */
 	UE_API void CompileIfChanged();
@@ -277,8 +285,11 @@ public:
 	UPROPERTY()
 	TObjectPtr<UObject> EditorData;
 
+	UE_DEPRECATED(5.7, "Use the compiler manager.")
 	FDelegateHandle OnObjectsReinstancedHandle;
+	UE_DEPRECATED(5.7, "Use the compiler manager.")
 	FDelegateHandle OnUserDefinedStructReinstancedHandle;
+	UE_DEPRECATED(5.7, "Use the compiler manager.")
 	FDelegateHandle OnPreBeginPIEHandle;
 #endif
 
@@ -300,15 +311,12 @@ protected:
 #endif
 	UE_API virtual void Serialize(FStructuredArchiveRecord Record) override;
 
-#if WITH_EDITOR || WITH_STATETREE_DEBUG
+#if WITH_STATETREE_DEBUG
 	UE_API virtual void BeginDestroy() override;
 	UE_API virtual void PostInitProperties() override;
 #endif
 
 #if WITH_EDITOR
-	using FReplacementObjectMap = TMap<UObject*, UObject*>;
-	UE_API void OnObjectsReinstanced(const FReplacementObjectMap& ObjectMap);
-	UE_API void OnUserDefinedStructReinstanced(const UUserDefinedStruct& UserDefinedStruct);
 	UE_API virtual void PostDuplicate(EDuplicateMode::Type DuplicateMode) override;
 	UE_API virtual void GetAssetRegistryTags(FAssetRegistryTagsContext Context) const override;
 protected:
@@ -515,8 +523,9 @@ private:
 	friend struct FStateTreeWeakExecutionContext;
 	friend struct TStateTreeStrongExecutionContext<true>;
 
-#if WITH_EDITORONLY_DATA
+#if WITH_EDITOR
 	friend struct FStateTreeCompiler;
+	friend class UE::StateTree::Compiler::Private::FCompilerManagerImpl;
 #endif
 };
 
