[StateTree] Prevent refresh from occurring when the detail customization is pending deletion. The EditorDataDetails is an EditorUndoClient. Until the shared ptr is destroyed, it will receive the event. The number of EditorDataDetails doubles every time an undo is requested. If the user repeatedly spams the undo command, it can lead to a long hitch.

diff --git a/Engine/Plugins/Runtime/StateTree/Source/StateTreeEditorModule/Private/Customizations/StateTreeEditorDataDetails.cpp b/Engine/Plugins/Runtime/StateTree/Source/StateTreeEditorModule/Private/Customizations/StateTreeEditorDataDetails.cpp
--- a/Engine/Plugins/Runtime/StateTree/Source/StateTreeEditorModule/Private/Customizations/StateTreeEditorDataDetails.cpp
+++ b/Engine/Plugins/Runtime/StateTree/Source/StateTreeEditorModule/Private/Customizations/StateTreeEditorDataDetails.cpp
@@ -255,6 +255,12 @@ void FStateTreeEditorDataDetails::CustomizeDetails(IDetailLayoutBuilder& DetailB
 		}));
 }
 
+void FStateTreeEditorDataDetails::PendingDelete()
+{
+	Super::PendingDelete();
+	bPendingDelete = true;
+}
+
 void FStateTreeEditorDataDetails::MakeArrayCategory(IDetailLayoutBuilder& DetailBuilder, FName CategoryName, const FText& DisplayName, int32 SortOrder, TSharedPtr<IPropertyHandle> PropertyHandle)
 {
 	IDetailCategoryBuilder& Category = DetailBuilder.EditCategory(CategoryName, DisplayName);
@@ -282,6 +288,11 @@ void FStateTreeEditorDataDetails::MakeArrayCategory(IDetailLayoutBuilder& Detail
 
 void FStateTreeEditorDataDetails::PostUndo(bool bSuccess)
 {
+	if (bPendingDelete)
+	{
+		return;
+	}
+
 	// Refresh view on undo or redo so that the customization based on e.g. Global tasks bindings, parameters, etc. will be reflected correctly.
 	if (const TSharedPtr<IPropertyUtilities> PropertyUtilities = WeakPropertyUtilities.Pin())
 	{
@@ -291,6 +302,11 @@ void FStateTreeEditorDataDetails::PostUndo(bool bSuccess)
 
 void FStateTreeEditorDataDetails::PostRedo(bool bSuccess)
 {
+	if (bPendingDelete)
+	{
+		return;
+	}
+
 	// Refresh view on undo or redo so that the customization based on e.g. Global tasks bindings, parameters, etc. will be reflected correctly.
 	if (const TSharedPtr<IPropertyUtilities> PropertyUtilities = WeakPropertyUtilities.Pin())
 	{
diff --git a/Engine/Plugins/Runtime/StateTree/Source/StateTreeEditorModule/Private/Customizations/StateTreeEditorDataDetails.h b/Engine/Plugins/Runtime/StateTree/Source/StateTreeEditorModule/Private/Customizations/StateTreeEditorDataDetails.h
--- a/Engine/Plugins/Runtime/StateTree/Source/StateTreeEditorModule/Private/Customizations/StateTreeEditorDataDetails.h
+++ b/Engine/Plugins/Runtime/StateTree/Source/StateTreeEditorModule/Private/Customizations/StateTreeEditorDataDetails.h
@@ -12,12 +12,16 @@ class IPropertyUtilities;
 
 class FStateTreeEditorDataDetails : public IDetailCustomization, FSelfRegisteringEditorUndoClient
 {
+private:
+	using Super = IDetailCustomization;
+
 public:
 	/** Makes a new instance of this detail layout class for a specific detail view requesting it */
 	static TSharedRef<IDetailCustomization> MakeInstance();
 
-	/** IDetailCustomization interface */
+	//~ IDetailCustomization interface
 	virtual void CustomizeDetails(IDetailLayoutBuilder& DetailBuilder) override;
+	virtual void PendingDelete() override;
 
 protected:
 	void MakeArrayCategory(IDetailLayoutBuilder& DetailBuilder, FName CategoryName, const FText& DisplayName, int32 SortOrder, TSharedPtr<IPropertyHandle> PropertyHandle);
@@ -26,5 +30,7 @@ protected:
 	virtual void PostUndo(bool bSuccess) override;
 	virtual void PostRedo(bool bSuccess) override;
 
+	bool bPendingDelete = false;
+
 	TWeakPtr<IPropertyUtilities> WeakPropertyUtilities;
 };
