[StatetreeDebugger] refactored events treeview to be reusable in the details view of the RewindDebugger StateTree tracks

diff --git a/Engine/Plugins/Runtime/StateTree/Source/StateTreeDeveloper/Internal/Debugger/SStateTreeFrameEventsView.h b/Engine/Plugins/Runtime/StateTree/Source/StateTreeDeveloper/Internal/Debugger/SStateTreeFrameEventsView.h
new file mode 100644
--- /dev/null
+++ b/Engine/Plugins/Runtime/StateTree/Source/StateTreeDeveloper/Internal/Debugger/SStateTreeFrameEventsView.h
@@ -0,0 +1,58 @@
+// Copyright Epic Games, Inc. All Rights Reserved.
+
+#pragma once
+
+#if WITH_STATETREE_TRACE_DEBUGGER
+
+#include "Debugger/StateTreeTraceTypes.h"
+#include "Widgets/SCompoundWidget.h"
+
+#define UE_API STATETREEDEVELOPER_API
+
+class UStateTree;
+
+template <typename ItemType> class STreeView;
+
+namespace UE::StateTreeDebugger
+{
+struct FFrameEventTreeElement;
+struct FInstanceEventCollection;
+struct FScrubState;
+
+/**
+ * TreeView representing all traced events on a StateTree instance at a given frame
+ */
+class SFrameEventsView : public SCompoundWidget
+{
+public:
+	SLATE_BEGIN_ARGS(SFrameEventsView) {}
+	SLATE_END_ARGS()
+
+	UE_API void Construct(const FArguments& InArgs, TNotNull<const UStateTree*> InStateTree);
+
+	/** Selects an element in the list based on a predicate applied on the currently displayed events */
+	UE_API void SelectByPredicate(TFunctionRef<bool(const FStateTreeTraceEventVariantType& Event)> InPredicate);
+
+	/** Rebuilds the view from events for a given frame */
+	UE_API void RequestRefresh(const FScrubState& InScrubState);
+
+private:
+	/** Recursively sets tree items as expanded. */
+	void ExpandAll(const TArray<TSharedPtr<FFrameEventTreeElement>>& Items);
+
+	static void GenerateElementsForProperties(const FStateTreeTraceEventVariantType& Event, const TSharedRef<FFrameEventTreeElement>& ParentElement);
+
+	TWeakObjectPtr<const UStateTree> WeakStateTree;
+
+	/** All trace events received for a given instance. */
+	TArray<TSharedPtr<FFrameEventTreeElement>> EventsTreeElements;
+
+	/** Tree view displaying the frame events of the instance associated to the selected track. */
+	TSharedPtr<STreeView<TSharedPtr<FFrameEventTreeElement>>> EventsTreeView;
+};
+
+} // UE::StateTreeDebugger
+
+#undef UE_API
+
+#endif // WITH_STATETREE_TRACE_DEBUGGER
\ No newline at end of file
diff --git a/Engine/Plugins/Runtime/StateTree/Source/StateTreeDeveloper/Private/Debugger/SStateTreeDebuggerViewRow.cpp b/Engine/Plugins/Runtime/StateTree/Source/StateTreeDeveloper/Private/Debugger/SStateTreeDebuggerViewRow.cpp
new file mode 100644
--- /dev/null
+++ b/Engine/Plugins/Runtime/StateTree/Source/StateTreeDeveloper/Private/Debugger/SStateTreeDebuggerViewRow.cpp
@@ -0,0 +1,246 @@
+// Copyright Epic Games, Inc. All Rights Reserved.
+
+#if WITH_STATETREE_TRACE_DEBUGGER
+
+#include "SStateTreeDebuggerViewRow.h"
+#include "StateTreeStyle.h"
+#include "StateTree.h"
+#include "Widgets/Images/SImage.h"
+#include "Widgets/Text/STextBlock.h"
+
+namespace UE::StateTreeDebugger
+{
+
+void SFrameEventViewRow::Construct(const FArguments& InArgs,
+		const TSharedPtr<STableViewBase>& InOwnerTableView,
+		const TSharedPtr<FFrameEventTreeElement>& InElement)
+{
+	Item = InElement;
+	STableRow::Construct(InArgs, InOwnerTableView.ToSharedRef());
+
+	const TSharedPtr<SHorizontalBox> HorizontalBox = SNew(SHorizontalBox);
+	HorizontalBox->SetToolTipText(GetEventTooltip());
+
+	HorizontalBox->AddSlot()
+		.VAlign(VAlign_Center)
+		.HAlign(HAlign_Left)
+		.AutoWidth()
+		[
+			SNew(SExpanderArrow, SharedThis(this))
+			.ShouldDrawWires(false)
+			.IndentAmount(32)
+			.BaseIndentLevel(0)
+		];
+
+	const TSharedPtr<SWidget> EventImage = CreateImageForEvent();
+	if (EventImage.IsValid())
+	{
+		HorizontalBox->AddSlot()
+		.VAlign(VAlign_Center)
+		.HAlign(HAlign_Left)
+		.AutoWidth()
+		[
+			EventImage.ToSharedRef()
+		];
+	}
+
+	HorizontalBox->AddSlot()
+		.Padding(2)
+		.AutoWidth()
+		.VAlign(VAlign_Center)
+		.HAlign(HAlign_Left)
+		[
+			SNew(STextBlock)
+			.TextStyle(&GetEventTextStyle())
+			.Text(GetEventDescription())
+		];
+
+	this->ChildSlot
+		.HAlign(HAlign_Fill)
+		[
+			HorizontalBox.ToSharedRef()
+		];
+}
+
+TSharedPtr<SWidget> SFrameEventViewRow::CreateImageForEvent() const
+{
+	const FStateTreeStyle& StyleSet = FStateTreeStyle::Get();
+
+	// Phase events
+	if (const FStateTreeTracePhaseEvent* PhaseEvent = Item->Event.TryGet<FStateTreeTracePhaseEvent>())
+	{
+		const FSlateBrush* Image = nullptr;
+		switch (PhaseEvent->Phase)
+		{
+		case EStateTreeUpdatePhase::EnterStates:	Image = StyleSet.GetBrush("StateTreeEditor.Debugger.State.Enter");		break;
+		case EStateTreeUpdatePhase::ExitStates:		Image = StyleSet.GetBrush("StateTreeEditor.Debugger.State.Exit");		break;
+		case EStateTreeUpdatePhase::StateCompleted:	Image = StyleSet.GetBrush("StateTreeEditor.Debugger.State.Completed");	break;
+		default:
+			return nullptr;
+		}
+
+		return SNew(SImage).Image(Image);
+	}
+
+	// Log events
+	if (const FStateTreeTraceLogEvent* LogEvent = Item->Event.TryGet<FStateTreeTraceLogEvent>())
+	{
+		const FSlateBrush* Image = nullptr;
+		switch (LogEvent->Verbosity)
+		{
+		case ELogVerbosity::Fatal:
+		case ELogVerbosity::Error:		Image = StyleSet.GetBrush("StateTreeEditor.Debugger.Log.Error");	break;
+		case ELogVerbosity::Warning:	Image = StyleSet.GetBrush("StateTreeEditor.Debugger.Log.Warning");	break;
+		default:
+			return nullptr;
+		}
+
+		return SNew(SImage).Image(Image);
+	}
+
+	// State events
+	if (const FStateTreeTraceStateEvent* StateEvent = Item->Event.TryGet<FStateTreeTraceStateEvent>())
+	{
+		const FSlateBrush* Image = nullptr;
+
+		switch (StateEvent->EventType)
+		{
+		case EStateTreeTraceEventType::OnStateSelected:		Image = StyleSet.GetBrush("StateTreeEditor.Debugger.State.Selected");	break;
+		default:
+			return nullptr;
+		}
+
+		return SNew(SImage).Image(Image);
+	}
+
+	// Task events
+	if (const FStateTreeTraceTaskEvent* TaskEvent = Item->Event.TryGet<FStateTreeTraceTaskEvent>())
+	{
+		const FSlateBrush* Image = nullptr;
+
+		switch (TaskEvent->EventType)
+		{
+		case EStateTreeTraceEventType::OnEntered:
+		case EStateTreeTraceEventType::OnTaskCompleted:
+		case EStateTreeTraceEventType::OnTicked:
+			switch (TaskEvent->Status)
+			{
+			case EStateTreeRunStatus::Failed:		Image = StyleSet.GetBrush("StateTreeEditor.Debugger.Task.Failed");		break;
+			case EStateTreeRunStatus::Succeeded:	Image = StyleSet.GetBrush("StateTreeEditor.Debugger.Task.Succeeded");	break;
+			case EStateTreeRunStatus::Stopped:		Image = StyleSet.GetBrush("StateTreeEditor.Debugger.Task.Stopped");		break;
+			default:
+				return nullptr;
+			}
+			break;
+		default:
+			return nullptr;
+		}
+
+		return SNew(SImage).Image(Image);
+	}
+
+	// Condition events
+	if (Item->Event.IsType<FStateTreeTraceConditionEvent>())
+	{
+		const FStateTreeTraceConditionEvent& ConditionEvent = Item->Event.Get<FStateTreeTraceConditionEvent>();
+		const FSlateBrush* Image = nullptr;
+
+		switch (ConditionEvent.EventType)
+		{
+		case EStateTreeTraceEventType::Passed:					Image = StyleSet.GetBrush("StateTreeEditor.Debugger.Condition.Passed");			break;
+		case EStateTreeTraceEventType::ForcedSuccess:			Image = StyleSet.GetBrush("StateTreeEditor.Debugger.Condition.Passed");			break;
+		case EStateTreeTraceEventType::Failed:					Image = StyleSet.GetBrush("StateTreeEditor.Debugger.Condition.Failed");			break;
+		case EStateTreeTraceEventType::ForcedFailure:			Image = StyleSet.GetBrush("StateTreeEditor.Debugger.Condition.Failed");			break;
+		case EStateTreeTraceEventType::InternalForcedFailure:	Image = StyleSet.GetBrush("StateTreeEditor.Debugger.Condition.Failed");			break;
+		case EStateTreeTraceEventType::OnEvaluating:			Image = StyleSet.GetBrush("StateTreeEditor.Debugger.Condition.OnEvaluating");	break;
+		case EStateTreeTraceEventType::OnTransition:			Image = StyleSet.GetBrush("StateTreeEditor.Debugger.Condition.OnTransition");	break;
+		default:
+			Image = StyleSet.GetBrush("StateTreeEditor.Debugger.Unset");
+		}
+
+		return SNew(SImage).Image(Image);
+	}
+
+	return nullptr;
+}
+
+const FTextBlockStyle& SFrameEventViewRow::GetEventTextStyle() const
+{
+	const FStateTreeStyle& StyleSet = FStateTreeStyle::Get();
+
+	if (Item->Event.IsType<FStateTreeTracePhaseEvent>())
+	{
+		return StyleSet.GetWidgetStyle<FTextBlockStyle>("StateTreeDebugger.Element.Bold");
+	}
+
+	if (Item->Event.IsType<FStateTreeTracePropertyEvent>())
+	{
+		return StyleSet.GetWidgetStyle<FTextBlockStyle>("StateTreeDebugger.Element.Subdued");
+	}
+
+	if (const FStateTreeTraceLogEvent* LogEvent = Item->Event.TryGet<FStateTreeTraceLogEvent>())
+	{
+		// Make verbose logs more subtle
+		if (LogEvent->Verbosity >= ELogVerbosity::Verbose)
+		{
+			return StyleSet.GetWidgetStyle<FTextBlockStyle>("StateTreeDebugger.Element.Subdued");
+		}
+	}
+
+	return StyleSet.GetWidgetStyle<FTextBlockStyle>("StateTreeDebugger.Element.Normal");
+}
+
+FText SFrameEventViewRow::GetEventDescription() const
+{
+	FString EventDescription;
+	if (Item->Description.IsEmpty())
+	{
+		if (const UStateTree* StateTree = Item->WeakStateTree.Get())
+		{
+			// Some types have some custom representations so we want to use a more minimal description.
+			if (Item->Event.IsType<FStateTreeTraceStateEvent>()
+				|| Item->Event.IsType<FStateTreeTraceTaskEvent>()
+				|| Item->Event.IsType<FStateTreeTraceEvaluatorEvent>()
+				|| Item->Event.IsType<FStateTreeTracePropertyEvent>()
+				|| Item->Event.IsType<FStateTreeTraceConditionEvent>()
+				|| Item->Event.IsType<FStateTreeTraceLogEvent>())
+			{
+				Visit([&EventDescription, StateTree](auto& TypedEvent)
+					{
+						EventDescription = TypedEvent.GetValueString(*StateTree);
+					}, Item->Event);
+			}
+			else
+			{
+				Visit([&EventDescription, StateTree](auto& TypedEvent)
+					{
+						EventDescription = TypedEvent.ToFullString(*StateTree);
+					}, Item->Event);
+			}
+		}
+	}
+	else
+	{
+		EventDescription = Item->Description;
+	}
+	
+	return FText::FromString(EventDescription);
+}
+
+FText SFrameEventViewRow::GetEventTooltip() const
+{
+	FString Tooltip;
+	if (const UStateTree* StateTree = Item->WeakStateTree.Get())
+	{
+		Visit([&Tooltip, StateTree](auto& TypedEvent)
+			{
+				Tooltip = TypedEvent.ToFullString(*StateTree);
+			}, Item->Event);
+	}
+
+	return FText::FromString(Tooltip);
+}
+
+} // UE::StateTreeDebugger
+
+#endif // WITH_STATETREE_TRACE_DEBUGGER
diff --git a/Engine/Plugins/Runtime/StateTree/Source/StateTreeDeveloper/Private/Debugger/SStateTreeDebuggerViewRow.h b/Engine/Plugins/Runtime/StateTree/Source/StateTreeDeveloper/Private/Debugger/SStateTreeDebuggerViewRow.h
new file mode 100644
--- /dev/null
+++ b/Engine/Plugins/Runtime/StateTree/Source/StateTreeDeveloper/Private/Debugger/SStateTreeDebuggerViewRow.h
@@ -0,0 +1,56 @@
+// Copyright Epic Games, Inc. All Rights Reserved.
+
+#pragma once
+
+#if WITH_STATETREE_TRACE_DEBUGGER
+
+#include "Debugger/StateTreeTraceTypes.h"
+#include "StateTree.h"
+#include "Templates/SharedPointer.h"
+#include "TraceServices/Model/Frames.h"
+#include "Widgets/Views/SListView.h"
+#include "Widgets/Views/STableRow.h"
+
+class FStateTreeViewModel;
+
+namespace UE::StateTreeDebugger
+{
+
+/** An item in the StateTreeDebugger trace event tree */
+struct FFrameEventTreeElement : TSharedFromThis<FFrameEventTreeElement>
+{
+	explicit FFrameEventTreeElement(const TraceServices::FFrame& Frame, const FStateTreeTraceEventVariantType& Event, const UStateTree* StateTree)
+		: Frame(Frame), Event(Event), WeakStateTree(StateTree)
+	{
+	}
+
+	TraceServices::FFrame Frame;
+	FStateTreeTraceEventVariantType Event;
+	TArray<TSharedPtr<FFrameEventTreeElement>> Children;
+	FString Description;
+	TWeakObjectPtr<const UStateTree> WeakStateTree;
+};
+
+
+/**
+ * Widget for row inside the StateTreeDebugger TreeView.
+ */
+class SFrameEventViewRow : public STableRow<TSharedPtr<FFrameEventTreeElement>>
+{
+public:
+	void Construct(const FArguments& InArgs,
+		const TSharedPtr<STableViewBase>& InOwnerTableView,
+		const TSharedPtr<FFrameEventTreeElement>& InElement);
+
+private:
+	TSharedPtr<SWidget> CreateImageForEvent() const;
+	const FTextBlockStyle& GetEventTextStyle() const;
+	FText GetEventDescription() const;
+	FText GetEventTooltip() const;
+
+	TSharedPtr<FFrameEventTreeElement> Item;
+};
+
+} // UE::StateTreeDebugger
+
+#endif // WITH_STATETREE_TRACE_DEBUGGER
\ No newline at end of file
diff --git a/Engine/Plugins/Runtime/StateTree/Source/StateTreeDeveloper/Private/Debugger/SStateTreeFrameEventsView.cpp b/Engine/Plugins/Runtime/StateTree/Source/StateTreeDeveloper/Private/Debugger/SStateTreeFrameEventsView.cpp
new file mode 100644
--- /dev/null
+++ b/Engine/Plugins/Runtime/StateTree/Source/StateTreeDeveloper/Private/Debugger/SStateTreeFrameEventsView.cpp
@@ -0,0 +1,420 @@
+// Copyright Epic Games, Inc. All Rights Reserved.
+
+#if WITH_STATETREE_TRACE_DEBUGGER
+
+#include "Debugger/SStateTreeFrameEventsView.h"
+#include "Debugger/StateTreeDebuggerTypes.h"
+#include "SStateTreeDebuggerViewRow.h"
+#include "UObject/Package.h"
+#include "Widgets/Layout/SScrollBar.h"
+#include "Widgets/Layout/SScrollBox.h"
+#include "Widgets/Views/STableViewBase.h"
+#include "Widgets/Views/STreeView.h"
+
+namespace UE::StateTreeDebugger
+{
+/**
+ * Iterates over all tree elements for the frame events
+ * @param Elements Container of hierarchical tree element to visit
+ * @param InFunc function called at each element, should return true if visiting is continued or false to stop.
+ */
+void VisitEventTreeElements(const TConstArrayView<TSharedPtr<FFrameEventTreeElement>> Elements
+	, TFunctionRef<bool(TSharedPtr<FFrameEventTreeElement>& VisitedElement)> InFunc)
+{
+	TArray<TSharedPtr<FFrameEventTreeElement>> Stack;
+	bool bContinue = true;
+
+	for (const TSharedPtr<FFrameEventTreeElement>& RootElement : Elements)
+	{
+		if (RootElement == nullptr)
+		{
+			continue;
+		}
+
+		Stack.Add(RootElement);
+
+		while (!Stack.IsEmpty() && bContinue)
+		{
+			TSharedPtr<FFrameEventTreeElement> StackedElement = Stack[0];
+			check(StackedElement);
+
+			Stack.RemoveAt(0);
+
+			bContinue = InFunc(StackedElement);
+
+			if (bContinue)
+			{
+				for (const TSharedPtr<FFrameEventTreeElement>& Child : StackedElement->Children)
+				{
+					if (Child.IsValid())
+					{
+						Stack.Add(Child);
+					}
+				}
+			}
+		}
+
+		if (!bContinue)
+		{
+			break;
+		}
+	}
+}
+
+
+//----------------------------------------------------------------------//
+// SFrameEventsView
+//----------------------------------------------------------------------//
+
+void SFrameEventsView::Construct(const FArguments& InArgs, TNotNull<const UStateTree*> InStateTree)
+{
+	WeakStateTree = InStateTree;
+
+	// EventsTreeView scrollbars
+	TSharedRef<SScrollBar> HorizontalScrollBar = SNew(SScrollBar)
+		.Orientation(Orient_Horizontal)
+		.Thickness(FVector2D(12.0f, 12.0f));
+
+	TSharedRef<SScrollBar> VerticalScrollBar = SNew(SScrollBar)
+		.Orientation(Orient_Vertical)
+		.Thickness(FVector2D(12.0f, 12.0f));
+
+	// EventsTreeView
+	EventsTreeView = SNew(STreeView<TSharedPtr<FFrameEventTreeElement>>)
+			.OnGenerateRow_Lambda([this](const TSharedPtr<FFrameEventTreeElement>& InElement, const TSharedRef<STableViewBase>& InOwnerTableView)
+			{
+				return SNew(SFrameEventViewRow, InOwnerTableView, InElement);
+			})
+			.OnGetChildren_Lambda([](const TSharedPtr<const FFrameEventTreeElement>& InParent, TArray<TSharedPtr<FFrameEventTreeElement>>& OutChildren)
+			{
+				if (const FFrameEventTreeElement* Parent = InParent.Get())
+				{
+					OutChildren.Append(Parent->Children);
+				}
+			})
+		.TreeItemsSource(&EventsTreeElements)
+		.AllowOverscroll(EAllowOverscroll::No)
+		.ExternalScrollbar(VerticalScrollBar);
+
+	ChildSlot
+	[
+		SNew(SBorder)
+		.Padding(4.0f)
+		.BorderImage(FAppStyle::GetBrush("ToolPanel.GroupBorder"))
+		[
+			SNew(SVerticalBox)
+			+ SVerticalBox::Slot()
+			[
+				SNew(SSplitter)
+				.Orientation(Orient_Horizontal)
+				+ SSplitter::Slot()
+				[
+					SNew(SVerticalBox)
+					+ SVerticalBox::Slot()
+					.FillHeight(1.0f)
+					[
+						SNew(SHorizontalBox)
+						+ SHorizontalBox::Slot()
+						.FillWidth(1.0f)
+						.Padding(0.0f)
+						[
+							SNew(SScrollBox)
+							.Orientation(Orient_Horizontal)
+							.ExternalScrollbar(HorizontalScrollBar)
+							+ SScrollBox::Slot()
+							.FillSize(1.0f)
+							[
+								EventsTreeView.ToSharedRef()
+							]
+						]
+
+						+ SHorizontalBox::Slot()
+						.AutoWidth()
+						[
+							VerticalScrollBar
+						]
+					]
+					+ SVerticalBox::Slot()
+					.AutoHeight()
+					[
+						HorizontalScrollBar
+					]
+				]
+			]
+		]
+	];
+}
+
+void SFrameEventsView::SelectByPredicate(TFunctionRef<bool(const FStateTreeTraceEventVariantType& Event)> InPredicate)
+{
+	TSharedPtr<FFrameEventTreeElement> MatchingElement = nullptr;
+	VisitEventTreeElements(EventsTreeElements, [&MatchingElement, InPredicate](const TSharedPtr<FFrameEventTreeElement>& VisitedElement)
+		{
+			if (InPredicate(VisitedElement->Event))
+			{
+				MatchingElement = VisitedElement;
+			}
+
+			// Continue visit until we find a matching event
+			return MatchingElement.IsValid() == false;
+		});
+
+	if (MatchingElement.IsValid())
+	{
+		EventsTreeView->SetSelection(MatchingElement);
+	}
+}
+
+void SFrameEventsView::GenerateElementsForProperties(const FStateTreeTraceEventVariantType& Event, const TSharedRef<FFrameEventTreeElement>& ParentElement)
+{
+	FString TypePath;
+	FString InstanceDataAsString;
+	FString DebugText;
+
+	Visit([&TypePath, &InstanceDataAsString, &DebugText](auto& TypedEvent)
+		{
+			TypePath = TypedEvent.GetDataTypePath();
+			InstanceDataAsString = TypedEvent.GetDataAsText();
+			DebugText = TypedEvent.GetDebugText();
+		}, Event);
+
+	auto CreateTreeElement = [ParentElement]<typename EventType>(const FString& Line)
+		{
+			// Create new event
+			EventType Event(/*RecordingWorldTime*/0, ELogVerbosity::Verbose, *Line);
+
+			// Create Tree element to hold the event
+			const TSharedPtr<FFrameEventTreeElement> NewChildElement = MakeShareable(new FFrameEventTreeElement(
+				ParentElement->Frame,
+				FStateTreeTraceEventVariantType(TInPlaceType<EventType>(), Event),
+				ParentElement->WeakStateTree.Get()));
+
+			ParentElement->Children.Add(NewChildElement);
+		};
+
+	if (!InstanceDataAsString.IsEmpty())
+	{
+		auto CreatePropertyElement = [&CreateTreeElement](const FStringView Line, const int32 NestedCount = 0)
+			{
+				constexpr int32 Indent = 4;
+				FString ConvertedString = FString::Printf(TEXT("%*s"), NestedCount * Indent, TEXT(""));
+				ConvertedString.Append(Line);
+				ConvertedString.ReplaceInline(TEXT("="), TEXT(" = "));
+				ConvertedString.ReplaceInline(TEXT("\""), TEXT(""));
+				CreateTreeElement.operator()<FStateTreeTracePropertyEvent>(ConvertedString);
+			};
+
+		// Try to parse Struct for which properties are exported between '(' and ')'
+		if (InstanceDataAsString.StartsWith("(") && InstanceDataAsString.EndsWith(")"))
+		{
+			const FStringView View(GetData(InstanceDataAsString) + 1, InstanceDataAsString.Len() - 2);
+			const TCHAR* ViewIt = View.GetData();
+			const TCHAR* const ViewEnd = ViewIt + View.Len();
+			const TCHAR* NextToken = ViewIt;
+			int32 NestedCount = 0;
+
+			for (; ViewIt != ViewEnd; ++ViewIt)
+			{
+				int32 LocalNestedCount = 0;
+				if (*ViewIt == TCHAR('('))
+				{
+					LocalNestedCount++;
+				}
+				else if (*ViewIt == TCHAR(')'))
+				{
+					LocalNestedCount--;
+				}
+				else if (*ViewIt != TCHAR(','))
+				{
+					continue;
+				}
+
+				if (ViewIt != NextToken)
+				{
+					CreatePropertyElement(FStringView(NextToken, UE_PTRDIFF_TO_INT32(ViewIt - NextToken)), NestedCount);
+				}
+				NextToken = ViewIt + 1;
+				NestedCount += LocalNestedCount;
+			}
+
+			if (ViewIt != NextToken)
+			{
+				CreatePropertyElement(FStringView(NextToken, UE_PTRDIFF_TO_INT32(ViewIt - NextToken)), NestedCount);
+			}
+		}
+		else
+		{
+			const TCHAR* Buffer = *InstanceDataAsString;
+			FParse::Next(&Buffer);
+			FString StrLine;
+			while (FParse::Line(&Buffer, StrLine))
+			{
+				const TCHAR* Str = *StrLine;
+				if (!FParse::Command(&Str, TEXT("BEGIN OBJECT"))
+					&& !FParse::Command(&Str, TEXT("END OBJECT")))
+				{
+					CreatePropertyElement(Str);
+				}
+			}
+		}
+	}
+
+	if (!DebugText.IsEmpty())
+	{
+		CreateTreeElement.operator()<FStateTreeTraceLogEvent>(DebugText);
+	}
+}
+
+void SFrameEventsView::ExpandAll(const TArray<TSharedPtr<FFrameEventTreeElement>>& Items)
+{
+	for (const TSharedPtr<FFrameEventTreeElement>& Item : Items)
+	{
+		bool bExpand = true;
+		if (Item->Children.Num() > 0)
+		{
+			if (const FFrameEventTreeElement* FirstChild = Item->Children[0].Get())
+			{
+				if (const FStateTreeTraceLogEvent* LogEvent = FirstChild->Event.TryGet<FStateTreeTraceLogEvent>())
+				{
+					// Do not auto expand verbose logs
+					bExpand = (LogEvent->Verbosity < ELogVerbosity::Verbose);
+				}
+				else if (FirstChild->Event.IsType<FStateTreeTracePropertyEvent>())
+				{
+					bExpand = false;
+				}
+			}
+		}
+
+		if (bExpand)
+		{
+			EventsTreeView->SetItemExpansion(Item, true);
+			ExpandAll(Item->Children);
+		}
+	}
+}
+
+void SFrameEventsView::RequestRefresh(const FScrubState& InScrubState)
+{
+	// Rebuild frame details from the events of that frame
+	EventsTreeElements.Reset();
+	ON_SCOPE_EXIT
+	{
+		EventsTreeView->ClearExpandedItems();
+		ExpandAll(EventsTreeElements);
+		EventsTreeView->RequestTreeRefresh();
+	};
+
+	const TConstArrayView<const FStateTreeTraceEventVariantType> Events = InScrubState.GetEventCollection().Events;
+
+	if (Events.IsEmpty()
+		|| !InScrubState.IsInBounds())
+	{
+		return;
+	}
+
+	const TConstArrayView<FFrameSpan> Spans = InScrubState.GetEventCollection().FrameSpans;
+	check(Spans.Num());
+	check(WeakStateTree.IsValid());
+
+	TArray<TSharedPtr<FFrameEventTreeElement>, TInlineAllocator<8>> ScopeStack;
+
+	const int32 SpanIdx = InScrubState.GetFrameSpanIndex();
+
+	if (!Spans.IsValidIndex(SpanIdx))
+	{
+		UE_LOG(LogStateTree, Error, TEXT("Invalid index in span: Idx: %i, Num Spans: %i"), SpanIdx, Spans.Num());
+		return;
+	}
+
+	const FFrameSpan& Span = Spans[SpanIdx];
+	if (InScrubState.GetScrubTime() < Span.GetWorldTimeStart()
+		|| InScrubState.GetScrubTime() > Span.GetWorldTimeEnd())
+	{
+		return;
+	}
+
+	const int32 FirstEventIdx = Span.EventIdx;
+	const TraceServices::FFrame Frame = Span.Frame;
+	const int32 MaxEventIdx = Spans.IsValidIndex(SpanIdx + 1) ? Spans[SpanIdx + 1].EventIdx : Events.Num();
+
+	const UStateTree* const RootTree = WeakStateTree.Get();
+	const UStateTree* ActiveTree = RootTree;
+
+	for (int32 EventIdx = FirstEventIdx; EventIdx < MaxEventIdx; EventIdx++)
+	{
+		const FStateTreeTraceEventVariantType& Event = Events[EventIdx];
+		FString CustomDescription;
+		bool bShouldAddToScopeStack = false;
+		bool bShouldPopScopeStack = false;
+
+		if (const FStateTreeTraceStateEvent* StateEvent = Event.TryGet<FStateTreeTraceStateEvent>())
+		{
+			if (StateEvent->EventType == EStateTreeTraceEventType::OnEntering
+				|| StateEvent->EventType == EStateTreeTraceEventType::OnExiting
+				|| StateEvent->EventType == EStateTreeTraceEventType::Push)
+			{
+				bShouldAddToScopeStack = true;
+			}
+			else if (StateEvent->EventType == EStateTreeTraceEventType::OnEntered
+				|| StateEvent->EventType == EStateTreeTraceEventType::OnExited
+				|| StateEvent->EventType == EStateTreeTraceEventType::Pop)
+			{
+				bShouldPopScopeStack = true;
+			}
+		}
+		else if (const FStateTreeTracePhaseEvent* PhaseEvent = Event.TryGet<FStateTreeTracePhaseEvent>())
+		{
+			if (PhaseEvent->EventType == EStateTreeTraceEventType::Push)
+			{
+				bShouldAddToScopeStack = true;
+			}
+			else if (PhaseEvent->EventType == EStateTreeTraceEventType::Pop)
+			{
+				bShouldPopScopeStack = true;
+			}
+		}
+		else if (const FStateTreeTraceInstanceFrameEvent* FrameEvent = Event.TryGet<FStateTreeTraceInstanceFrameEvent>())
+		{
+			ActiveTree = FrameEvent->WeakStateTree.Get();
+			check(ActiveTree);
+
+			// We don't want to create an entry.
+			continue;
+		}
+
+		if (bShouldPopScopeStack)
+		{
+			// Pop scope and remove associated element if empty
+			if (ensureMsgf(ScopeStack.Num() > 0, TEXT("Expected to pop an entry in the scope stack but it is already empty.")))
+			{
+				TSharedPtr<FFrameEventTreeElement> Scope = ScopeStack.Pop();
+				if (Scope->Children.IsEmpty())
+				{
+					TArray<TSharedPtr<FFrameEventTreeElement>>& TreeElements = ScopeStack.IsEmpty() ? EventsTreeElements : ScopeStack.Top()->Children;
+					TreeElements.Remove(Scope);
+				}
+			}
+			// We don't want to create a child when a scope is popped.
+
+			continue;
+		}
+
+		const TSharedRef<FFrameEventTreeElement> NewElement = MakeShareable(new FFrameEventTreeElement(Frame, Event, ActiveTree));
+		NewElement->Description = CustomDescription;
+
+		TArray<TSharedPtr<FFrameEventTreeElement>>& TreeElements = ScopeStack.IsEmpty() ? EventsTreeElements : ScopeStack.Top()->Children;
+		const TSharedPtr<FFrameEventTreeElement>& ElementPtr = TreeElements.Add_GetRef(NewElement);
+
+		if (bShouldAddToScopeStack)
+		{
+			ScopeStack.Push(ElementPtr);
+		}
+
+		GenerateElementsForProperties(Event, NewElement);
+	}
+}
+
+} // UE::StateTreeDebugger
+
+#endif // WITH_STATETREE_TRACE_DEBUGGER
\ No newline at end of file
diff --git a/Engine/Plugins/Runtime/StateTree/Source/StateTreeDeveloper/Private/StateTreeStyle.cpp b/Engine/Plugins/Runtime/StateTree/Source/StateTreeDeveloper/Private/StateTreeStyle.cpp
--- a/Engine/Plugins/Runtime/StateTree/Source/StateTreeDeveloper/Private/StateTreeStyle.cpp
+++ b/Engine/Plugins/Runtime/StateTree/Source/StateTreeDeveloper/Private/StateTreeStyle.cpp
@@ -7,6 +7,7 @@
 #include "Styling/SlateStyleMacros.h"
 #include "Styling/SlateStyleRegistry.h"
 #include "Styling/SlateTypes.h"
+#include "Styling/StyleColors.h"
 
 const FLazyName FStateTreeStyle::StateTitleTextStyleName("StateTree.State.Title");
 const FString FStateTreeStyle::EngineSlateContentDir(FPaths::EngineContentDir() / TEXT("Slate"));
@@ -24,6 +25,22 @@ FStateTreeStyle::FStateTreeStyle(const FName& InStyleSetName)
 
 	const FTextBlockStyle& NormalText = FAppStyle::Get().GetWidgetStyle<FTextBlockStyle>("NormalText");
 
+	// Debugger
+	{
+		Set("StateTreeDebugger.Element.Normal",
+			FTextBlockStyle(NormalText)
+			.SetFont(DEFAULT_FONT("Regular", 10)));
+
+		Set("StateTreeDebugger.Element.Bold",
+			FTextBlockStyle(NormalText)
+			.SetFont(DEFAULT_FONT("Bold", 10)));
+
+		Set("StateTreeDebugger.Element.Subdued",
+			FTextBlockStyle(NormalText)
+			.SetFont(DEFAULT_FONT("Regular", 10))
+			.SetColorAndOpacity(FSlateColor::UseSubduedForeground()));
+	}
+
 	// State
 	{
 		const FTextBlockStyle StateTitle = FTextBlockStyle(NormalText)
@@ -63,7 +80,30 @@ FStateTreeStyle::FStateTreeStyle(const FName& InStyleSetName)
 		Set("StateTreeEditor.TrySelectChildrenInOrder", new IMAGE_BRUSH_SVG("Icons/Try_Select_Children_In_Order", CoreStyleConstants::Icon16x16));
 		Set("StateTreeEditor.TrySelectChildrenAtRandom", new IMAGE_BRUSH_SVG("Icons/Try_Select_Children_At_Random", CoreStyleConstants::Icon16x16));
 		Set("StateTreeEditor.TryFollowTransitions", new IMAGE_BRUSH_SVG("Icons/Try_Follow_Transitions", CoreStyleConstants::Icon16x16));
+
+		Set("StateTreeEditor.Debugger.Condition.OnTransition", new IMAGE_BRUSH_SVG("Icons/State_Conditions", CoreStyleConstants::Icon16x16, FStyleColors::AccentGray));
+	}
+
+	{
+		// From generic Engine
+		FContentRootScope Scope(this, EngineSlateContentDir);
+		Set("StateTreeEditor.Debugger.State.Enter", new CORE_IMAGE_BRUSH_SVG("Starship/Common/arrow-right", CoreStyleConstants::Icon16x16, FStyleColors::Foreground));
+		Set("StateTreeEditor.Debugger.State.Exit", new CORE_IMAGE_BRUSH_SVG("Starship/Common/arrow-left", CoreStyleConstants::Icon16x16, FStyleColors::Foreground));
+		Set("StateTreeEditor.Debugger.State.Completed", new CORE_IMAGE_BRUSH_SVG("Starship/Common/check", CoreStyleConstants::Icon16x16, FStyleColors::AccentGreen));
+		Set("StateTreeEditor.Debugger.State.Selected", new CORE_IMAGE_BRUSH_SVG("Starship/Common/arrow-right", CoreStyleConstants::Icon16x16, FStyleColors::AccentYellow));
+
+		Set("StateTreeEditor.Debugger.Log.Warning", new CORE_IMAGE_BRUSH_SVG("Starship/Common/alert-circle", CoreStyleConstants::Icon16x16, FStyleColors::AccentYellow));
+		Set("StateTreeEditor.Debugger.Log.Error", new CORE_IMAGE_BRUSH_SVG("Starship/Common/x-circle", CoreStyleConstants::Icon16x16, FStyleColors::AccentRed));
+
+		Set("StateTreeEditor.Debugger.Task.Failed", new CORE_IMAGE_BRUSH_SVG("Starship/Common/close-small", CoreStyleConstants::Icon16x16, FStyleColors::AccentRed));
+		Set("StateTreeEditor.Debugger.Task.Succeeded", new CORE_IMAGE_BRUSH_SVG("Starship/Common/check", CoreStyleConstants::Icon16x16, FStyleColors::AccentGreen));
+		Set("StateTreeEditor.Debugger.Task.Stopped", new CORE_IMAGE_BRUSH_SVG("Starship/Common/close-small", CoreStyleConstants::Icon16x16, FStyleColors::AccentRed));
+
+		Set("StateTreeEditor.Debugger.Condition.Passed", new CORE_IMAGE_BRUSH_SVG("Starship/Common/check", CoreStyleConstants::Icon16x16, FStyleColors::AccentGreen));
+		Set("StateTreeEditor.Debugger.Condition.Failed", new CORE_IMAGE_BRUSH_SVG("Starship/Common/close-small", CoreStyleConstants::Icon16x16, FStyleColors::AccentRed));
+		Set("StateTreeEditor.Debugger.Condition.OnEvaluating", new CORE_IMAGE_BRUSH_SVG("Starship/Common/Update", CoreStyleConstants::Icon16x16, FStyleColors::AccentYellow));
 	}
+
 }
 
 void FStateTreeStyle::Register()
diff --git a/Engine/Plugins/Runtime/StateTree/Source/StateTreeDeveloper/Public/StateTreeStyle.h b/Engine/Plugins/Runtime/StateTree/Source/StateTreeDeveloper/Public/StateTreeStyle.h
--- a/Engine/Plugins/Runtime/StateTree/Source/StateTreeDeveloper/Public/StateTreeStyle.h
+++ b/Engine/Plugins/Runtime/StateTree/Source/StateTreeDeveloper/Public/StateTreeStyle.h
@@ -19,6 +19,24 @@ public:
 	static UE_API const FSlateBrush* GetBrushForSelectionBehaviorType(EStateTreeStateSelectionBehavior InSelectionBehavior, bool bInHasChildren, EStateTreeStateType InStateType);
 
 protected:
+	struct FContentRootScope
+	{
+		FContentRootScope(FSlateStyleSet* InStyle, const FString& NewContentRoot)
+			: Style(InStyle)
+			, PreviousContentRoot(InStyle->GetContentRootDir())
+		{
+			Style->SetContentRoot(NewContentRoot);
+		}
+
+		~FContentRootScope()
+		{
+			Style->SetContentRoot(PreviousContentRoot);
+		}
+	private:
+		FSlateStyleSet* Style;
+		FString PreviousContentRoot;
+	};
+
 	friend class FStateTreeDeveloperModule;
 
 	UE_API explicit FStateTreeStyle(const FName& InStyleSetName);
diff --git a/Engine/Plugins/Runtime/StateTree/Source/StateTreeEditorModule/Private/Debugger/SStateTreeDebuggerView.cpp b/Engine/Plugins/Runtime/StateTree/Source/StateTreeEditorModule/Private/Debugger/SStateTreeDebuggerView.cpp
--- a/Engine/Plugins/Runtime/StateTree/Source/StateTreeEditorModule/Private/Debugger/SStateTreeDebuggerView.cpp
+++ b/Engine/Plugins/Runtime/StateTree/Source/StateTreeEditorModule/Private/Debugger/SStateTreeDebuggerView.cpp
@@ -3,194 +3,26 @@
 #if WITH_STATETREE_TRACE_DEBUGGER
 
 #include "SStateTreeDebuggerView.h"
-#include "Debugger/StateTreeDebugger.h"
-#include "DetailsViewArgs.h"
+#include "Debugger/SStateTreeFrameEventsView.h"
 #include "Editor.h"
-#include "Factories.h"
 #include "Framework/MultiBox/MultiBoxBuilder.h"
 #include "Kismet2/DebuggerCommands.h"
 #include "Modules/ModuleManager.h"
-#include "PropertyEditorModule.h"
 #include "SStateTreeDebuggerInstanceTree.h"
 #include "SStateTreeDebuggerTimelines.h"
-#include "SStateTreeDebuggerViewRow.h"
-#include "StateTree.h"
-#include "StateTreeCompiler.h"
 #include "StateTreeDebuggerCommands.h"
 #include "StateTreeDebuggerTrack.h"
 #include "StateTreeDelegates.h"
 #include "StateTreeEditorData.h"
-#include "StateTreeEditorModule.h"
 #include "StateTreeEditorSettings.h"
 #include "StateTreeEditorStyle.h"
 #include "StateTreeModule.h"
-#include "StateTreeState.h"
 #include "StateTreeViewModel.h"
 #include "Widgets/Input/SComboButton.h"
-#include "Widgets/Layout/SBorder.h"
 #include "Widgets/Layout/SScrollBox.h"
 
 #define LOCTEXT_NAMESPACE "StateTreeEditor"
 
-
-namespace UE::StateTreeDebugger
-{
-/**
- * Iterates over all tree elements for the frame events
- * @param Elements Container of hierarchical tree element to visit
- * @param InFunc function called at each element, should return true if visiting is continued or false to stop.
- */
-void VisitEventTreeElements(const TConstArrayView<TSharedPtr<FStateTreeDebuggerEventTreeElement>> Elements, TFunctionRef<bool(TSharedPtr<FStateTreeDebuggerEventTreeElement>& VisitedElement)> InFunc)
-{
-	TArray<TSharedPtr<FStateTreeDebuggerEventTreeElement>> Stack;
-	bool bContinue = true;
-
-	for (const TSharedPtr<FStateTreeDebuggerEventTreeElement>& RootElement : Elements)
-	{
-		if (RootElement == nullptr)
-		{
-			continue;
-		}
-
-		Stack.Add(RootElement);
-
-		while (!Stack.IsEmpty() && bContinue)
-		{
-			TSharedPtr<FStateTreeDebuggerEventTreeElement> StackedElement = Stack[0];
-			check(StackedElement);
-
-			Stack.RemoveAt(0);
-
-			bContinue = InFunc(StackedElement);
-
-			if (bContinue)
-			{
-				for (const TSharedPtr<FStateTreeDebuggerEventTreeElement>& Child : StackedElement->Children)
-				{
-					if (Child.IsValid())
-					{
-						Stack.Add(Child);
-					}
-				}
-			}
-		}
-
-		if (!bContinue)
-		{
-			break;
-		}
-	}
-}
-
-//----------------------------------------------------------------------//
-// FTraceTextObjectFactory
-//----------------------------------------------------------------------//
-struct FTraceTextObjectFactory : FCustomizableTextObjectFactory
-{
-	UObject* NodeInstanceObject = nullptr;
-	FTraceTextObjectFactory() : FCustomizableTextObjectFactory(GWarn) {}
-
-protected:
-	virtual bool CanCreateClass(UClass* ObjectClass, bool& bOmitSubObjs) const override
-	{
-		return true;
-	}
-
-	virtual void ProcessConstructedObject(UObject* CreatedObject) override
-	{
-		NodeInstanceObject = CreatedObject;
-	}
-};
-
-/**
- * Prototype methode to try import struct/object from exported text.
- * Currently not used since it might fail if type is not found and for requires every property to be created
- * but some types (i.e. Object based) will simply set none as the imported value but for the debugger we still
- * want text to describe the dumped values.
- */
-void GenerateElementForProperties(const TCHAR* TypeAsText, const TCHAR* ValueAsText, TSharedRef<FStateTreeDebuggerEventTreeElement>& ParentElement)
-{
-	const UStruct* StructTypeToVisit = nullptr;
-	const void* StructValueToVisit = nullptr;
-
-	// UStruct
-	UScriptStruct* ScriptStruct = FindObject<UScriptStruct>(nullptr, TypeAsText, EFindObjectFlags::None);
-	if (ScriptStruct == nullptr)
-	{
-		ScriptStruct = LoadObject<UScriptStruct>(nullptr, TypeAsText);
-	}
-
-	if (ScriptStruct != nullptr)
-	{
-		FInstancedStruct NodeDataStruct;
-		NodeDataStruct.InitializeAs(ScriptStruct);
-
-		ScriptStruct->ImportText(ValueAsText, NodeDataStruct.GetMutableMemory(), /*OwnerObject*/nullptr, PPF_None, GLog, ScriptStruct->GetName());
-		StructTypeToVisit = ScriptStruct;
-		StructValueToVisit = NodeDataStruct.GetMemory();
-	}
-
-	// UObject
-	UClass* Class = FindObject<UClass>(nullptr, TypeAsText, EFindObjectFlags::None);
-	if (Class == nullptr)
-	{
-		Class = LoadObject<UClass>(nullptr, TypeAsText);
-	}
-
-	if (Class != nullptr)
-	{
-		FTraceTextObjectFactory ObjectFactory;
-		if (ObjectFactory.CanCreateObjectsFromText(ValueAsText))
-		{
-			ObjectFactory.ProcessBuffer(GetTransientPackage(), RF_Transactional, ValueAsText);
-			if (ObjectFactory.NodeInstanceObject)
-			{
-				TWeakObjectPtr<> NodeDataObject;
-				NodeDataObject = ObjectFactory.NodeInstanceObject;
-				NodeDataObject->AddToRoot();
-
-				StructTypeToVisit = ObjectFactory.NodeInstanceObject->GetClass();
-				StructValueToVisit = ObjectFactory.NodeInstanceObject;
-			}
-		}
-	}
-
-	if (StructTypeToVisit != nullptr && StructValueToVisit != nullptr)
-	{
-		for (TPropertyValueIterator<const FProperty> PropertyIt(StructTypeToVisit, StructValueToVisit); PropertyIt; ++PropertyIt)
-		{
-			const FProperty* const Property = PropertyIt.Key();
-			check(Property);
-			const EStateTreePropertyUsage Usage = StateTree::GetUsageFromMetaData(Property);
-
-			// If the property is set to one of these usages, display it even if it is not edit on instance.
-			// It is a common mistake to forget to set the "eye" on these properties it and wonder why it does not show up.
-			const bool bShouldShowByUsage = Usage == EStateTreePropertyUsage::Input || Usage == EStateTreePropertyUsage::Output || Usage == EStateTreePropertyUsage::Context;
-			const bool bIsEditable = Property->HasAllPropertyFlags(CPF_Edit) && !Property->HasAllPropertyFlags(CPF_DisableEditOnInstance);
-
-			if (bShouldShowByUsage || bIsEditable)
-			{
-				FText Name = Property->GetDisplayNameText();
-				FString ValueString;
-				Property->ExportTextItem_Direct(ValueString, PropertyIt.Value(), nullptr, nullptr, PPF_None);
-
-				// Create new property Log
-				FStateTreeTracePropertyEvent PropertyEvent(/*RecordingWorldTime*/0, FString::Printf(TEXT("%s = %s"), *Name.ToString(), *ValueString));
-
-				// Create Tree element to hold the event
-				const TSharedPtr<FStateTreeDebuggerEventTreeElement> NewChildElement = MakeShareable(new FStateTreeDebuggerEventTreeElement(
-						ParentElement->Frame,
-						FStateTreeTraceEventVariantType(TInPlaceType<FStateTreeTracePropertyEvent>(), PropertyEvent),
-						ParentElement->WeakStateTree.Get()));
-
-				ParentElement->Children.Add(NewChildElement);
-			}
-		}
-	}
-}
-
-} // UE::StateTreeDebugger
-
 //----------------------------------------------------------------------//
 // SStateTreeDebuggerView
 //----------------------------------------------------------------------//
@@ -203,7 +35,6 @@ SStateTreeDebuggerView::~SStateTreeDebuggerView()
 	Debugger->OnBreakpointHit.Unbind();
 	Debugger->OnNewSession.Unbind();
 	Debugger->OnNewInstance.Unbind();
-	Debugger->OnSelectedInstanceCleared.Unbind();
 
 	Debugger->ClearSelection();
 
@@ -222,7 +53,7 @@ void SStateTreeDebuggerView::StartRecording()
 		Debugger->ClearSelection();
 
 		// Stop the current analysis so we have a chance to reconnect to an existing instance
-		// if there is existing data and it is still active.
+		// if there is existing data, and it is still active.
 		Debugger->StopSessionAnalysis();
 
 		// We give priority to the Editor actions even if an analysis was active (remote process)
@@ -237,7 +68,7 @@ void SStateTreeDebuggerView::StopRecording()
 	{
 		IStateTreeModule& StateTreeModule = FModuleManager::GetModuleChecked<IStateTreeModule>("StateTreeModule");
 		StateTreeModule.StopTraces();
-		
+
 		bRecording = false;
 
 		// Update max duration from current recording until track data gets reset
@@ -261,7 +92,7 @@ void SStateTreeDebuggerView::StopRecording()
 bool SStateTreeDebuggerView::CanResumeDebuggerAnalysis() const
 {
 	check(Debugger);
-	return Debugger->IsAnalysisSessionPaused() && Debugger->HasHitBreakpoint(); 
+	return Debugger->IsAnalysisSessionPaused() && Debugger->HasHitBreakpoint();
 }
 
 void SStateTreeDebuggerView::ResumeDebuggerAnalysis() const
@@ -333,24 +164,27 @@ void SStateTreeDebuggerView::OnPIESingleStepped(bool bSimulating) const
 	Debugger->SyncToCurrentSessionDuration();
 }
 
-void SStateTreeDebuggerView::Construct(const FArguments& InArgs, const UStateTree& InStateTree, const TSharedRef<FStateTreeViewModel>& InStateTreeViewModel, const TSharedRef<FUICommandList>& InCommandList)
+void SStateTreeDebuggerView::Construct(const FArguments& InArgs
+	, const TNotNull<const UStateTree*> InStateTree
+	, const TSharedRef<FStateTreeViewModel>& InStateTreeViewModel
+	, const TSharedRef<FUICommandList>& InCommandList)
 {
 	BeginPIEHandle = FEditorDelegates::BeginPIE.AddSP(this, &SStateTreeDebuggerView::OnPIEStarted);
 	PausePIEHandle = FEditorDelegates::PausePIE.AddSP(this, &SStateTreeDebuggerView::OnPIEPaused);
 	ResumePIEHandle = FEditorDelegates::ResumePIE.AddSP(this, &SStateTreeDebuggerView::OnPIEResumed);
-	SingleStepPIEHandle =FEditorDelegates::SingleStepPIE.AddSP(this, &SStateTreeDebuggerView::OnPIESingleStepped);
+	SingleStepPIEHandle = FEditorDelegates::SingleStepPIE.AddSP(this, &SStateTreeDebuggerView::OnPIESingleStepped);
 	// Using ShutdownPIE instead of EndPIE to make sure all traces emitted during world EndPlay get processed before disabling channels
 	ShutdownPIEHandle = FEditorDelegates::ShutdownPIE.AddSP(this, &SStateTreeDebuggerView::OnPIEStopped);
 
 	StateTreeViewModel = InStateTreeViewModel;
-	StateTree = &InStateTree;
-	StateTreeEditorData = Cast<UStateTreeEditorData>(InStateTree.EditorData.Get());
+	StateTree = InStateTree;
+	StateTreeEditorData = Cast<UStateTreeEditorData>(InStateTree->EditorData.Get());
 	CommandList = InCommandList;
 
 	Debugger = InStateTreeViewModel->GetDebugger();
 	check(Debugger);
 
-	IStateTreeModule& StateTreeModule = FModuleManager::GetModuleChecked<IStateTreeModule>("StateTreeModule");
+	const IStateTreeModule& StateTreeModule = FModuleManager::GetModuleChecked<IStateTreeModule>("StateTreeModule");
 	bRecording = StateTreeModule.IsTracing();
 	UE::StateTree::Delegates::OnTracingStateChanged.AddSPLambda(this, [&bRecording=bRecording](const EStateTreeTraceStatus TraceStatus)
 		{
@@ -372,7 +206,6 @@ void SStateTreeDebuggerView::Construct(const FArguments& InArgs, const UStateTre
 	Debugger->OnNewInstance.BindSP(this, &SStateTreeDebuggerView::OnNewInstance);
 	Debugger->OnScrubStateChanged.BindSP(this, &SStateTreeDebuggerView::OnDebuggerScrubStateChanged);
 	Debugger->OnBreakpointHit.BindSP(this, &SStateTreeDebuggerView::OnBreakpointHit, InCommandList);
-	Debugger->OnSelectedInstanceCleared.BindSP(this, &SStateTreeDebuggerView::OnSelectedInstanceCleared);
 
 	// Bind our scrub time attribute to follow the value computed by the debugger
 	ScrubTimeAttribute = TAttribute<double>(InStateTreeViewModel->GetDebugger(), &FStateTreeDebugger::GetScrubTime);
@@ -470,10 +303,10 @@ void SStateTreeDebuggerView::Construct(const FArguments& InArgs, const UStateTre
 			})
 		];
 
-	const TSharedPtr<SScrollBar> ScrollBar = SNew(SScrollBar); 
-	
+	const TSharedPtr<SScrollBar> ScrollBar = SNew(SScrollBar);
+
 	// Instances TreeView
-	InstancesTreeView =	SNew(SStateTreeDebuggerInstanceTree)
+	InstancesTreeView = SNew(SStateTreeDebuggerInstanceTree)
 		.ExternalScrollBar(ScrollBar)
 		.OnExpansionChanged_Lambda([this]() { InstanceTimelinesTreeView->RestoreExpansion(); })
 		.OnScrolled_Lambda([this](double ScrollOffset)
@@ -490,7 +323,7 @@ void SStateTreeDebuggerView::Construct(const FArguments& InArgs, const UStateTre
 					StateTreeTrack->OnSelected();
 				}
 			});
-	
+
 	// Timelines TreeView
 	InstanceTimelinesTreeView = SNew(SStateTreeDebuggerTimelines)
 		.ExternalScrollbar(ScrollBar)
@@ -503,32 +336,6 @@ void SStateTreeDebuggerView::Construct(const FArguments& InArgs, const UStateTre
 		.ScrubPosition(ScrubTimeAttribute)
 		.OnScrubPositionChanged_Lambda([this](double NewScrubTime, bool bIsScrubbing) { OnTimeLineScrubPositionChanged(NewScrubTime, bIsScrubbing); });
 
-	// EventsTreeView scrollbars
-	TSharedRef<SScrollBar> HorizontalScrollBar = SNew(SScrollBar)
-		.Orientation(Orient_Horizontal)
-		.Thickness(FVector2D(12.0f, 12.0f));
-
-	TSharedRef<SScrollBar> VerticalScrollBar = SNew(SScrollBar)
-		.Orientation(Orient_Vertical)
-		.Thickness(FVector2D(12.0f, 12.0f));
-
-	// EventsTreeView
-	EventsTreeView = SNew(STreeView<TSharedPtr<FStateTreeDebuggerEventTreeElement>>)
-			.OnGenerateRow_Lambda([this](const TSharedPtr<FStateTreeDebuggerEventTreeElement>& InElement, const TSharedRef<STableViewBase>& InOwnerTableView)
-			{
-				return SNew(SStateTreeDebuggerViewRow, InOwnerTableView, InElement);
-			})
-			.OnGetChildren_Lambda([](const TSharedPtr<const FStateTreeDebuggerEventTreeElement>& InParent, TArray<TSharedPtr<FStateTreeDebuggerEventTreeElement>>& OutChildren)
-			{
-				if (const FStateTreeDebuggerEventTreeElement* Parent = InParent.Get())
-				{
-					OutChildren.Append(Parent->Children);
-				}
-			})
-		.TreeItemsSource(&EventsTreeElements)
-		.AllowOverscroll(EAllowOverscroll::No)
-		.ExternalScrollbar(VerticalScrollBar);
-
 	ChildSlot
 	[
 		SNew(SBorder)
@@ -639,31 +446,7 @@ void SStateTreeDebuggerView::Construct(const FArguments& InArgs, const UStateTre
 					+ SVerticalBox::Slot()
 					.FillHeight(1.0f)
 					[
-						SNew(SHorizontalBox)
-						+ SHorizontalBox::Slot()
-						.FillWidth(1.0f)
-						.Padding(0.0f)
-						[
-							SNew(SScrollBox)
-							.Orientation(Orient_Horizontal)
-							.ExternalScrollbar(HorizontalScrollBar)
-							+ SScrollBox::Slot()
-							.FillSize(1.0f)
-							[
-								EventsTreeView.ToSharedRef()
-							]
-						]
-
-						+ SHorizontalBox::Slot()
-						.AutoWidth()
-						[
-							VerticalScrollBar
-						]
-					]
-					+SVerticalBox::Slot()
-					.AutoHeight()
-					[
-						HorizontalScrollBar
+						SAssignNew(EventsView, UE::StateTreeDebugger::SFrameEventsView, InStateTree)
 					]
 				]
 			]
@@ -673,7 +456,7 @@ void SStateTreeDebuggerView::Construct(const FArguments& InArgs, const UStateTre
 	if (Debugger->IsAnalysisSessionActive())
 	{
 		// Creating the view for an already active analysis session
-		// so rebuild track views for all instances and scrub to the analysis current time 
+		// so rebuild track views for all instances and scrub to the analysis current time
 
 		TArray<const TSharedRef<const UE::StateTreeDebugger::FInstanceDescriptor>> InstanceDescriptors;
 		Debugger->GetSessionInstanceDescriptors(InstanceDescriptors);
@@ -736,7 +519,7 @@ void SStateTreeDebuggerView::Tick(const FGeometry& AllottedGeometry, const doubl
 	}
 
 	MaxTrackRecordingDuration = FMath::Max(MaxTrackRecordingDuration, MaxTrackRecordingDuration);
-	
+
 	if ((Debugger->IsAnalysisSessionActive() && !Debugger->IsAnalysisSessionPaused())
 		|| bHasMoreRecentData)
 	{
@@ -762,7 +545,7 @@ void SStateTreeDebuggerView::Tick(const FGeometry& AllottedGeometry, const doubl
 			Debugger->SetScrubTime(Debugger->GetScrubTime());
 		}
 	}
-	
+
 	RefreshTracks();
 }
 
@@ -799,7 +582,7 @@ void SStateTreeDebuggerView::BindDebuggerToolbarCommands(const TSharedRef<FUICom
 		FCanExecuteAction::CreateSP(this, &SStateTreeDebuggerView::CanStopRecording),
 		FIsActionChecked(),
 		FIsActionButtonVisible::CreateRaw(this, &SStateTreeDebuggerView::CanStopRecording));
-	
+
 	ToolkitCommands->MapAction(
 		Commands.PreviousFrameWithStateChange,
 		FExecuteAction::CreateSP(this, &SStateTreeDebuggerView::StepBackToPreviousStateChange),
@@ -907,251 +690,7 @@ void SStateTreeDebuggerView::OnDebuggerScrubStateChanged(const UE::StateTreeDebu
 {
 	TrackCursor();
 
-	// Rebuild frame details from the events of that frame
-	EventsTreeElements.Reset();
-	ON_SCOPE_EXIT
-	{
-		EventsTreeView->ClearExpandedItems();
-		ExpandAll(EventsTreeElements);
-		EventsTreeView->RequestTreeRefresh();
-	};
-
-	const UE::StateTreeDebugger::FInstanceEventCollection& EventCollection = ScrubState.GetEventCollection();
-	const TConstArrayView<const FStateTreeTraceEventVariantType> Events = EventCollection.Events;
-
-	if (Events.IsEmpty() || ScrubState.IsInBounds() == false)
-	{
-		return;
-	}
-
-	const TConstArrayView<UE::StateTreeDebugger::FFrameSpan> Spans = EventCollection.FrameSpans;
-	check(Spans.Num());
-	check(StateTree.IsValid());
-
-	TArray<TSharedPtr<FStateTreeDebuggerEventTreeElement>, TInlineAllocator<8>> ScopeStack;
-
-	const int32 SpanIdx = ScrubState.GetFrameSpanIndex();
-
-	if (!Spans.IsValidIndex(SpanIdx))
-	{
-		UE_LOG(LogStateTreeEditor, Error, TEXT("Invalid index in span: Idx: %i, Num Spans: %i"), SpanIdx, Spans.Num());
-		return;
-	}
-
-	const UE::StateTreeDebugger::FFrameSpan& Span = Spans[SpanIdx];
-	if (ScrubState.GetScrubTime() < Span.GetWorldTimeStart()
-		|| ScrubState.GetScrubTime() > Span.GetWorldTimeEnd())
-	{
-		return;
-	}
-
-	const int32 FirstEventIdx = Span.EventIdx;
-	const TraceServices::FFrame Frame = Span.Frame;
-	const int32 MaxEventIdx = Spans.IsValidIndex(SpanIdx+1) ? Spans[SpanIdx+1].EventIdx : Events.Num();
-
-	const UStateTree* const RootTree = StateTree.Get();
-	const UStateTree* ActiveTree = RootTree;
-	
-	for (int32 EventIdx = FirstEventIdx; EventIdx < MaxEventIdx; EventIdx++)
-	{
-		const FStateTreeTraceEventVariantType& Event = Events[EventIdx];
-		FString CustomDescription;
-		bool bShouldAddToScopeStack = false;
-		bool bShouldPopScopeStack = false;
-
-		if (const FStateTreeTraceStateEvent* StateEvent = Event.TryGet<FStateTreeTraceStateEvent>())
-		{
-			if (StateEvent->EventType == EStateTreeTraceEventType::OnEntering
-			 	|| StateEvent->EventType == EStateTreeTraceEventType::OnExiting
-				|| StateEvent->EventType == EStateTreeTraceEventType::Push)
-			{
-				bShouldAddToScopeStack = true;
-			}
-			else if (StateEvent->EventType == EStateTreeTraceEventType::OnEntered
-				|| StateEvent->EventType == EStateTreeTraceEventType::OnExited
-				|| StateEvent->EventType == EStateTreeTraceEventType::Pop)
-			{
-				bShouldPopScopeStack = true;
-			}
-		}
-		else if (const FStateTreeTracePhaseEvent* PhaseEvent = Event.TryGet<FStateTreeTracePhaseEvent>())
-		{
-			if (PhaseEvent->EventType == EStateTreeTraceEventType::Push)
-			{
-				bShouldAddToScopeStack = true;
-			}
-			else if (PhaseEvent->EventType == EStateTreeTraceEventType::Pop)
-			{
-				bShouldPopScopeStack = true;
-			}
-		}
-		else if (const FStateTreeTraceInstanceFrameEvent* FrameEvent = Event.TryGet<FStateTreeTraceInstanceFrameEvent>())
-		{
-			ActiveTree = FrameEvent->WeakStateTree.Get();
-			check(ActiveTree);
-
-			// We don't want to create an entry.
-			continue;
-		}
-
-		if (bShouldPopScopeStack)
-		{
-			// Pop scope and remove associated element if empty
-			if (ensureMsgf(ScopeStack.Num() > 0, TEXT("Expected to pop an entry in the scope stack but it is already empty.")))
-			{
-				TSharedPtr<FStateTreeDebuggerEventTreeElement> Scope = ScopeStack.Pop();
-				if (Scope->Children.IsEmpty())
-				{
-					TArray<TSharedPtr<FStateTreeDebuggerEventTreeElement>>& TreeElements = ScopeStack.IsEmpty() ? EventsTreeElements : ScopeStack.Top()->Children;
-					TreeElements.Remove(Scope);
-				}
-			}
-			// We don't want to create a child when a scope is popped.
-			continue;
-		}
-
-		const TSharedRef<FStateTreeDebuggerEventTreeElement> NewElement = MakeShareable(new FStateTreeDebuggerEventTreeElement(Frame, Event, ActiveTree));
-		NewElement->Description = CustomDescription;
-
-		TArray<TSharedPtr<FStateTreeDebuggerEventTreeElement>>& TreeElements = ScopeStack.IsEmpty() ? EventsTreeElements : ScopeStack.Top()->Children;
-		const TSharedPtr<FStateTreeDebuggerEventTreeElement>& ElementPtr = TreeElements.Add_GetRef(NewElement);
-
-		if (bShouldAddToScopeStack)
-		{
-			ScopeStack.Push(ElementPtr);
-		}
-
-		GenerateElementsForProperties(Event, NewElement);
-	}
-}
-
-void SStateTreeDebuggerView::GenerateElementsForProperties(const FStateTreeTraceEventVariantType& Event, const TSharedRef<FStateTreeDebuggerEventTreeElement>& ParentElement)
-{
-	FString TypePath;
-	FString InstanceDataAsString;
-	FString DebugText;
-
-	Visit([&TypePath, &InstanceDataAsString, &DebugText](auto& TypedEvent)
-		{
-			TypePath = TypedEvent.GetDataTypePath();
-			InstanceDataAsString = TypedEvent.GetDataAsText();
-			DebugText = TypedEvent.GetDebugText();
-		}, Event);
-
-	auto CreateTreeElement = [ParentElement]<typename EventType>(const FString& Line)
-		{
-			// Create new event
-			EventType Event(/*RecordingWorldTime*/0, ELogVerbosity::Verbose, *Line);
-
-			// Create Tree element to hold the event
-			const TSharedPtr<FStateTreeDebuggerEventTreeElement> NewChildElement = MakeShareable(new FStateTreeDebuggerEventTreeElement(
-				ParentElement->Frame,
-				FStateTreeTraceEventVariantType(TInPlaceType<EventType>(), Event),
-				ParentElement->WeakStateTree.Get()));
-
-			ParentElement->Children.Add(NewChildElement);
-		};
-
-	if (!InstanceDataAsString.IsEmpty())
-	{
-		auto CreatePropertyElement = [&CreateTreeElement](const FStringView Line, const int32 NestedCount = 0)
-			{
-				constexpr int32 Indent = 4;
-				FString ConvertedString = FString::Printf(TEXT("%*s"), NestedCount * Indent, TEXT(""));
-				ConvertedString.Append(Line);
-				ConvertedString.ReplaceInline(TEXT("="), TEXT(" = "));
-				ConvertedString.ReplaceInline(TEXT("\""), TEXT(""));
-				CreateTreeElement.operator()<FStateTreeTracePropertyEvent>(ConvertedString);
-			};
-
-		// Try to parse Struct for which properties are exported between '(' and ')' 
-		if (InstanceDataAsString.StartsWith("(") && InstanceDataAsString.EndsWith(")"))
-		{
-			const FStringView View(GetData(InstanceDataAsString) + 1, InstanceDataAsString.Len() - 2);
-			const TCHAR* ViewIt = View.GetData();
-			const TCHAR* const ViewEnd = ViewIt + View.Len();
-			const TCHAR* NextToken = ViewIt;
-			int32 NestedCount = 0;
-			
-			for (; ViewIt != ViewEnd; ++ViewIt)
-			{
-				int32 LocalNestedCount = 0;
-				if (*ViewIt == TCHAR('('))
-				{
-					LocalNestedCount++;
-				}
-				else if (*ViewIt == TCHAR(')'))
-				{
-					LocalNestedCount--;
-				}
-				else if (*ViewIt != TCHAR(','))
-				{
-					continue;
-				}
-
-				if (ViewIt != NextToken)
-				{
-					CreatePropertyElement(FStringView(NextToken, UE_PTRDIFF_TO_INT32(ViewIt - NextToken)), NestedCount);
-				}
-				NextToken = ViewIt+1;
-				NestedCount += LocalNestedCount;
-			}
-
-			if (ViewIt != NextToken)
-			{
-				CreatePropertyElement(FStringView(NextToken, UE_PTRDIFF_TO_INT32(ViewIt - NextToken)), NestedCount);
-			}
-		}
-		else
-		{
-			const TCHAR* Buffer = *InstanceDataAsString;
-			FParse::Next(&Buffer);
-			FString StrLine;
-			while (FParse::Line(&Buffer, StrLine))
-			{
-				const TCHAR* Str = *StrLine;
-				if (!FParse::Command(&Str, TEXT("BEGIN OBJECT"))
-					&& !FParse::Command(&Str, TEXT("END OBJECT")))
-				{
-					CreatePropertyElement(Str);
-				}
-			}
-		}
-	}
-
-	if (!DebugText.IsEmpty())
-	{
-		CreateTreeElement.operator()<FStateTreeTraceLogEvent>(DebugText);
-	}
-}
-
-void SStateTreeDebuggerView::ExpandAll(const TArray<TSharedPtr<FStateTreeDebuggerEventTreeElement>>& Items)
-{
-	for (const TSharedPtr<FStateTreeDebuggerEventTreeElement>& Item : Items)
-	{
-		bool bExpand = true;
-		if (Item->Children.Num() > 0)
-		{
-			if (const FStateTreeDebuggerEventTreeElement* FirstChild = Item->Children[0].Get())
-			{
-				if (const FStateTreeTraceLogEvent* LogEvent = FirstChild->Event.TryGet<FStateTreeTraceLogEvent>())
-				{
-					// Do not auto expand verbose logs
-					bExpand = (LogEvent->Verbosity < ELogVerbosity::Verbose);
-				}
-				else if (FirstChild->Event.IsType<FStateTreeTracePropertyEvent>())
-				{
-					bExpand = false;
-				}
-			}
-		}
-
-		if (bExpand)
-		{
-			EventsTreeView->SetItemExpansion(Item, true);
-			ExpandAll(Item->Children);
-		}
-	}
+	EventsView->RequestRefresh(ScrubState);
 }
 
 void SStateTreeDebuggerView::OnBreakpointHit(const FStateTreeInstanceDebugId InstanceId, const FStateTreeDebuggerBreakpoint Breakpoint, const TSharedRef<FUICommandList> ActionList) const
@@ -1173,24 +712,12 @@ void SStateTreeDebuggerView::OnBreakpointHit(const FStateTreeInstanceDebugId Ins
 	}
 
 	// Find matching event in the tree view and select it
-	if (EventsTreeView.IsValid())
+	if (EventsView.IsValid())
 	{
-		TSharedPtr<FStateTreeDebuggerEventTreeElement> MatchingElement = nullptr;
-		UE::StateTreeDebugger::VisitEventTreeElements(EventsTreeElements, [&MatchingElement, Breakpoint](const TSharedPtr<FStateTreeDebuggerEventTreeElement>& VisitedElement)
+		EventsView->SelectByPredicate([Breakpoint](const FStateTreeTraceEventVariantType& Event)
 			{
-				if (Breakpoint.IsMatchingEvent(VisitedElement->Event))
-				{
-					MatchingElement = VisitedElement;
-				}
-
-				// Continue visit until we find a matching event
-				return MatchingElement.IsValid() == false;
+				return Breakpoint.IsMatchingEvent(Event);
 			});
-
-		if (MatchingElement.IsValid())
-		{
-			EventsTreeView->SetSelection(MatchingElement);
-		}
 	}
 }
 
@@ -1243,7 +770,7 @@ void SStateTreeDebuggerView::OnNewInstance(FStateTreeInstanceDebugId InstanceId)
 
 	if (FStateTreeDebuggerOwnerTrack* OwnerTrack = static_cast<FStateTreeDebuggerOwnerTrack*>(ExistingOwnerTrack->Get()))
 	{
-		const FString TrackName = FString::Printf(TEXT("Execution #%d"), OwnerTrack->NumSubTracks()+1);
+		const FString TrackName = FString::Printf(TEXT("Execution #%d"), OwnerTrack->NumSubTracks() + 1);
 		const TSharedPtr<FStateTreeDebuggerInstanceTrack> SubTrack = MakeShared<FStateTreeDebuggerInstanceTrack>(SharedThis(this), Debugger.ToSharedRef(), InstanceId, FText::FromString(TrackName), ViewRange);
 		OwnerTrack->AddSubTrack(SubTrack);
 
@@ -1267,15 +794,6 @@ void SStateTreeDebuggerView::OnNewInstance(FStateTreeInstanceDebugId InstanceId)
 	InstanceTimelinesTreeView->Refresh();
 }
 
-void SStateTreeDebuggerView::OnSelectedInstanceCleared()
-{
-	EventsTreeElements.Reset();
-	if (EventsTreeView)
-	{
-		EventsTreeView->RequestTreeRefresh();
-	}
-}
-
 TSharedRef<SWidget> SStateTreeDebuggerView::OnGetDebuggerTracesMenu() const
 {
 	FMenuBuilder MenuBuilder(/*bInShouldCloseWindowAfterMenuSelection*/true, /*InCommandList*/nullptr);
diff --git a/Engine/Plugins/Runtime/StateTree/Source/StateTreeEditorModule/Private/Debugger/SStateTreeDebuggerView.h b/Engine/Plugins/Runtime/StateTree/Source/StateTreeEditorModule/Private/Debugger/SStateTreeDebuggerView.h
--- a/Engine/Plugins/Runtime/StateTree/Source/StateTreeEditorModule/Private/Debugger/SStateTreeDebuggerView.h
+++ b/Engine/Plugins/Runtime/StateTree/Source/StateTreeEditorModule/Private/Debugger/SStateTreeDebuggerView.h
@@ -15,6 +15,7 @@ namespace RewindDebugger
 namespace UE::StateTreeDebugger
 {
 	struct FFrameSpan;
+	class SFrameEventsView;
 }
 
 enum class EStateTreeBreakpointType : uint8;
@@ -23,7 +24,6 @@ struct FSlateIcon;
 struct FStateTreeDebugger;
 struct FStateTreeInstanceDebugId;
 struct FStateTreeDebuggerBreakpoint;
-struct FStateTreeDebuggerEventTreeElement;
 class FStateTreeEditor;
 class FStateTreeViewModel;
 class FUICommandList;
@@ -45,7 +45,7 @@ public:
 
 	~SStateTreeDebuggerView();
 
-	void Construct(const FArguments& InArgs, const UStateTree& InStateTree, const TSharedRef<FStateTreeViewModel>& InStateTreeViewModel, const TSharedRef<FUICommandList>& InCommandList);
+	void Construct(const FArguments& InArgs, TNotNull<const UStateTree*> InStateTree, const TSharedRef<FStateTreeViewModel>& InStateTreeViewModel, const TSharedRef<FUICommandList>& InCommandList);
 
 	[[nodiscard]] TSharedPtr<FStateTreeDebugger> GetDebugger() const
 	{
@@ -72,11 +72,10 @@ private:
 	void OnPIEPaused(bool bIsSimulating) const;
 	void OnPIEResumed(bool bIsSimulating) const;
 	void OnPIESingleStepped(bool bIsSimulating) const;
-	
+
 	void OnBreakpointHit(const FStateTreeInstanceDebugId InstanceId, const FStateTreeDebuggerBreakpoint Breakpoint, const TSharedRef<FUICommandList> ActionList) const;
 	void OnNewSession();
 	void OnNewInstance(FStateTreeInstanceDebugId InstanceId);
-	void OnSelectedInstanceCleared();
 
 	void BindDebuggerToolbarCommands(const TSharedRef<FUICommandList>& ToolkitCommands);
 
@@ -125,11 +124,6 @@ private:
 	/** Callback used to reflect debugger scrub state in the UI. */
 	void OnDebuggerScrubStateChanged(const UE::StateTreeDebugger::FScrubState& ScrubState);
 
-	static void GenerateElementsForProperties(const FStateTreeTraceEventVariantType& Event, const TSharedRef<FStateTreeDebuggerEventTreeElement>& ParentElement);
-
-	/** Recursively sets tree items as expanded. */
-	void ExpandAll(const TArray<TSharedPtr<FStateTreeDebuggerEventTreeElement>>& Items);
-
 	TSharedPtr<FStateTreeDebugger> Debugger;
 	TSharedPtr<FStateTreeViewModel> StateTreeViewModel;
 	TWeakObjectPtr<const UStateTree> StateTree;
@@ -159,11 +153,8 @@ private:
 	/** Splitter between instances names and their timelines. Used with HeaderSplitter to keep header and content in sync */
 	TSharedPtr<SSplitter> TreeViewsSplitter;
 
-	/** All trace events received for a given instance. */
-	TArray<TSharedPtr<FStateTreeDebuggerEventTreeElement>> EventsTreeElements;
-	
 	/** Tree view displaying the frame events of the instance associated to the selected track. */
-	TSharedPtr<STreeView<TSharedPtr<FStateTreeDebuggerEventTreeElement>>> EventsTreeView;
+	TSharedPtr<UE::StateTreeDebugger::SFrameEventsView> EventsView;
 
 	/** Attribute provided by the debugger scrub position to control cursor and timelines positions. */
 	TAttribute<double> ScrubTimeAttribute;
diff --git a/Engine/Plugins/Runtime/StateTree/Source/StateTreeEditorModule/Private/Debugger/SStateTreeDebuggerViewRow.cpp b/Engine/Plugins/Runtime/StateTree/Source/StateTreeEditorModule/Private/Debugger/SStateTreeDebuggerViewRow.cpp
deleted file mode 100644
diff --git a/Engine/Plugins/Runtime/StateTree/Source/StateTreeEditorModule/Private/Debugger/SStateTreeDebuggerViewRow.h b/Engine/Plugins/Runtime/StateTree/Source/StateTreeEditorModule/Private/Debugger/SStateTreeDebuggerViewRow.h
deleted file mode 100644
diff --git a/Engine/Plugins/Runtime/StateTree/Source/StateTreeEditorModule/Private/Debugger/StateTreeRewindDebuggerTrack.cpp b/Engine/Plugins/Runtime/StateTree/Source/StateTreeEditorModule/Private/Debugger/StateTreeRewindDebuggerTrack.cpp
--- a/Engine/Plugins/Runtime/StateTree/Source/StateTreeEditorModule/Private/Debugger/StateTreeRewindDebuggerTrack.cpp
+++ b/Engine/Plugins/Runtime/StateTree/Source/StateTreeEditorModule/Private/Debugger/StateTreeRewindDebuggerTrack.cpp
@@ -2,6 +2,7 @@
 
 #include "StateTreeRewindDebuggerTrack.h"
 #include "Debugger/SCompactTreeDebuggerView.h"
+#include "Debugger/SStateTreeFrameEventsView.h"
 #include "Debugger/StateTreeDebugger.h"
 #include "Debugger/StateTreeTraceProvider.h"
 #include "Editor.h"
@@ -17,7 +18,6 @@ namespace UE::StateTreeDebugger
 
 static const FLazyName StateTreeInstancesName("StateTreeInstances");
 
-
 //----------------------------------------------------------------------//
 // FRewindDebuggerTrackCreator
 //----------------------------------------------------------------------//
@@ -94,7 +94,7 @@ FRewindDebuggerTrack::FRewindDebuggerTrack(const RewindDebugger::FObjectId& InOb
 	}
 
 	check(Descriptor.IsValid());
-	SAssignNew(DebugTreeView, UE::StateTree::SCompactTreeDebuggerView, Descriptor->StateTree.Get())
+	SAssignNew(CompactTreeView, UE::StateTree::SCompactTreeDebuggerView, Descriptor->StateTree.Get())
 		.ActiveStates_Lambda([WeakTrackHelper = InstanceTrackHelper->AsWeak()]
 			{
 				if (const TSharedPtr<FInstanceTrackHelper> TrackHelper = WeakTrackHelper.Pin())
@@ -103,6 +103,8 @@ FRewindDebuggerTrack::FRewindDebuggerTrack(const RewindDebugger::FObjectId& InOb
 				}
 				return FStateTreeTraceActiveStates{};
 			});
+
+	SAssignNew(EventsView, UE::StateTreeDebugger::SFrameEventsView, Descriptor->StateTree.Get());
 }
 
 FText FRewindDebuggerTrack::GetDisplayNameInternal() const
@@ -171,13 +173,14 @@ bool FRewindDebuggerTrack::UpdateInternal()
 		if (!FMath::IsNearlyEqual(LastUpdateScrubTime, RewindDebugger->GetScrubTime()))
 		{
 			const double InstanceLastActiveTime = Descriptor->IsInstanceActive() ? RewindDebugger->GetRecordingDuration() : Descriptor->Lifetime.GetUpperBoundValue();
-			LastUpdateScrubTime = RewindDebugger->GetScrubTime();
+			const double ScrubTime = RewindDebugger->GetScrubTime();
+			LastUpdateScrubTime = ScrubTime;
 
 			FStateTreeTraceActiveStates ActiveStates;
 			const bool bTimelineDataChanged = InstanceTrackHelper.Get().RebuildEventData(StateTree
 				, EventCollection
 				, InstanceLastActiveTime
-				, RewindDebugger->GetScrubTime()
+				, ScrubTime
 				, &ActiveStates);
 
 			const bool bActivateStatesChanged = ActiveStates != InstanceTrackHelper.Get().GetActiveStates();
@@ -189,7 +192,12 @@ bool FRewindDebuggerTrack::UpdateInternal()
 			if (bTimelineDataChanged
 				|| bActivateStatesChanged)
 			{
-				DebugTreeView->Refresh();
+				CompactTreeView->Refresh();
+
+				FScrubState ScrubState({EventCollection});
+				ScrubState.SetScrubTime(ScrubTime);
+				ScrubState.SetEventCollectionIndex(0);
+				EventsView->RequestRefresh(ScrubState);
 				return true;
 			}
 		}
@@ -205,7 +213,18 @@ TSharedPtr<SWidget> FRewindDebuggerTrack::GetTimelineViewInternal()
 
 TSharedPtr<SWidget> FRewindDebuggerTrack::GetDetailsViewInternal()
 {
-	return DebugTreeView;
+	return SNew(SHorizontalBox)
+		+ SHorizontalBox::Slot()
+		.AutoWidth()
+		.VAlign(VAlign_Fill)
+		[
+			CompactTreeView.ToSharedRef()
+		]
+		+ SHorizontalBox::Slot()
+		.FillWidth(1.0f)
+		[
+			EventsView.ToSharedRef()
+		];
 }
 
 bool FRewindDebuggerTrack::HandleDoubleClickInternal()
diff --git a/Engine/Plugins/Runtime/StateTree/Source/StateTreeEditorModule/Private/Debugger/StateTreeRewindDebuggerTrack.h b/Engine/Plugins/Runtime/StateTree/Source/StateTreeEditorModule/Private/Debugger/StateTreeRewindDebuggerTrack.h
--- a/Engine/Plugins/Runtime/StateTree/Source/StateTreeEditorModule/Private/Debugger/StateTreeRewindDebuggerTrack.h
+++ b/Engine/Plugins/Runtime/StateTree/Source/StateTreeEditorModule/Private/Debugger/StateTreeRewindDebuggerTrack.h
@@ -14,6 +14,8 @@ class SCompactTreeDebuggerView;
 
 namespace UE::StateTreeDebugger
 {
+class SFrameEventsView;
+
 /**
  * RewindDebugger track creator for StateTree instances
  */
@@ -59,7 +61,8 @@ private:
 	FInstanceEventCollection EventCollection;
 	TSharedRef<FInstanceTrackHelper> InstanceTrackHelper;
 	TSharedPtr<const FInstanceDescriptor> Descriptor;
-	TSharedPtr<UE::StateTree::SCompactTreeDebuggerView> DebugTreeView;
+	TSharedPtr<StateTree::SCompactTreeDebuggerView> CompactTreeView;
+	TSharedPtr<SFrameEventsView> EventsView;
 	FSlateIcon Icon;
 	RewindDebugger::FObjectId ObjectId;
 
diff --git a/Engine/Plugins/Runtime/StateTree/Source/StateTreeEditorModule/Private/StateTreeEditorModeToolkit.cpp b/Engine/Plugins/Runtime/StateTree/Source/StateTreeEditorModule/Private/StateTreeEditorModeToolkit.cpp
--- a/Engine/Plugins/Runtime/StateTree/Source/StateTreeEditorModule/Private/StateTreeEditorModeToolkit.cpp
+++ b/Engine/Plugins/Runtime/StateTree/Source/StateTreeEditorModule/Private/StateTreeEditorModeToolkit.cpp
@@ -57,8 +57,8 @@ void FStateTreeEditorModeToolkit::RequestModeUITabs()
 		{
 			TSharedPtr<FAssetEditorModeUILayer> ModeUILayerPtr = ModeUILayer.Pin();
 			TSharedPtr<FWorkspaceItem> MenuGroup = ModeUILayerPtr->GetModeMenuCategory();
-			if(!MenuGroup)
-			{ 
+			if (!MenuGroup)
+			{
 				return;
 			}
 
@@ -218,7 +218,7 @@ FSlateIcon FStateTreeEditorModeToolkit::GetCompileStatusImage() const
 			return FSlateIcon(FAppStyle::GetAppStyleSetName(), CompileStatusBackground, NAME_None, CompileStatusGood);
 		}
 	}
-	
+
 	return FSlateIcon(FAppStyle::GetAppStyleSetName(), CompileStatusBackground, NAME_None, CompileStatusUnknown);
 }
 
@@ -276,7 +276,7 @@ void FStateTreeEditorModeToolkit::ExtendSecondaryModeToolbar(UToolMenu* ToolBar)
 
 		return nullptr;
 	};
-	
+
 	CompileSection.AddDynamicEntry("CompileCommands", FNewToolMenuSectionDelegate::CreateLambda([GetToolkitFromAssetEditorContext](FToolMenuSection& InSection)
 	{
 		if (UAssetEditorToolkitMenuContext* ToolkitContextObject = InSection.FindContext<UAssetEditorToolkitMenuContext>())
@@ -292,7 +292,7 @@ void FStateTreeEditorModeToolkit::ExtendSecondaryModeToolbar(UToolMenu* ToolBar)
 					{
 						return SharedToolkit->GetCompileStatusImage();
 					}
-				
+
 					return FSlateIcon();
 				}))
 			);
@@ -334,7 +334,7 @@ void FStateTreeEditorModeToolkit::ExtendSecondaryModeToolbar(UToolMenu* ToolBar)
 								 }
 							 }
 						}
-					 	
+
 						return false;
 					})),
 				 FOnGetContent::CreateLambda([ToolkitContextObject, GetToolkitFromAssetEditorContext]() -> TSharedRef<SWidget>
@@ -350,7 +350,7 @@ void FStateTreeEditorModeToolkit::ExtendSecondaryModeToolbar(UToolMenu* ToolBar)
 				 LOCTEXT("CreateNewTask_ToolbarTooltip", "Create a new Blueprint State Tree Task"),
 				 GetNewTaskButtonImage()
 			 ));
-				
+
 			InSection.AddEntry(FToolMenuEntry::InitComboButton(
 				 "CreateNewConditionComboButton",
 				 FUIAction(FExecuteAction(), FCanExecuteAction(), FIsActionChecked(),
@@ -369,7 +369,7 @@ void FStateTreeEditorModeToolkit::ExtendSecondaryModeToolbar(UToolMenu* ToolBar)
 								}
 							}
 						}
-				 		
+
 						return false;
 					})),
 				FOnGetContent::CreateLambda([ToolkitContextObject, GetToolkitFromAssetEditorContext]() -> TSharedRef<SWidget>
@@ -385,7 +385,7 @@ void FStateTreeEditorModeToolkit::ExtendSecondaryModeToolbar(UToolMenu* ToolBar)
 				 LOCTEXT("CreateNewCondition_ToolbarTooltip", "Create a new Blueprint State Tree Condition"),
 				 GetNewConditionButtonImage()
 			 ));
-				
+
 			 InSection.AddEntry(FToolMenuEntry::InitComboButton(
 				 "CreateNewConsiderationComboButton",
 				 FUIAction(FExecuteAction(), FCanExecuteAction(), FIsActionChecked(),
@@ -404,7 +404,7 @@ void FStateTreeEditorModeToolkit::ExtendSecondaryModeToolbar(UToolMenu* ToolBar)
 								 }
 							 }
 						}
-					 	
+
 						return false;
 					})),
 				 FOnGetContent::CreateLambda([ToolkitContextObject, GetToolkitFromAssetEditorContext]() -> TSharedRef<SWidget>
@@ -437,9 +437,9 @@ void FStateTreeEditorModeToolkit::OnStateTreeChanged()
 {
 	// Update underlying state tree state
 	UpdateStateTreeOutliner();
-	
+
 #if WITH_STATETREE_TRACE_DEBUGGER
-	UpdateDebuggerView();	
+	UpdateDebuggerView();
 #endif // WITH_STATETREE_TRACE_DEBUGGER
 }
 
@@ -627,10 +627,10 @@ void FStateTreeEditorModeToolkit::UpdateStateTreeOutliner()
 			if (UStateTreeEditingSubsystem* Subsystem = GEditor->GetEditorSubsystem<UStateTreeEditingSubsystem>())
 			{
 				StateTreeOutliner = SNew(SStateTreeOutliner, Subsystem->FindOrAddViewModel(EditorMode->GetStateTree()), GetToolkitCommands());
-			}		
+			}
 		}
 	}
-	
+
 	if(TSharedPtr<SDockTab> SharedOutlinerTab = WeakOutlinerTab.Pin())
 	{
 		SharedOutlinerTab->SetContent(StateTreeOutliner.ToSharedRef());
@@ -657,7 +657,7 @@ void FStateTreeEditorModeToolkit::UpdateDebuggerView()
 		{
 			if (UStateTreeEditingSubsystem* Subsystem = GEditor->GetEditorSubsystem<UStateTreeEditingSubsystem>())
 			{
-				DebuggerView = SNew(SStateTreeDebuggerView, *StateTree, Subsystem->FindOrAddViewModel(StateTree), GetToolkitCommands());
+				DebuggerView = SNew(SStateTreeDebuggerView, StateTree, Subsystem->FindOrAddViewModel(StateTree), GetToolkitCommands());
 			}
 		}
 	}
diff --git a/Engine/Plugins/Runtime/StateTree/Source/StateTreeEditorModule/Private/StateTreeEditorStyle.cpp b/Engine/Plugins/Runtime/StateTree/Source/StateTreeEditorModule/Private/StateTreeEditorStyle.cpp
--- a/Engine/Plugins/Runtime/StateTree/Source/StateTreeEditorModule/Private/StateTreeEditorStyle.cpp
+++ b/Engine/Plugins/Runtime/StateTree/Source/StateTreeEditorModule/Private/StateTreeEditorStyle.cpp
@@ -12,31 +12,6 @@
 #include "StateTreeTypes.h"
 #include "Styling/SlateStyleMacros.h"
 
-
-namespace UE::StateTree::Editor
-{
-
-class FContentRootScope
-{
-public:
-	FContentRootScope(FStateTreeEditorStyle* InStyle, const FString& NewContentRoot)
-		: Style(InStyle)
-		, PreviousContentRoot(InStyle->GetContentRootDir())
-	{
-		Style->SetContentRoot(NewContentRoot);
-	}
-
-	~FContentRootScope()
-	{
-		Style->SetContentRoot(PreviousContentRoot);
-	}
-private:
-	FStateTreeEditorStyle* Style;
-	FString PreviousContentRoot;
-};
-
-}; // UE::StateTree::Editor
-
 FStateTreeEditorStyle::FStateTreeEditorStyle()
 	: FStateTreeStyle(TEXT("StateTreeEditorStyle"))
 {
@@ -54,7 +29,7 @@ FStateTreeEditorStyle::FStateTreeEditorStyle()
 			.SetBackgroundImageHovered(CORE_BOX_BRUSH("Common/TextBox_Hovered", FMargin(4.0f / 16.0f)))
 			.SetBackgroundImageFocused(CORE_BOX_BRUSH("Common/TextBox_Hovered", FMargin(4.0f / 16.0f)))
 			.SetBackgroundImageReadOnly(CORE_BOX_BRUSH("Common/TextBox_ReadOnly", FMargin(4.0f / 16.0f)))
-			.SetBackgroundColor(FLinearColor(0,0,0,0.1f))
+			.SetBackgroundColor(FLinearColor(0.0f, 0.0f, 0.0f, 0.1f))
 			.SetPadding(FMargin(0))
 			.SetScrollBarStyle(ScrollBar);
 		Set("StateTree.State.TitleEditableText", StateTitleEditableText);
@@ -85,7 +60,7 @@ FStateTreeEditorStyle::FStateTreeEditorStyle()
 
 	// Task
 	{
-		const FLinearColor ForegroundCol =  FStyleColors::Foreground.GetSpecifiedColor();
+		const FLinearColor ForegroundCol = FStyleColors::Foreground.GetSpecifiedColor();
 
 		Set("StateTree.Task.Title", FTextBlockStyle(NormalText)
 			.SetFont(DEFAULT_FONT("Regular", 10))
@@ -121,7 +96,7 @@ FStateTreeEditorStyle::FStateTreeEditorStyle()
 
 	// Transition rich text
 	{
-		const FLinearColor ForegroundCol =  FStyleColors::White.GetSpecifiedColor();
+		const FLinearColor ForegroundCol = FStyleColors::White.GetSpecifiedColor();
 		Set("Transition.Normal", FTextBlockStyle(NormalText)
 			.SetColorAndOpacity(ForegroundCol.CopyWithNewOpacity(0.9f))
 			.SetFont(DEFAULT_FONT("Regular", 11)));
@@ -150,22 +125,6 @@ FStateTreeEditorStyle::FStateTreeEditorStyle()
 		Set("DiffTools.Properties", FLinearColor(0.2f, 0.4f, 1.f)); // blue
 	}
 
-	// Debugger
-	{
-		Set("StateTreeDebugger.Element.Normal",
-			FTextBlockStyle(NormalText)
-			.SetFont(DEFAULT_FONT("Regular", 10)));
-
-		Set("StateTreeDebugger.Element.Bold",
-			FTextBlockStyle(NormalText)
-			.SetFont(DEFAULT_FONT("Bold", 10)));
-
-		Set("StateTreeDebugger.Element.Subdued",
-			FTextBlockStyle(NormalText)
-			.SetFont(DEFAULT_FONT("Regular", 10))
-			.SetColorAndOpacity(FSlateColor::UseSubduedForeground()));
-	}
-	
 	const FLinearColor SelectionColor = FColor(0, 0, 0, 32);
 	const FTableRowStyle& NormalTableRowStyle = FAppStyle::Get().GetWidgetStyle<FTableRowStyle>("TableView.Row");
 	Set("StateTree.Selection",
@@ -207,7 +166,7 @@ FStateTreeEditorStyle::FStateTreeEditorStyle()
 		.SetFontSize(7));
 
 	Set("StateTree.Param.Background", new FSlateRoundedBoxBrush(FStyleColors::Hover, 6.f));
-	
+
 	// Expression Indent combo button
 	const FButtonStyle IndentButton = FButtonStyle()
 		.SetNormal(FSlateRoundedBoxBrush(FLinearColor::Transparent, 2.0f))
@@ -218,10 +177,9 @@ FStateTreeEditorStyle::FStateTreeEditorStyle()
 		.SetPressedForeground(FStyleColors::Foreground)
 		.SetNormalPadding(FMargin(2, 2, 2, 2))
 		.SetPressedPadding(FMargin(2, 3, 2, 1));
-	
+
 	Set("StateTree.Node.Indent.ComboBox", FComboButtonStyle(ComboButtonStyle).SetButtonStyle(IndentButton));
 
-	
 	// Node text styles
 	{
 		FEditableTextStyle EditableTextStyle = FEditableTextStyle(FAppStyle::GetWidgetStyle<FEditableTextStyle>("NormalEditableText"));
@@ -256,7 +214,7 @@ FStateTreeEditorStyle::FStateTreeEditorStyle()
 	// Command icons
 	{
 		// From generic Engine
-		UE::StateTree::Editor::FContentRootScope Scope(this, EngineSlateContentDir);
+		FContentRootScope Scope(this, EngineSlateContentDir);
 		Set("StateTreeEditor.CutStates", new IMAGE_BRUSH_SVG("Starship/Common/Cut", CoreStyleConstants::Icon16x16));
 		Set("StateTreeEditor.CopyStates", new IMAGE_BRUSH_SVG("Starship/Common/Copy", CoreStyleConstants::Icon16x16));
 		Set("StateTreeEditor.DuplicateStates", new IMAGE_BRUSH_SVG("Starship/Common/Duplicate", CoreStyleConstants::Icon16x16));
@@ -266,26 +224,11 @@ FStateTreeEditorStyle::FStateTreeEditorStyle()
 
 		Set("StateTreeEditor.Debugger.ResetTracks", new IMAGE_BRUSH_SVG("Starship/Common/Delete", CoreStyleConstants::Icon16x16));
 
-		Set("StateTreeEditor.Debugger.State.Enter", new CORE_IMAGE_BRUSH_SVG("Starship/Common/arrow-right", CoreStyleConstants::Icon16x16, FStyleColors::Foreground));
-		Set("StateTreeEditor.Debugger.State.Exit", new CORE_IMAGE_BRUSH_SVG("Starship/Common/arrow-left", CoreStyleConstants::Icon16x16, FStyleColors::Foreground));
-		Set("StateTreeEditor.Debugger.State.Selected", new CORE_IMAGE_BRUSH_SVG("Starship/Common/arrow-right", CoreStyleConstants::Icon16x16, FStyleColors::AccentYellow));
-		Set("StateTreeEditor.Debugger.State.Completed", new CORE_IMAGE_BRUSH_SVG("Starship/Common/check", CoreStyleConstants::Icon16x16, FStyleColors::AccentGreen));
-
 		Set("StateTreeEditor.Debugger.Task.Enter", new CORE_IMAGE_BRUSH_SVG("Starship/Common/arrow-right", CoreStyleConstants::Icon16x16, FStyleColors::Foreground));
 		Set("StateTreeEditor.Debugger.Task.Exit", new CORE_IMAGE_BRUSH_SVG("Starship/Common/arrow-left", CoreStyleConstants::Icon16x16, FStyleColors::Foreground));
-		Set("StateTreeEditor.Debugger.Task.Failed", new CORE_IMAGE_BRUSH_SVG("Starship/Common/close-small", CoreStyleConstants::Icon16x16, FStyleColors::AccentRed));
-		Set("StateTreeEditor.Debugger.Task.Succeeded", new CORE_IMAGE_BRUSH_SVG("Starship/Common/check", CoreStyleConstants::Icon16x16, FStyleColors::AccentGreen));
-		Set("StateTreeEditor.Debugger.Task.Stopped", new CORE_IMAGE_BRUSH_SVG("Starship/Common/close-small", CoreStyleConstants::Icon16x16, FStyleColors::AccentRed));
-
-		Set("StateTreeEditor.Debugger.Condition.Passed", new CORE_IMAGE_BRUSH_SVG("Starship/Common/check", CoreStyleConstants::Icon16x16, FStyleColors::AccentGreen));
-		Set("StateTreeEditor.Debugger.Condition.Failed", new CORE_IMAGE_BRUSH_SVG("Starship/Common/close-small", CoreStyleConstants::Icon16x16, FStyleColors::AccentRed));
-		Set("StateTreeEditor.Debugger.Condition.OnEvaluating", new CORE_IMAGE_BRUSH_SVG("Starship/Common/Update", CoreStyleConstants::Icon16x16, FStyleColors::AccentYellow));
 
 		Set("StateTreeEditor.Debugger.Unset", new CORE_IMAGE_BRUSH_SVG("Starship/Common/help", CoreStyleConstants::Icon16x16, FStyleColors::AccentBlack));
 
-		Set("StateTreeEditor.Debugger.Log.Warning", new CORE_IMAGE_BRUSH_SVG("Starship/Common/alert-circle", CoreStyleConstants::Icon16x16, FStyleColors::AccentYellow));
-		Set("StateTreeEditor.Debugger.Log.Error", new CORE_IMAGE_BRUSH_SVG("Starship/Common/x-circle", CoreStyleConstants::Icon16x16, FStyleColors::AccentRed));
-
 		// Common Node Icons
 		Set("Node.EnableDisable", new CORE_IMAGE_BRUSH_SVG("Starship/Common/check-circle", CoreStyleConstants::Icon16x16));
 		Set("Node.Time", new CORE_IMAGE_BRUSH_SVG("Starship/Common/Recent", CoreStyleConstants::Icon16x16));
@@ -294,11 +237,11 @@ FStateTreeEditorStyle::FStateTreeEditorStyle()
 
 	{
 		// From generic Engine Editor
-		UE::StateTree::Editor::FContentRootScope Scope(this, EngineEditorSlateContentDir);
+		FContentRootScope Scope(this, EngineEditorSlateContentDir);
 
 		Set("StateTreeEditor.Debugger.StartRecording", new IMAGE_BRUSH("Sequencer/Transport_Bar/Record_24x", CoreStyleConstants::Icon16x16));
 		Set("StateTreeEditor.Debugger.StopRecording", new IMAGE_BRUSH("Sequencer/Transport_Bar/Recording_24x", CoreStyleConstants::Icon16x16));
-		
+
 		Set("StateTreeEditor.Debugger.PreviousFrameWithStateChange", new IMAGE_BRUSH("Sequencer/Transport_Bar/Go_To_Front_24x", CoreStyleConstants::Icon16x16));
 		Set("StateTreeEditor.Debugger.PreviousFrameWithEvents", new IMAGE_BRUSH("Sequencer/Transport_Bar/Step_Backwards_24x", CoreStyleConstants::Icon16x16));
 		Set("StateTreeEditor.Debugger.NextFrameWithEvents", new IMAGE_BRUSH("Sequencer/Transport_Bar/Step_Forward_24x", CoreStyleConstants::Icon16x16));
@@ -313,7 +256,7 @@ FStateTreeEditorStyle::FStateTreeEditorStyle()
 		Set("StateTreeEditor.Debugger.InstanceTrack", new IMAGE_BRUSH_SVG("Starship/AssetIcons/AnimInstance_64", CoreStyleConstants::Icon16x16));
 
 		Set("StateTreeEditor.EnableStates", new IMAGE_BRUSH("Icons/Empty_16x", CoreStyleConstants::Icon16x16));
-		Set("StateTreeEditor.Debugger.Breakpoint.EnabledAndValid", new IMAGE_BRUSH_SVG( "Starship/Blueprints/Breakpoint_Valid", CoreStyleConstants::Icon16x16, FStyleColors::AccentRed));
+		Set("StateTreeEditor.Debugger.Breakpoint.EnabledAndValid", new IMAGE_BRUSH_SVG("Starship/Blueprints/Breakpoint_Valid", CoreStyleConstants::Icon16x16, FStyleColors::AccentRed));
 		Set("StateTreeEditor.Debugger.ResumeDebuggerAnalysis", new IMAGE_BRUSH_SVG("Starship/Common/Timeline", CoreStyleConstants::Icon16x16));
 
 		Set("StateTreeEditor.Transition.None", new CORE_IMAGE_BRUSH_SVG("Starship/Common/x-circle", CoreStyleConstants::Icon16x16, FSlateColor::UseSubduedForeground()));
@@ -366,7 +309,6 @@ FStateTreeEditorStyle::FStateTreeEditorStyle()
 		Set("StateTreeEditor.Transition.Parent", new IMAGE_BRUSH_SVG("Icons/Transition_Parent", CoreStyleConstants::Icon16x16, FStyleColors::Foreground));
 
 		Set("StateTreeEditor.Transition.Condition", new IMAGE_BRUSH_SVG("Icons/State_Conditions", CoreStyleConstants::Icon16x16, FStyleColors::AccentGray));
-		Set("StateTreeEditor.Debugger.Condition.OnTransition", new IMAGE_BRUSH_SVG("Icons/State_Conditions", CoreStyleConstants::Icon16x16, FStyleColors::AccentGray));
 
 		// Common Node Icons
 		Set("Node.Movement", new IMAGE_BRUSH_SVG("Icons/Movement", CoreStyleConstants::Icon16x16));
diff --git a/Engine/Plugins/Runtime/StateTree/Source/StateTreeModule/Private/Debugger/StateTreeDebugger.cpp b/Engine/Plugins/Runtime/StateTree/Source/StateTreeModule/Private/Debugger/StateTreeDebugger.cpp
--- a/Engine/Plugins/Runtime/StateTree/Source/StateTreeModule/Private/Debugger/StateTreeDebugger.cpp
+++ b/Engine/Plugins/Runtime/StateTree/Source/StateTreeModule/Private/Debugger/StateTreeDebugger.cpp
@@ -27,7 +27,7 @@
 //----------------------------------------------------------------//
 namespace UE::StateTreeDebugger
 {
-	struct FDiagnosticsSessionAnalyzer : public UE::Trace::IAnalyzer
+	struct FDiagnosticsSessionAnalyzer : UE::Trace::IAnalyzer
 	{
 		virtual void OnAnalysisBegin(const FOnAnalysisContext& Context) override
 		{
@@ -93,6 +93,7 @@ FStateTreeDebugger::FStateTreeDebugger()
 		});
 }
 
+PRAGMA_DISABLE_DEPRECATION_WARNINGS
 FStateTreeDebugger::~FStateTreeDebugger()
 {
 	UE::StateTree::Delegates::OnTracingStateChanged.Remove(TracingStateChangedHandle);
@@ -103,6 +104,7 @@ FStateTreeDebugger::~FStateTreeDebugger()
 
 	StopSessionAnalysis();
 }
+PRAGMA_ENABLE_DEPRECATION_WARNINGS
 
 void FStateTreeDebugger::Tick(const float DeltaTime)
 {
@@ -191,9 +193,6 @@ void FStateTreeDebugger::SelectInstance(const FStateTreeInstanceDebugId Instance
 	{
 		SelectedInstanceId = InstanceId;
 
-		// Notify so listener can clean up anything related to previous instance
-		OnSelectedInstanceCleared.ExecuteIfBound();
-
 		// Update event collection index for newly debugged instance
 		SetScrubStateCollectionIndex(InstanceId.IsValid()
 			? EventCollections.IndexOfByPredicate([InstanceId = InstanceId](const UE::StateTreeDebugger::FInstanceEventCollection& Entry)
@@ -429,7 +428,7 @@ void FStateTreeDebugger::UpdateMetadata(FTraceDescriptor& TraceDescriptor) const
 	}
 
 	// inspired from FStoreBrowser
-	struct FDataStream : public UE::Trace::IInDataStream
+	struct FDataStream : UE::Trace::IInDataStream
 	{
 		enum class EReadStatus
 		{
diff --git a/Engine/Plugins/Runtime/StateTree/Source/StateTreeModule/Public/Debugger/StateTreeDebugger.h b/Engine/Plugins/Runtime/StateTree/Source/StateTreeModule/Public/Debugger/StateTreeDebugger.h
--- a/Engine/Plugins/Runtime/StateTree/Source/StateTreeModule/Public/Debugger/StateTreeDebugger.h
+++ b/Engine/Plugins/Runtime/StateTree/Source/StateTreeModule/Public/Debugger/StateTreeDebugger.h
@@ -75,6 +75,8 @@ DECLARE_DELEGATE_TwoParams(FOnStateTreeDebuggerBreakpointHit, FStateTreeInstance
 DECLARE_DELEGATE_OneParam(FOnStateTreeDebuggerActiveStatesChanges, const FStateTreeTraceActiveStates& ActiveStates);
 DECLARE_DELEGATE_OneParam(FOnStateTreeDebuggerNewInstance, FStateTreeInstanceDebugId InstanceId);
 DECLARE_DELEGATE(FOnStateTreeDebuggerNewSession);
+
+UE_DEPRECATED(5.7, "Delegate type no longer used and will be removed")
 DECLARE_DELEGATE(FOnStateTreeDebuggerDebuggedInstanceSet);
 
 struct FStateTreeDebugger : FTickableGameObject, UE::StateTreeDebugger::ITraceReader
@@ -85,7 +87,7 @@ struct FStateTreeDebugger : FTickableGameObject, UE::StateTreeDebugger::ITraceRe
 		FTraceDescriptor(const FString& Name, const uint32 Id) : Name(Name), TraceId(Id)
 		{
 		}
-		
+
 		bool operator==(const FTraceDescriptor& Other) const
 		{
 			return Other.TraceId == TraceId;
@@ -145,10 +147,10 @@ struct FStateTreeDebugger : FTickableGameObject, UE::StateTreeDebugger::ITraceRe
 	{
 		StateTreeAsset = Asset;
 	}
-	
+
 	/** Forces a single refresh to latest state. Useful when simulation is paused. */
 	UE_API void SyncToCurrentSessionDuration();
-	
+
 	UE_API bool CanStepBackToPreviousStateWithEvents() const;
 	UE_API void StepBackToPreviousStateWithEvents();
 
@@ -214,6 +216,11 @@ struct FStateTreeDebugger : FTickableGameObject, UE::StateTreeDebugger::ITraceRe
 		return HitBreakpoint.IsSet();
 	}
 
+	bool CanProcessBreakpoints() const
+	{
+		return OnBreakpointHit.IsBound();
+	}
+
 	static UE_API FText DescribeTrace(const FTraceDescriptor& TraceDescriptor);
 	static UE_API FText DescribeInstance(const UE::StateTreeDebugger::FInstanceDescriptor& StateTreeInstanceDesc);
 
@@ -327,7 +334,12 @@ struct FStateTreeDebugger : FTickableGameObject, UE::StateTreeDebugger::ITraceRe
 
 	FOnStateTreeDebuggerNewSession OnNewSession;
 	FOnStateTreeDebuggerNewInstance OnNewInstance;
+
+PRAGMA_DISABLE_DEPRECATION_WARNINGS
+	UE_DEPRECATED(5.7, "Delegate no longer used and will be removed")
 	FOnStateTreeDebuggerDebuggedInstanceSet OnSelectedInstanceCleared;
+PRAGMA_ENABLE_DEPRECATION_WARNINGS
+
 	FOnStateTreeDebuggerScrubStateChanged OnScrubStateChanged;
 	FOnStateTreeDebuggerBreakpointHit OnBreakpointHit;
 	FOnStateTreeDebuggerActiveStatesChanges OnActiveStatesChanged;
@@ -354,7 +366,7 @@ protected:
 	{
 		RETURN_QUICK_DECLARE_CYCLE_STAT(FStateTreeDebugger, STATGROUP_Tickables);
 	}
-	
+
 private:
 	enum class EAnalysisSourceType : uint8
 	{
@@ -377,7 +389,7 @@ private:
 	};
 
 	void UpdateAnalysisTransitionType(EAnalysisSourceType SourceType);
-	
+
 	void ReadTrace(double ScrubTime);
 	void ReadTrace(uint64 FrameIndex);
 
@@ -429,9 +441,9 @@ private:
 	/** Module used to access the store client and analysis sessions .*/
 	IStateTreeModule& StateTreeModule;
 
-	/** The StateTree asset associated to this debugger. All instances will be using this asset. */ 
+	/** The StateTree asset associated to this debugger. All instances will be using this asset. */
 	TWeakObjectPtr<const UStateTree> StateTreeAsset;
-	
+
 	/** The trace analysis session. */
 	TSharedPtr<const TraceServices::IAnalysisSession> AnalysisSession;
 
@@ -455,13 +467,13 @@ private:
 	 * This is to keep track of the previous last live session id so we can detect when the new one is available.
 	 */
 	int32 LastLiveSessionId = INDEX_NONE;
-	
+
 	/**
 	 * When auto-connecting on next live session it is possible that a few frames are required for the tracing session to be accessible and connected to.
 	 * This is to keep track of the time window where we will retry.
 	 */
 	float RetryLoadNextLiveSessionTimer = 0.0f;
-	
+
 	/**
 	 * Most recent processed world time. This value is updated from received events so it might increase by large deltas.
 	 */
@@ -473,7 +485,7 @@ private:
 	/** Last time in the recording that we use to fetch events and we will use for the next read. */
 	double LastTraceReadTime = 0;
 
-	/** Combined information regarding current scrub time (e.g. frame index, event collection index, etc.) */ 
+	/** Combined information regarding current scrub time (e.g. frame index, event collection index, etc.) */
 	UE::StateTreeDebugger::FScrubState ScrubState;
 
 	/** Information stored when a breakpoint is hit while processing events and used to send notifications. */
