[StateTreeDebugger] disable actions to add breakpoints if debugger tabs is not opened; clear them when closing the debugger tab

diff --git a/Engine/Plugins/Runtime/StateTree/Source/StateTreeEditorModule/Private/Debugger/SStateTreeDebuggerView.cpp b/Engine/Plugins/Runtime/StateTree/Source/StateTreeEditorModule/Private/Debugger/SStateTreeDebuggerView.cpp
--- a/Engine/Plugins/Runtime/StateTree/Source/StateTreeEditorModule/Private/Debugger/SStateTreeDebuggerView.cpp
+++ b/Engine/Plugins/Runtime/StateTree/Source/StateTreeEditorModule/Private/Debugger/SStateTreeDebuggerView.cpp
@@ -28,6 +28,11 @@
 //----------------------------------------------------------------------//
 SStateTreeDebuggerView::~SStateTreeDebuggerView()
 {
+	if (StateTreeViewModel)
+	{
+		StateTreeViewModel->RemoveAllBreakpoints();
+	}
+
 	UE::StateTree::Delegates::OnTracingStateChanged.RemoveAll(this);
 
 	check(Debugger);
diff --git a/Engine/Plugins/Runtime/StateTree/Source/StateTreeEditorModule/Private/SStateTreeView.cpp b/Engine/Plugins/Runtime/StateTree/Source/StateTreeEditorModule/Private/SStateTreeView.cpp
--- a/Engine/Plugins/Runtime/StateTree/Source/StateTreeEditorModule/Private/SStateTreeView.cpp
+++ b/Engine/Plugins/Runtime/StateTree/Source/StateTreeEditorModule/Private/SStateTreeView.cpp
@@ -268,16 +268,19 @@ void SStateTreeView::BindCommands()
 				StateTreeViewModel->HandleEnableStateBreakpoint(EStateTreeBreakpointType::OnEnter);
 			}
 		}),
-		FCanExecuteAction(),
+		FCanExecuteAction::CreateSPLambda(this, [this]
+		{
+			return (StateTreeViewModel)
+				&& (StateTreeViewModel->CanAddStateBreakpoint(EStateTreeBreakpointType::OnEnter)
+					|| StateTreeViewModel->CanRemoveStateBreakpoint(EStateTreeBreakpointType::OnEnter));
+		}),
 		FGetActionCheckState::CreateSPLambda(this, [this]
 		{
 			return StateTreeViewModel ? StateTreeViewModel->GetStateBreakpointCheckState(EStateTreeBreakpointType::OnEnter) : ECheckBoxState::Unchecked;
 		}),
 		FIsActionButtonVisible::CreateSPLambda(this, [this]
 		{
-			return (StateTreeViewModel)
-				&& (StateTreeViewModel->CanAddStateBreakpoint(EStateTreeBreakpointType::OnEnter)
-					|| StateTreeViewModel->CanRemoveStateBreakpoint(EStateTreeBreakpointType::OnEnter));
+			return StateTreeViewModel.IsValid();
 		}));
 
 	CommandList->MapAction(
@@ -289,16 +292,19 @@ void SStateTreeView::BindCommands()
 				StateTreeViewModel->HandleEnableStateBreakpoint(EStateTreeBreakpointType::OnExit);
 			}
 		}),
-		FCanExecuteAction(),
+		FCanExecuteAction::CreateSPLambda(this, [this]
+		{
+			return (StateTreeViewModel)
+				&& (StateTreeViewModel->CanAddStateBreakpoint(EStateTreeBreakpointType::OnExit)
+					|| StateTreeViewModel->CanRemoveStateBreakpoint(EStateTreeBreakpointType::OnExit));
+		}),
 		FGetActionCheckState::CreateSPLambda(this, [this]
 		{
 			return StateTreeViewModel ? StateTreeViewModel->GetStateBreakpointCheckState(EStateTreeBreakpointType::OnExit) : ECheckBoxState::Unchecked;
 		}),
 		FIsActionButtonVisible::CreateSPLambda(this, [this]
 		{
-			return (StateTreeViewModel)
-				&& (StateTreeViewModel->CanAddStateBreakpoint(EStateTreeBreakpointType::OnExit)
-					|| StateTreeViewModel->CanRemoveStateBreakpoint(EStateTreeBreakpointType::OnExit));
+			return StateTreeViewModel.IsValid();
 		}));
 #endif // WITH_STATETREE_TRACE_DEBUGGER
 }
diff --git a/Engine/Plugins/Runtime/StateTree/Source/StateTreeEditorModule/Private/StateTreeEditorCommands.cpp b/Engine/Plugins/Runtime/StateTree/Source/StateTreeEditorModule/Private/StateTreeEditorCommands.cpp
--- a/Engine/Plugins/Runtime/StateTree/Source/StateTreeEditorModule/Private/StateTreeEditorCommands.cpp
+++ b/Engine/Plugins/Runtime/StateTree/Source/StateTreeEditorModule/Private/StateTreeEditorCommands.cpp
@@ -35,8 +35,8 @@ void FStateTreeEditorCommands::RegisterCommands()
 	UI_COMMAND(EnableStates, "State Enabled", "Enables Selected States", EUserInterfaceActionType::Check, FInputChord());
 
 #if WITH_STATETREE_TRACE_DEBUGGER
-	UI_COMMAND(EnableOnEnterStateBreakpoint, "Break on Enter", "Adds or removes a breakpoint when entering the selected state(s)", EUserInterfaceActionType::Check, FInputChord(EKeys::F9));
-	UI_COMMAND(EnableOnExitStateBreakpoint, "Break on Exit", "Adds or removes a breakpoint when exiting the selected state(s)", EUserInterfaceActionType::Check, FInputChord(EModifierKey::Shift, EKeys::F9));
+	UI_COMMAND(EnableOnEnterStateBreakpoint, "Break on Enter", "Adds or removes a breakpoint when entering the selected state(s). (Debugger Window required)", EUserInterfaceActionType::Check, FInputChord(EKeys::F9));
+	UI_COMMAND(EnableOnExitStateBreakpoint, "Break on Exit", "Adds or removes a breakpoint when exiting the selected state(s). (Debugger Window required)", EUserInterfaceActionType::Check, FInputChord(EModifierKey::Shift, EKeys::F9));
 #endif //WITH_STATETREE_TRACE_DEBUGGER
 }
 
diff --git a/Engine/Plugins/Runtime/StateTree/Source/StateTreeEditorModule/Private/StateTreeEditorData.cpp b/Engine/Plugins/Runtime/StateTree/Source/StateTreeEditorModule/Private/StateTreeEditorData.cpp
--- a/Engine/Plugins/Runtime/StateTree/Source/StateTreeEditorModule/Private/StateTreeEditorData.cpp
+++ b/Engine/Plugins/Runtime/StateTree/Source/StateTreeEditorModule/Private/StateTreeEditorData.cpp
@@ -1680,6 +1680,15 @@ bool UStateTreeEditorData::RemoveBreakpoint(const FGuid ID, const EStateTreeBrea
 	return Index != INDEX_NONE;
 }
 
+void UStateTreeEditorData::RemoveAllBreakpoints()
+{
+	Breakpoints.Reset();
+
+	const UStateTree* StateTree = GetTypedOuter<UStateTree>();
+	checkf(StateTree, TEXT("UStateTreeEditorData should only be allocated within a UStateTree"));
+	UE::StateTree::Delegates::OnBreakpointsChanged.Broadcast(*StateTree);
+}
+
 #endif // WITH_STATETREE_TRACE_DEBUGGER
 
 #undef LOCTEXT_NAMESPACE
\ No newline at end of file
diff --git a/Engine/Plugins/Runtime/StateTree/Source/StateTreeEditorModule/Private/StateTreeViewModel.cpp b/Engine/Plugins/Runtime/StateTree/Source/StateTreeEditorModule/Private/StateTreeViewModel.cpp
--- a/Engine/Plugins/Runtime/StateTree/Source/StateTreeEditorModule/Private/StateTreeViewModel.cpp
+++ b/Engine/Plugins/Runtime/StateTree/Source/StateTreeEditorModule/Private/StateTreeViewModel.cpp
@@ -442,6 +442,11 @@ void FStateTreeViewModel::HandleIdentifierChanged(const UStateTree& StateTree) c
 #if WITH_STATETREE_TRACE_DEBUGGER
 bool FStateTreeViewModel::CanAddStateBreakpoint(const EStateTreeBreakpointType Type) const
 {
+	if (!Debugger.Get().CanProcessBreakpoints())
+	{
+		return false;
+	}
+
 	const UStateTreeEditorData* EditorData = TreeDataWeak.Get();
 	if (!ensure(EditorData != nullptr))
 	{
@@ -635,6 +640,16 @@ void FStateTreeViewModel::HandlePostCompile(const UStateTree& StateTree)
 	}
 }
 
+void FStateTreeViewModel::RemoveAllBreakpoints()
+{
+	if (UStateTreeEditorData* TreeData = TreeDataWeak.Get())
+	{
+		Debugger->ClearAllBreakpoints();
+
+		TreeData->RemoveAllBreakpoints();
+	}
+}
+
 void FStateTreeViewModel::RefreshDebuggerBreakpoints()
 {
 	const UStateTree* StateTree = GetStateTree();
diff --git a/Engine/Plugins/Runtime/StateTree/Source/StateTreeEditorModule/Public/StateTreeEditorData.h b/Engine/Plugins/Runtime/StateTree/Source/StateTreeEditorModule/Public/StateTreeEditorData.h
--- a/Engine/Plugins/Runtime/StateTree/Source/StateTreeEditorModule/Public/StateTreeEditorData.h
+++ b/Engine/Plugins/Runtime/StateTree/Source/StateTreeEditorModule/Public/StateTreeEditorData.h
@@ -309,6 +309,7 @@ PRAGMA_ENABLE_DEPRECATION_WARNINGS
 	UE_API const FStateTreeEditorBreakpoint* GetBreakpoint(FGuid ID, EStateTreeBreakpointType BreakpointType) const;
 	UE_API void AddBreakpoint(FGuid ID, EStateTreeBreakpointType BreakpointType);
 	UE_API bool RemoveBreakpoint(FGuid ID, EStateTreeBreakpointType BreakpointType);
+	UE_API void RemoveAllBreakpoints();
 #endif // WITH_STATETREE_TRACE_DEBUGGER
 
 	// ~StateTree Builder API
diff --git a/Engine/Plugins/Runtime/StateTree/Source/StateTreeEditorModule/Public/StateTreeViewModel.h b/Engine/Plugins/Runtime/StateTree/Source/StateTreeEditorModule/Public/StateTreeViewModel.h
--- a/Engine/Plugins/Runtime/StateTree/Source/StateTreeEditorModule/Public/StateTreeViewModel.h
+++ b/Engine/Plugins/Runtime/StateTree/Source/StateTreeEditorModule/Public/StateTreeViewModel.h
@@ -131,7 +131,12 @@ public:
 
 	UE_API UStateTreeState* FindStateAssociatedToBreakpoint(FStateTreeDebuggerBreakpoint Breakpoint) const;
 
-	TSharedRef<FStateTreeDebugger> GetDebugger() const { return Debugger; }
+	TSharedRef<FStateTreeDebugger> GetDebugger() const
+	{
+		return Debugger;
+	}
+
+	UE_API void RemoveAllBreakpoints();
 	UE_API void RefreshDebuggerBreakpoints();
 #endif // WITH_STATETREE_TRACE_DEBUGGER
 
