[StateTreeDebugger] updated 'Resume Analysis' button tooltip so user understands when the button is enabled

diff --git a/Engine/Plugins/Runtime/StateTree/Source/StateTreeEditorModule/Private/Debugger/StateTreeDebuggerCommands.cpp b/Engine/Plugins/Runtime/StateTree/Source/StateTreeEditorModule/Private/Debugger/StateTreeDebuggerCommands.cpp
--- a/Engine/Plugins/Runtime/StateTree/Source/StateTreeEditorModule/Private/Debugger/StateTreeDebuggerCommands.cpp
+++ b/Engine/Plugins/Runtime/StateTree/Source/StateTreeEditorModule/Private/Debugger/StateTreeDebuggerCommands.cpp
@@ -24,12 +24,12 @@ void FStateTreeDebuggerCommands::RegisterCommands()
 {
 	UI_COMMAND(StartRecording, "Start Recording", "Start a new trace session.", EUserInterfaceActionType::Button, FInputChord(EModifierKey::Control | EModifierKey::Shift, EKeys::R));
 	UI_COMMAND(StopRecording, "Stop Recording", "Stop the current trace session.", EUserInterfaceActionType::Button, FInputChord(EModifierKey::Control | EModifierKey::Shift, EKeys::S));
-	
+
 	UI_COMMAND(PreviousFrameWithStateChange, "Previous State Change", "Jump to previous state changed frame", EUserInterfaceActionType::Button, FInputChord(EModifierKey::Control, EKeys::Left));
 	UI_COMMAND(PreviousFrameWithEvents, "Previous Frame", "Step one frame back", EUserInterfaceActionType::Button, FInputChord(EKeys::Left));
 	UI_COMMAND(NextFrameWithEvents, "Next Frame", "Step one frame forward", EUserInterfaceActionType::Button, FInputChord(EKeys::Right));
 	UI_COMMAND(NextFrameWithStateChange, "Next State Change", "Jump to next state changed frame", EUserInterfaceActionType::Button, FInputChord(EModifierKey::Control, EKeys::Right));
-	UI_COMMAND(ResumeDebuggerAnalysis, "Resume Analysis", "Resume trace analysis.", EUserInterfaceActionType::Button, FInputChord());
+	UI_COMMAND(ResumeDebuggerAnalysis, "Resume Analysis", "Resume trace analysis after hitting breakpoint", EUserInterfaceActionType::Button, FInputChord());
 	UI_COMMAND(ResetTracks, "Reset Tracks", "Clear current tracks and associated events.", EUserInterfaceActionType::Button, FInputChord());
 }
 
