[StateTreeDebugger] refactored ScrubState to use single collection to fix issue introduced in 43692923

diff --git a/Engine/Plugins/Runtime/StateTree/Source/StateTreeDeveloper/Private/Debugger/SStateTreeFrameEventsView.cpp b/Engine/Plugins/Runtime/StateTree/Source/StateTreeDeveloper/Private/Debugger/SStateTreeFrameEventsView.cpp
--- a/Engine/Plugins/Runtime/StateTree/Source/StateTreeDeveloper/Private/Debugger/SStateTreeFrameEventsView.cpp
+++ b/Engine/Plugins/Runtime/StateTree/Source/StateTreeDeveloper/Private/Debugger/SStateTreeFrameEventsView.cpp
@@ -296,6 +296,12 @@ void SFrameEventsView::ExpandAll(const TArray<TSharedPtr<FFrameEventTreeElement>
 
 void SFrameEventsView::RequestRefresh(const FScrubState& InScrubState)
 {
+	const FInstanceEventCollection& EventCollection = InScrubState.GetEventCollection();
+	if (EventCollection.IsInvalid())
+	{
+		return;
+	}
+
 	// Rebuild frame details from the events of that frame
 	EventsTreeElements.Reset();
 	ON_SCOPE_EXIT
@@ -305,7 +311,7 @@ void SFrameEventsView::RequestRefresh(const FScrubState& InScrubState)
 		EventsTreeView->RequestTreeRefresh();
 	};
 
-	const TConstArrayView<const FStateTreeTraceEventVariantType> Events = InScrubState.GetEventCollection().Events;
+	const TConstArrayView<const FStateTreeTraceEventVariantType> Events = EventCollection.Events;
 
 	if (Events.IsEmpty()
 		|| !InScrubState.IsInBounds())
@@ -313,7 +319,7 @@ void SFrameEventsView::RequestRefresh(const FScrubState& InScrubState)
 		return;
 	}
 
-	const TConstArrayView<FFrameSpan> Spans = InScrubState.GetEventCollection().FrameSpans;
+	const TConstArrayView<FFrameSpan> Spans = EventCollection.FrameSpans;
 	check(Spans.Num());
 	check(WeakStateTree.IsValid());
 
diff --git a/Engine/Plugins/Runtime/StateTree/Source/StateTreeEditorModule/Private/Debugger/StateTreeRewindDebuggerTrack.cpp b/Engine/Plugins/Runtime/StateTree/Source/StateTreeEditorModule/Private/Debugger/StateTreeRewindDebuggerTrack.cpp
--- a/Engine/Plugins/Runtime/StateTree/Source/StateTreeEditorModule/Private/Debugger/StateTreeRewindDebuggerTrack.cpp
+++ b/Engine/Plugins/Runtime/StateTree/Source/StateTreeEditorModule/Private/Debugger/StateTreeRewindDebuggerTrack.cpp
@@ -194,9 +194,8 @@ bool FRewindDebuggerTrack::UpdateInternal()
 			{
 				CompactTreeView->Refresh();
 
-				FScrubState ScrubState({EventCollection});
+				FScrubState ScrubState(&EventCollection);
 				ScrubState.SetScrubTime(ScrubTime);
-				ScrubState.SetEventCollectionIndex(0);
 				EventsView->RequestRefresh(ScrubState);
 				return true;
 			}
diff --git a/Engine/Plugins/Runtime/StateTree/Source/StateTreeModule/Private/Debugger/StateTreeDebugger.cpp b/Engine/Plugins/Runtime/StateTree/Source/StateTreeModule/Private/Debugger/StateTreeDebugger.cpp
--- a/Engine/Plugins/Runtime/StateTree/Source/StateTreeModule/Private/Debugger/StateTreeDebugger.cpp
+++ b/Engine/Plugins/Runtime/StateTree/Source/StateTreeModule/Private/Debugger/StateTreeDebugger.cpp
@@ -76,7 +76,6 @@ namespace UE::StateTreeDebugger
 //----------------------------------------------------------------//
 FStateTreeDebugger::FStateTreeDebugger()
 	: StateTreeModule(FModuleManager::GetModuleChecked<IStateTreeModule>("StateTreeModule"))
-	, ScrubState(EventCollections)
 {
 	TracingStateChangedHandle = UE::StateTree::Delegates::OnTracingStateChanged.AddLambda([this](const EStateTreeTraceStatus TraceStatus)
 		{
@@ -193,13 +192,8 @@ void FStateTreeDebugger::SelectInstance(const FStateTreeInstanceDebugId Instance
 	{
 		SelectedInstanceId = InstanceId;
 
-		// Update event collection index for newly debugged instance
-		SetScrubStateCollectionIndex(InstanceId.IsValid()
-			? EventCollections.IndexOfByPredicate([InstanceId = InstanceId](const UE::StateTreeDebugger::FInstanceEventCollection& Entry)
-				{
-					return Entry.InstanceId == InstanceId;
-				})
-			: INDEX_NONE);
+		// Update event collection for newly debugged instance
+		SetScrubStateCollection(GetMutableEventCollection(InstanceId));
 	}
 }
 
@@ -293,13 +287,13 @@ bool FStateTreeDebugger::TryStartNewLiveSessionAnalysis(const float RetryPolling
 		{
 			UpdateAnalysisTransitionType(EAnalysisSourceType::EditorSession);
 
-			SetScrubStateCollectionIndex(INDEX_NONE);
+			SetScrubStateCollection(nullptr);
 			OnNewSession.ExecuteIfBound();
 		}
 
 		return bStarted;
 	}
-	
+
 	RetryLoadNextLiveSessionTimer = RetryPollingDuration;
 	UE_CLOG(RetryLoadNextLiveSessionTimer > 0, LogStateTree, Log, TEXT("Unable to start analysis for the most recent live session."));
 
@@ -357,7 +351,7 @@ bool FStateTreeDebugger::StartSessionAnalysis(const FTraceDescriptor& TraceDescr
 			TraceName = FPaths::Combine(TraceName, Name);
 			FPaths::NormalizeFilename(TraceName);
 		}
-		
+
 		ITraceServicesModule& TraceServicesModule = FModuleManager::LoadModuleChecked<ITraceServicesModule>("TraceServices");
 		if (const TSharedPtr<TraceServices::IAnalysisService> TraceAnalysisService = TraceServicesModule.GetAnalysisService())
 		{
@@ -375,9 +369,9 @@ bool FStateTreeDebugger::StartSessionAnalysis(const FTraceDescriptor& TraceDescr
 	return ActiveSessionTraceDescriptor.IsValid();
 }
 
-void FStateTreeDebugger::SetScrubStateCollectionIndex(const int32 EventCollectionIndex)
+void FStateTreeDebugger::SetScrubStateCollection(const UE::StateTreeDebugger::FInstanceEventCollection* Collection)
 {
-	ScrubState.SetEventCollectionIndex(EventCollectionIndex);
+	ScrubState.SetEventCollection(Collection);
 
 	OnScrubStateChanged.ExecuteIfBound(ScrubState);
 
@@ -541,13 +535,13 @@ void FStateTreeDebugger::RefreshActiveStates()
 {
 	if (ScrubState.IsPointingToValidActiveStates())
 	{
-		const UE::StateTreeDebugger::FInstanceEventCollection& EventCollection = EventCollections[ScrubState.GetEventCollectionIndex()];
+		const UE::StateTreeDebugger::FInstanceEventCollection& EventCollection = ScrubState.GetEventCollection();
 		const int32 EventIndex = EventCollection.ActiveStatesChanges[ScrubState.GetActiveStatesIndex()].EventIndex;
 		SetActiveStates(EventCollection.Events[EventIndex].Get<FStateTreeTraceActiveStatesEvent>().ActiveStates);
 	}
 	else
 	{
-		SetActiveStates(FStateTreeTraceActiveStates());	
+		SetActiveStates(FStateTreeTraceActiveStates());
 	}
 }
 
@@ -560,7 +554,7 @@ void FStateTreeDebugger::StepBackToPreviousStateWithEvents()
 {
 	ScrubState.GotoPreviousFrame();
 	OnScrubStateChanged.Execute(ScrubState);
-	
+
 	RefreshActiveStates();
 }
 
@@ -586,7 +580,7 @@ void FStateTreeDebugger::StepBackToPreviousStateChange()
 {
 	ScrubState.GotoPreviousActiveStates();
 	OnScrubStateChanged.Execute(ScrubState);
-	
+
 	RefreshActiveStates();
 }
 
@@ -612,7 +606,7 @@ bool FStateTreeDebugger::HasStateBreakpoint(const FStateTreeStateHandle StateHan
 				const FStateTreeStateHandle* BreakpointStateHandle = Breakpoint.ElementIdentifier.TryGet<FStateTreeStateHandle>();
 				return (BreakpointStateHandle != nullptr && *BreakpointStateHandle == StateHandle);
 			}
-			return false; 
+			return false;
 		});
 }
 
@@ -625,7 +619,7 @@ bool FStateTreeDebugger::HasTaskBreakpoint(const FStateTreeIndex16 Index, const
 			const FStateTreeDebuggerBreakpoint::FStateTreeTaskIndex* BreakpointTaskIndex = Breakpoint.ElementIdentifier.TryGet<FStateTreeDebuggerBreakpoint::FStateTreeTaskIndex>();
 			return (BreakpointTaskIndex != nullptr && BreakpointTaskIndex->Index == Index);
 		}
-		return false; 
+		return false;
 	});
 }
 
@@ -638,7 +632,7 @@ bool FStateTreeDebugger::HasTransitionBreakpoint(const FStateTreeIndex16 Index,
 			const FStateTreeDebuggerBreakpoint::FStateTreeTransitionIndex* BreakpointTransitionIndex = Breakpoint.ElementIdentifier.TryGet<FStateTreeDebuggerBreakpoint::FStateTreeTransitionIndex>();
 			return (BreakpointTransitionIndex != nullptr && BreakpointTransitionIndex->Index == Index);
 		}
-		return false; 
+		return false;
 	});
 }
 
@@ -686,8 +680,8 @@ bool FStateTreeDebugger::RequestSessionAnalysis(const FTraceDescriptor& TraceDes
 	if (StartSessionAnalysis(TraceDescriptor))
 	{
 		UpdateAnalysisTransitionType(EAnalysisSourceType::SelectedSession);
-		
-		SetScrubStateCollectionIndex(INDEX_NONE);
+
+		SetScrubStateCollection(nullptr);
 		OnNewSession.ExecuteIfBound();
 		return true;
 	}
@@ -787,7 +781,7 @@ void FStateTreeDebugger::SendNotifications()
 		// This will notify scrub position changed and active states
 		SetScrubTime(HitBreakpoint.Time);
 
-		// Make sure the instance is selected in case the breakpoint was set for any instances 
+		// Make sure the instance is selected in case the breakpoint was set for any instances
 		if (SelectedInstanceId != HitBreakpoint.InstanceId)
 		{
 			SelectInstance(HitBreakpoint.InstanceId);
@@ -804,7 +798,7 @@ bool FStateTreeDebugger::EvaluateBreakpoints(const FStateTreeInstanceDebugId Ins
 	if (StateTreeAsset == nullptr // asset is required to properly match state handles
 		|| HitBreakpoint.IsSet() // Only consider first hit breakpoint in the frame
 		|| Breakpoints.IsEmpty()
-		|| (SelectedInstanceId.IsValid() && InstanceId != SelectedInstanceId)) // ignore events not for the selected instances		
+		|| (SelectedInstanceId.IsValid() && InstanceId != SelectedInstanceId)) // ignore events not for the selected instances
 	{
 		return false;
 	}
@@ -825,11 +819,7 @@ bool FStateTreeDebugger::EvaluateBreakpoints(const FStateTreeInstanceDebugId Ins
 
 UE::StateTreeDebugger::FInstanceEventCollection* FStateTreeDebugger::GetOrCreateEventCollection(FStateTreeInstanceDebugId InstanceId)
 {
-	UE::StateTreeDebugger::FInstanceEventCollection* ExistingCollection = EventCollections.FindByPredicate(
-		[InstanceId](const UE::StateTreeDebugger::FInstanceEventCollection& Entry)
-		{
-			return Entry.InstanceId == InstanceId;
-		});
+	UE::StateTreeDebugger::FInstanceEventCollection* ExistingCollection = GetMutableEventCollection(InstanceId);
 
 	// Create missing EventCollection if necessary
 	if (ExistingCollection == nullptr)
@@ -837,13 +827,14 @@ UE::StateTreeDebugger::FInstanceEventCollection* FStateTreeDebugger::GetOrCreate
 		// Push deferred notification for new instance Id
 		NewInstances.Push(InstanceId);
 
-		ExistingCollection = &EventCollections.Emplace_GetRef(InstanceId);
+		ExistingCollection = new UE::StateTreeDebugger::FInstanceEventCollection(InstanceId);
+		EventCollections.Add(ExistingCollection);
 
-		// Update the active event collection index when it's newly created for the currently debugged instance.
+		// Update the active event collection when it's newly created for the currently debugged instance.
 		// Otherwise (i.e. EventCollection already exists) it is updated when switching instance (i.e. SelectInstance)
-		if (SelectedInstanceId == InstanceId && ScrubState.GetEventCollectionIndex() == INDEX_NONE)
+		if (SelectedInstanceId == InstanceId && ScrubState.GetEventCollection().IsInvalid())
 		{
-			SetScrubStateCollectionIndex(EventCollections.Num() - 1);
+			SetScrubStateCollection(ExistingCollection);
 		}
 	}
 	return ExistingCollection;
@@ -859,21 +850,36 @@ void FStateTreeDebugger::OnTraceEventProcessed(const FStateTreeInstanceDebugId I
 	EvaluateBreakpoints(InstanceId, Event);
 }
 
-const UE::StateTreeDebugger::FInstanceEventCollection& FStateTreeDebugger::GetEventCollection(FStateTreeInstanceDebugId InstanceId) const
+const UE::StateTreeDebugger::FInstanceEventCollection& FStateTreeDebugger::GetEventCollection(const FStateTreeInstanceDebugId InstanceId) const
 {
-	using namespace UE::StateTreeDebugger;
-	const FInstanceEventCollection* ExistingCollection = EventCollections.FindByPredicate([InstanceId](const FInstanceEventCollection& Entry)
+	for (const UE::StateTreeDebugger::FInstanceEventCollection& EventCollection : EventCollections)
 	{
-		return Entry.InstanceId == InstanceId;
-	});
+		if (EventCollection.InstanceId == InstanceId)
+		{
+			return EventCollection;
+		}
+	}
+
+	return UE::StateTreeDebugger::FInstanceEventCollection::Invalid;
+}
+
+UE::StateTreeDebugger::FInstanceEventCollection* FStateTreeDebugger::GetMutableEventCollection(const FStateTreeInstanceDebugId InstanceId)
+{
+	for (UE::StateTreeDebugger::FInstanceEventCollection& EventCollection : EventCollections)
+	{
+		if (EventCollection.InstanceId == InstanceId)
+		{
+			return &EventCollection;
+		}
+	}
 
-	return ExistingCollection != nullptr ? *ExistingCollection : FInstanceEventCollection::Invalid;
+	return nullptr;
 }
 
 void FStateTreeDebugger::ResetEventCollections()
 {
 	EventCollections.Reset();
-	SetScrubStateCollectionIndex(INDEX_NONE);
+	SetScrubStateCollection(nullptr);
 	LastProcessedRecordedWorldTime = 0;
 }
 
diff --git a/Engine/Plugins/Runtime/StateTree/Source/StateTreeModule/Private/Debugger/StateTreeDebuggerTypes.cpp b/Engine/Plugins/Runtime/StateTree/Source/StateTreeModule/Private/Debugger/StateTreeDebuggerTypes.cpp
--- a/Engine/Plugins/Runtime/StateTree/Source/StateTreeModule/Private/Debugger/StateTreeDebuggerTypes.cpp
+++ b/Engine/Plugins/Runtime/StateTree/Source/StateTreeModule/Private/Debugger/StateTreeDebuggerTypes.cpp
@@ -35,7 +35,7 @@ const FInstanceEventCollection FInstanceEventCollection::Invalid;
 //----------------------------------------------------------------//
 bool FScrubState::SetScrubTime(const double NewScrubTime, const bool bForceRefresh)
 {
-	if (! bForceRefresh
+	if (!bForceRefresh
 		&& NewScrubTime == ScrubTime)
 	{
 		return false;
@@ -46,9 +46,9 @@ bool FScrubState::SetScrubTime(const double NewScrubTime, const bool bForceRefre
 	FrameSpanIndex = INDEX_NONE;
 	ActiveStatesIndex = INDEX_NONE;
 
-	if (EventCollectionIndex != INDEX_NONE)
+	if (EventCollection)
 	{
-		const TArray<FFrameSpan>& Spans = EventCollections[EventCollectionIndex].FrameSpans;
+		const TArray<FFrameSpan>& Spans = EventCollection->FrameSpans;
 		if (Spans.Num() > 0)
 		{
 			const double SpansLowerBound = Spans[0].GetWorldTimeStart();
@@ -84,24 +84,26 @@ bool FScrubState::SetScrubTime(const double NewScrubTime, const bool bForceRefre
 	return true;
 }
 
-void FScrubState::SetEventCollectionIndex(const int32 InEventCollectionIndex)
+void FScrubState::SetEventCollection(const FInstanceEventCollection* InEventCollection)
 {
-	EventCollectionIndex = InEventCollectionIndex;
+	EventCollection = InEventCollection;
 
-	// Force refresh of internal indices with current time applied on the new event collection.
-	constexpr bool bForceRefresh = true;
-	SetScrubTime(ScrubTime, bForceRefresh);
+	if (EventCollection)
+	{
+		// Force refresh of internal indices with current time applied on the new event collection.
+		constexpr bool bForceRefresh = true;
+		SetScrubTime(ScrubTime, bForceRefresh);
+	}
 }
 
 void FScrubState::SetFrameSpanIndex(const int32 NewFrameSpanIndex)
 {
 	FrameSpanIndex = NewFrameSpanIndex;
-	checkf(EventCollections.IsValidIndex(EventCollectionIndex), TEXT("Internal method expecting validity checks before getting called."));
-	const FInstanceEventCollection& EventCollection = EventCollections[EventCollectionIndex];
+	checkf(EventCollection != nullptr, TEXT("Internal method expecting validity checks before getting called."));
+	checkf(EventCollection->FrameSpans.IsValidIndex(FrameSpanIndex), TEXT("Internal method expecting validity checks before getting called."));
 
-	checkf(EventCollection.FrameSpans.IsValidIndex(FrameSpanIndex), TEXT("Internal method expecting validity checks before getting called."));
-	ScrubTime = EventCollection.FrameSpans[FrameSpanIndex].GetWorldTimeStart();
-	TraceFrameIndex = EventCollection.FrameSpans[FrameSpanIndex].Frame.Index;
+	ScrubTime = EventCollection->FrameSpans[FrameSpanIndex].GetWorldTimeStart();
+	TraceFrameIndex = EventCollection->FrameSpans[FrameSpanIndex].Frame.Index;
 	ScrubTimeBoundState = EScrubTimeBoundState::InBounds;
 	UpdateActiveStatesIndex(NewFrameSpanIndex);
 }
@@ -110,36 +112,35 @@ void FScrubState::SetActiveStatesIndex(const int32 NewActiveStatesIndex)
 {
 	ActiveStatesIndex = NewActiveStatesIndex;
 
-	checkf(EventCollections.IsValidIndex(EventCollectionIndex), TEXT("Internal method expecting validity checks before getting called."));
-	const FInstanceEventCollection& EventCollection = EventCollections[EventCollectionIndex];
+	checkf(EventCollection != nullptr, TEXT("Internal method expecting validity checks before getting called."));
+	checkf(EventCollection->ActiveStatesChanges.IsValidIndex(ActiveStatesIndex), TEXT("Internal method expecting validity checks before getting called."));
 
-	checkf(EventCollection.ActiveStatesChanges.IsValidIndex(ActiveStatesIndex), TEXT("Internal method expecting validity checks before getting called."));
-	FrameSpanIndex = EventCollection.ActiveStatesChanges[ActiveStatesIndex].SpanIndex;
-	ScrubTime = EventCollection.FrameSpans[FrameSpanIndex].GetWorldTimeStart();
-	TraceFrameIndex = EventCollection.FrameSpans[FrameSpanIndex].Frame.Index;
+	FrameSpanIndex = EventCollection->ActiveStatesChanges[ActiveStatesIndex].SpanIndex;
+	ScrubTime = EventCollection->FrameSpans[FrameSpanIndex].GetWorldTimeStart();
+	TraceFrameIndex = EventCollection->FrameSpans[FrameSpanIndex].Frame.Index;
 	ScrubTimeBoundState = EScrubTimeBoundState::InBounds;
 }
 
 bool FScrubState::HasPreviousFrame() const
 {
-	if (EventCollectionIndex != INDEX_NONE)
+	if (EventCollection != nullptr)
 	{
-		return IsInBounds() ? EventCollections[EventCollectionIndex].FrameSpans.IsValidIndex(FrameSpanIndex - 1) : ScrubTimeBoundState == EScrubTimeBoundState::AfterHigherBound;
+		return IsInBounds() ? EventCollection->FrameSpans.IsValidIndex(FrameSpanIndex - 1) : ScrubTimeBoundState == EScrubTimeBoundState::AfterHigherBound;
 	}
 	return false;
 }
 
 double FScrubState::GotoPreviousFrame()
 {
-	SetFrameSpanIndex(IsInBounds() ? (FrameSpanIndex - 1) : EventCollections[EventCollectionIndex].FrameSpans.Num()-1);
+	SetFrameSpanIndex(IsInBounds() ? (FrameSpanIndex - 1) : EventCollection->FrameSpans.Num() - 1);
 	return ScrubTime;
 }
 
 bool FScrubState::HasNextFrame() const
 {
-	if (EventCollectionIndex != INDEX_NONE)
+	if (EventCollection != nullptr)
 	{
-		return IsInBounds()	? EventCollections[EventCollectionIndex].FrameSpans.IsValidIndex(FrameSpanIndex + 1) : ScrubTimeBoundState == EScrubTimeBoundState::BeforeLowerBound;
+		return IsInBounds() ? EventCollection->FrameSpans.IsValidIndex(FrameSpanIndex + 1) : ScrubTimeBoundState == EScrubTimeBoundState::BeforeLowerBound;
 	}
 	return false;
 }
@@ -152,12 +153,12 @@ double FScrubState::GotoNextFrame()
 
 bool FScrubState::HasPreviousActiveStates() const
 {
-	if (EventCollectionIndex == INDEX_NONE || ActiveStatesIndex == INDEX_NONE)
+	if (EventCollection == nullptr || ActiveStatesIndex == INDEX_NONE)
 	{
 		return false;
 	}
 
-	const TArray<FInstanceEventCollection::FActiveStatesChangePair>& ActiveStatesChanges = EventCollections[EventCollectionIndex].ActiveStatesChanges;
+	const TArray<FInstanceEventCollection::FActiveStatesChangePair>& ActiveStatesChanges = EventCollection->ActiveStatesChanges;
 	if (ScrubTimeBoundState == EScrubTimeBoundState::AfterHigherBound && ActiveStatesChanges.Num() > 0)
 	{
 		return true;
@@ -173,10 +174,10 @@ bool FScrubState::HasPreviousActiveStates() const
 
 double FScrubState::GotoPreviousActiveStates()
 {
-	const TArray<FInstanceEventCollection::FActiveStatesChangePair>& ActiveStatesChanges = EventCollections[EventCollectionIndex].ActiveStatesChanges;
+	const TArray<FInstanceEventCollection::FActiveStatesChangePair>& ActiveStatesChanges = EventCollection->ActiveStatesChanges;
 	if (ScrubTimeBoundState == EScrubTimeBoundState::AfterHigherBound)
 	{
-		SetActiveStatesIndex(ActiveStatesChanges.Num()-1);
+		SetActiveStatesIndex(ActiveStatesChanges.Num() - 1);
 	}
 	else if (ActiveStatesChanges.IsValidIndex(ActiveStatesIndex) && ActiveStatesChanges[ActiveStatesIndex].SpanIndex < FrameSpanIndex)
 	{
@@ -192,18 +193,18 @@ double FScrubState::GotoPreviousActiveStates()
 
 bool FScrubState::HasNextActiveStates() const
 {
-	if (EventCollectionIndex == INDEX_NONE)
+	if (EventCollection == nullptr)
 	{
 		return false;
 	}
 
-	const TArray<FInstanceEventCollection::FActiveStatesChangePair>& ActiveStatesChanges = EventCollections[EventCollectionIndex].ActiveStatesChanges;
+	const TArray<FInstanceEventCollection::FActiveStatesChangePair>& ActiveStatesChanges = EventCollection->ActiveStatesChanges;
 	if (ScrubTimeBoundState == EScrubTimeBoundState::BeforeLowerBound && ActiveStatesChanges.Num() > 0)
 	{
 		return true;
 	}
 
-	return ActiveStatesIndex != INDEX_NONE && EventCollections[EventCollectionIndex].ActiveStatesChanges.IsValidIndex(ActiveStatesIndex + 1);
+	return ActiveStatesIndex != INDEX_NONE && EventCollection->ActiveStatesChanges.IsValidIndex(ActiveStatesIndex + 1);
 }
 
 double FScrubState::GotoNextActiveStates()
@@ -221,16 +222,15 @@ double FScrubState::GotoNextActiveStates()
 
 const FInstanceEventCollection& FScrubState::GetEventCollection() const
 {
-	return EventCollectionIndex != INDEX_NONE ? EventCollections[EventCollectionIndex] : FInstanceEventCollection::Invalid;
+	return EventCollection != nullptr ? *EventCollection : FInstanceEventCollection::Invalid;
 }
 
 void FScrubState::UpdateActiveStatesIndex(const int32 SpanIndex)
 {
-	check(EventCollectionIndex != INDEX_NONE);
-	const FInstanceEventCollection& EventCollection = EventCollections[EventCollectionIndex];
+	check(EventCollection != nullptr);
 
 	// Need to find the index of a frame span that contains an active states changed event; either the current one has it otherwise look backward to find the last one
-	ActiveStatesIndex = EventCollection.ActiveStatesChanges.FindLastByPredicate(
+	ActiveStatesIndex = EventCollection->ActiveStatesChanges.FindLastByPredicate(
 		[SpanIndex](const FInstanceEventCollection::FActiveStatesChangePair& SpanAndEventIndices)
 		{
 			return SpanAndEventIndices.SpanIndex <= SpanIndex;
diff --git a/Engine/Plugins/Runtime/StateTree/Source/StateTreeModule/Public/Debugger/StateTreeDebugger.h b/Engine/Plugins/Runtime/StateTree/Source/StateTreeModule/Public/Debugger/StateTreeDebugger.h
--- a/Engine/Plugins/Runtime/StateTree/Source/StateTreeModule/Public/Debugger/StateTreeDebugger.h
+++ b/Engine/Plugins/Runtime/StateTree/Source/StateTreeModule/Public/Debugger/StateTreeDebugger.h
@@ -368,6 +368,9 @@ protected:
 	}
 
 private:
+
+	UE::StateTreeDebugger::FInstanceEventCollection* GetMutableEventCollection(const FStateTreeInstanceDebugId InstanceId);
+
 	enum class EAnalysisSourceType : uint8
 	{
 		// Analysis selected from available sessions
@@ -426,7 +429,7 @@ private:
 
 	bool StartSessionAnalysis(const FTraceDescriptor& TraceDescriptor);
 
-	void SetScrubStateCollectionIndex(int32 EventCollectionIndex);
+	void SetScrubStateCollection(const UE::StateTreeDebugger::FInstanceEventCollection* Collection);
 
 	/**
 	 * Recompute index of the span that contains the active states change event and update the active states.
@@ -451,7 +454,7 @@ private:
 	FTraceDescriptor ActiveSessionTraceDescriptor;
 
 	/** Processed events for each instance. */
-	TArray<UE::StateTreeDebugger::FInstanceEventCollection> EventCollections;
+	TIndirectArray<UE::StateTreeDebugger::FInstanceEventCollection> EventCollections;
 
 	/** Specific instance selected for more details */
 	FStateTreeInstanceDebugId SelectedInstanceId;
diff --git a/Engine/Plugins/Runtime/StateTree/Source/StateTreeModule/Public/Debugger/StateTreeDebuggerTypes.h b/Engine/Plugins/Runtime/StateTree/Source/StateTreeModule/Public/Debugger/StateTreeDebuggerTypes.h
--- a/Engine/Plugins/Runtime/StateTree/Source/StateTreeModule/Public/Debugger/StateTreeDebuggerTypes.h
+++ b/Engine/Plugins/Runtime/StateTree/Source/StateTreeModule/Public/Debugger/StateTreeDebuggerTypes.h
@@ -25,6 +25,8 @@ enum class EStateTreeBreakpointType : uint8
 
 #if WITH_STATETREE_TRACE_DEBUGGER
 
+#define UE_API STATETREEMODULE_API
+
 class UStateTree;
 enum class EStateTreeTraceEventType : uint8;
 
@@ -63,9 +65,9 @@ struct FFrameSpan
 struct FInstanceDescriptor : TSharedFromThis<FInstanceDescriptor>
 {
 	FInstanceDescriptor() = default;
-	STATETREEMODULE_API FInstanceDescriptor(const UStateTree* InStateTree, const FStateTreeInstanceDebugId InId, const FString& InName, const TRange<double>& InLifetime);
+	UE_API FInstanceDescriptor(const UStateTree* InStateTree, const FStateTreeInstanceDebugId InId, const FString& InName, const TRange<double>& InLifetime);
 
-	STATETREEMODULE_API bool IsValid() const;
+	UE_API bool IsValid() const;
 
 	bool operator==(const FInstanceDescriptor& Other) const
 	{
@@ -125,14 +127,21 @@ struct FInstanceEventCollection
 		return !(Lhs == RHS);
 	}
 
-	bool IsValid() const { return InstanceId.IsValid(); }
-	bool IsInvalid() const { return !IsValid(); }
+	bool IsValid() const
+	{
+		return InstanceId.IsValid();
+	}
+
+	bool IsInvalid() const
+	{
+		return !IsValid();
+	}
 
 	struct FActiveStatesChangePair
 	{
 		FActiveStatesChangePair(const int32 SpanIndex, const int32 EventIndex)
-			: SpanIndex(SpanIndex),
-			  EventIndex(EventIndex)
+			: SpanIndex(SpanIndex)
+			, EventIndex(EventIndex)
 		{
 		}
 
@@ -171,30 +180,52 @@ struct FInstanceEventCollection
 	 * An invalid empty collection is returned if there is no selected instance. (IsValid needs to be called).
 	 * @return Event collection associated to the selected instance or an invalid one if not found.
 	 */
-	static STATETREEMODULE_API const FInstanceEventCollection Invalid;
+	static UE_API const FInstanceEventCollection Invalid;
 };
 
 struct FScrubState
 {
+	FScrubState() = default;
+	explicit FScrubState(const FInstanceEventCollection* EventCollection)
+		: EventCollection(EventCollection)
+	{
+	}
+	UE_DEPRECATED(5.7, "FScrubState will no longer support multiple collections.")
 	explicit FScrubState(const TArray<FInstanceEventCollection>& EventCollections)
-		: EventCollections(EventCollections)
 	{
 	}
 
-	/** @return Index of the currently selected event collection; INDEX_NONE if nothing is selected. */
-	int32 GetEventCollectionIndex() const { return EventCollectionIndex; }
+	UE_DEPRECATED(5.7, "FScrubState will no longer support multiple collections.")
+	int32 GetEventCollectionIndex() const
+	{
+		return INDEX_NONE;
+	}
+
+	UE_DEPRECATED(5.7, "FScrubState will no longer support multiple collections.")
+	void SetEventCollectionIndex(const int32 InEventCollectionIndex)
+	{
+	}
 
-	/** Assigns a new collection index and updates internal indices for current scrub time. */
-	STATETREEMODULE_API void SetEventCollectionIndex(const int32 InEventCollectionIndex);
+	/** Assigns a new collection and updates internal indices for current scrub time. */
+	UE_API void SetEventCollection(const FInstanceEventCollection* InEventCollection);
 
 	/** @return Index of the span for the currently selected frame; INDEX_NONE if there is no span for the current scrub time. */
-	int32 GetFrameSpanIndex() const { return FrameSpanIndex; }
+	int32 GetFrameSpanIndex() const
+	{
+		return FrameSpanIndex;
+	}
 
 	/** @return Index of the list of active states for the currently selected frame; INDEX_NONE if there is no active states for the current scrub time. */
-	int32 GetActiveStatesIndex() const { return ActiveStatesIndex; }
+	int32 GetActiveStatesIndex() const
+	{
+		return ActiveStatesIndex;
+	}
 
 	/** @return Current scrub time. */
-	double GetScrubTime() const { return ScrubTime; }
+	double GetScrubTime() const
+	{
+		return ScrubTime;
+	}
 
 	/**
 	 * Updates internal indices based on the new time.
@@ -202,66 +233,72 @@ struct FScrubState
 	 * @param bForceRefresh Whether to force update of internal indices regardless if the current stored scrub time is the same as the provided one
 	 * @return true if values were updated; false otherwise (i.e. no changes)
 	 */
-	STATETREEMODULE_API bool SetScrubTime(double NewScrubTime, bool bForceRefresh = false);
+	UE_API bool SetScrubTime(double NewScrubTime, bool bForceRefresh = false);
 
 	/**
 	 * Indicates if the current scrub state points to a valid frame.
 	 * @return True if the frame index is set
 	 */
-	bool IsInBounds() const { return ScrubTimeBoundState == EScrubTimeBoundState::InBounds; }
+	bool IsInBounds() const
+	{
+		return ScrubTimeBoundState == EScrubTimeBoundState::InBounds;
+	}
 
 	/**
 	 * Indicates if the current scrub state points to an active states entry in the event collection.
 	 * @return True if the collection and active states indices are set
 	 */
-	bool IsPointingToValidActiveStates() const { return EventCollectionIndex != INDEX_NONE && ActiveStatesIndex != INDEX_NONE; }
+	bool IsPointingToValidActiveStates() const
+	{
+		return EventCollection != nullptr && ActiveStatesIndex != INDEX_NONE;
+	}
 
 	/** Indicates if there is a frame before with events. */
-	STATETREEMODULE_API bool HasPreviousFrame() const;
-	
+	UE_API bool HasPreviousFrame() const;
+
 	/**
 	 * Set scrubbing info using the previous frame with events.
 	 * HasPreviousFrame must be used to validate that this method can be called otherwise some checks might fail.
 	 * @return Adjusted scrub time
 	 */
-	STATETREEMODULE_API double GotoPreviousFrame();
-	
+	UE_API double GotoPreviousFrame();
+
 	/** Indicates if there is a frame after with events. */
-	STATETREEMODULE_API bool HasNextFrame() const;
-	
+	UE_API bool HasNextFrame() const;
+
 	/**
 	 * Set scrubbing info using the next frame with events.
 	 * HasPreviousFrame must be used to validate that this method can be called otherwise some checks might fail.
 	 * @return Adjusted scrub time
 	 */
-	STATETREEMODULE_API double GotoNextFrame();
-	
+	UE_API double GotoNextFrame();
+
 	/** Indicates if there is a frame before where the StateTree has a different list of active states. */
-	STATETREEMODULE_API bool HasPreviousActiveStates() const;
-	
+	UE_API bool HasPreviousActiveStates() const;
+
 	/**
 	 * Set scrubbing info using the previous frame where the StateTree has a different list of active states.
 	 * HasPreviousActiveStates must be used to validate that this method can be called otherwise some checks might fail.
 	 * @return Adjusted scrub time
 	 */
-	STATETREEMODULE_API double GotoPreviousActiveStates();
-	
+	UE_API double GotoPreviousActiveStates();
+
 	/** Indicates if there is a frame after where the StateTree has a different list of active states. */
-	STATETREEMODULE_API bool HasNextActiveStates() const;
-	
+	UE_API bool HasNextActiveStates() const;
+
 	/**
 	 * Set scrubbing info using the next frame where the StateTree has a different list of active states.
 	 * HasNextActiveStates must be used to validate that this method can be called otherwise some checks might fail.
 	 * @return Adjusted scrub time
 	 */
-	STATETREEMODULE_API double GotoNextActiveStates();
+	UE_API double GotoNextActiveStates();
 
 	/**
-	 * Returns the event collection associated to the selected instance.
-	 * An invalid empty collection is returned if there is no selected instance (IsValid needs to be called).
-	 * @return Event collection associated to the selected instance or an invalid one if not found.
+	 * Returns the event collection associated.
+	 * An invalid empty collection is returned if no collection was set (IsValid needs to be called).
+	 * @return Event collection set or an invalid one if not set.
 	 */
-	STATETREEMODULE_API const FInstanceEventCollection& GetEventCollection() const;
+	UE_API const FInstanceEventCollection& GetEventCollection() const;
 
 private:
 	enum class EScrubTimeBoundState : uint8
@@ -275,13 +312,12 @@ private:
 		AfterHigherBound
 	};
 
-	STATETREEMODULE_API void SetFrameSpanIndex(int32 NewFrameSpanIndex);
-	STATETREEMODULE_API void SetActiveStatesIndex(int32 NewActiveStatesIndex);
-	STATETREEMODULE_API void UpdateActiveStatesIndex(int32 SpanIndex);
+	UE_API void SetFrameSpanIndex(int32 NewFrameSpanIndex);
+	UE_API void SetActiveStatesIndex(int32 NewActiveStatesIndex);
+	UE_API void UpdateActiveStatesIndex(int32 SpanIndex);
 
-	const TArray<FInstanceEventCollection>& EventCollections;
+	const FInstanceEventCollection* EventCollection = nullptr;
 	double ScrubTime = 0;
-	int32 EventCollectionIndex = INDEX_NONE;
 	uint64 TraceFrameIndex = INDEX_NONE;
 	int32 FrameSpanIndex = INDEX_NONE;
 	int32 ActiveStatesIndex = INDEX_NONE;
@@ -303,7 +339,7 @@ struct FStateTreeDebuggerBreakpoint
 
 		FStateTreeIndex16 Index;
 	};
-	
+
 	struct FStateTreeTransitionIndex
 	{
 		FStateTreeTransitionIndex() = default;
@@ -315,13 +351,13 @@ struct FStateTreeDebuggerBreakpoint
 	};
 
 	using FIdentifierVariantType = TVariant<FStateTreeStateHandle, FStateTreeTaskIndex, FStateTreeTransitionIndex>;
-	
+
 	FStateTreeDebuggerBreakpoint();
 	explicit FStateTreeDebuggerBreakpoint(const FStateTreeStateHandle StateHandle, const EStateTreeBreakpointType BreakpointType);
 	explicit FStateTreeDebuggerBreakpoint(const FStateTreeTaskIndex Index, const EStateTreeBreakpointType BreakpointType);
 	explicit FStateTreeDebuggerBreakpoint(const FStateTreeTransitionIndex Index, const EStateTreeBreakpointType BreakpointType);
 
-	STATETREEMODULE_API bool IsMatchingEvent(const FStateTreeTraceEventVariantType& Event) const;
+	UE_API bool IsMatchingEvent(const FStateTreeTraceEventVariantType& Event) const;
 
 	FIdentifierVariantType ElementIdentifier;
 	EStateTreeBreakpointType BreakpointType;
@@ -331,4 +367,6 @@ private:
 	static EStateTreeTraceEventType GetMatchingEventType(EStateTreeBreakpointType BreakpointType);
 };
 
+#undef UE_API
+
 #endif // WITH_STATETREE_TRACE_DEBUGGER
