[StateTreeDebugger] applied some code-review feedback and fixed missing reset on the last recorded world time when stopping/restarting a recording without closing the debugger.

diff --git a/Engine/Plugins/Runtime/StateTree/Source/StateTreeEditorModule/Private/Debugger/SStateTreeDebuggerView.cpp b/Engine/Plugins/Runtime/StateTree/Source/StateTreeEditorModule/Private/Debugger/SStateTreeDebuggerView.cpp
--- a/Engine/Plugins/Runtime/StateTree/Source/StateTreeEditorModule/Private/Debugger/SStateTreeDebuggerView.cpp
+++ b/Engine/Plugins/Runtime/StateTree/Source/StateTreeEditorModule/Private/Debugger/SStateTreeDebuggerView.cpp
@@ -71,27 +71,43 @@ void SStateTreeDebuggerView::StopRecording()
 {
 	if (CanStopRecording())
 	{
+		// Calling StopTraces on the module will notify all regsitered views through OnTracesStateChanged delegate.
 		IStateTreeModule& StateTreeModule = FModuleManager::GetModuleChecked<IStateTreeModule>("StateTreeModule");
 		StateTreeModule.StopTraces();
+	}
+}
 
-		bRecording = false;
+void SStateTreeDebuggerView::HandleTracesStateChanged(const EStateTreeTraceStatus TraceStatus)
+{
+	if (TraceStatus == EStateTreeTraceStatus::TracesStarted)
+	{
+		bRecording = true;
+	}
+	else
+	{
+		HandleTracesStopped();
+	}
+}
 
-		// Update max duration from current recording until track data gets reset
-		check(Debugger);
-		MaxTrackRecordingDuration = FMath::Max(MaxTrackRecordingDuration, ExtrapolatedRecordedWorldTime);
+void SStateTreeDebuggerView::HandleTracesStopped()
+{
+	bRecording = false;
 
-		// Mark all tracks from the stopped session as stale to have different look.
-		for (const TSharedPtr<RewindDebugger::FRewindDebuggerTrack>& DebugTrack : InstanceOwnerTracks)
+	// Update max duration from current recording until track data gets reset
+	check(Debugger);
+	MaxTrackRecordingDuration = FMath::Max(MaxTrackRecordingDuration, ExtrapolatedRecordedWorldTime);
+
+	// Mark all tracks from the stopped session as stale to have different look.
+	for (const TSharedPtr<RewindDebugger::FRewindDebuggerTrack>& DebugTrack : InstanceOwnerTracks)
+	{
+		if (FStateTreeDebuggerBaseTrack* StateTreeTrack = static_cast<FStateTreeDebuggerBaseTrack*>(DebugTrack.Get()))
 		{
-			if (FStateTreeDebuggerBaseTrack* StateTreeTrack = static_cast<FStateTreeDebuggerBaseTrack*>(DebugTrack.Get()))
-			{
-				StateTreeTrack->MarkAsStale(ExtrapolatedRecordedWorldTime);
-			}
+			StateTreeTrack->MarkAsStale(ExtrapolatedRecordedWorldTime);
 		}
-
-		// Reset the extrapolated time for next runs
-		ExtrapolatedRecordedWorldTime = 0;
 	}
+
+	// Reset the extrapolated time for next runs
+	ExtrapolatedRecordedWorldTime = 0;
 }
 
 bool SStateTreeDebuggerView::CanResumeDebuggerAnalysis() const
@@ -191,10 +207,7 @@ void SStateTreeDebuggerView::Construct(const FArguments& InArgs
 
 	const IStateTreeModule& StateTreeModule = FModuleManager::GetModuleChecked<IStateTreeModule>("StateTreeModule");
 	bRecording = StateTreeModule.IsTracing();
-	UE::StateTree::Delegates::OnTracingStateChanged.AddSPLambda(this, [&bRecording=bRecording](const EStateTreeTraceStatus TraceStatus)
-		{
-			bRecording = TraceStatus == EStateTreeTraceStatus::TracesStarted;
-		});
+	UE::StateTree::Delegates::OnTracingStateChanged.AddSP(this, &SStateTreeDebuggerView::HandleTracesStateChanged);
 
 	// Registers a delegate to be notified when the associated StateTree asset get successfully recompiled
 	// to clear any previous recorded data that could mismatch new compiled tree data version.
diff --git a/Engine/Plugins/Runtime/StateTree/Source/StateTreeEditorModule/Private/Debugger/SStateTreeDebuggerView.h b/Engine/Plugins/Runtime/StateTree/Source/StateTreeEditorModule/Private/Debugger/SStateTreeDebuggerView.h
--- a/Engine/Plugins/Runtime/StateTree/Source/StateTreeEditorModule/Private/Debugger/SStateTreeDebuggerView.h
+++ b/Engine/Plugins/Runtime/StateTree/Source/StateTreeEditorModule/Private/Debugger/SStateTreeDebuggerView.h
@@ -100,6 +100,9 @@ private:
 
 	void StopRecording();
 
+	void HandleTracesStateChanged(EStateTreeTraceStatus TraceStatus);
+	void HandleTracesStopped();
+
 	bool CanResumeDebuggerAnalysis() const;
 	void ResumeDebuggerAnalysis() const;
 
diff --git a/Engine/Plugins/Runtime/StateTree/Source/StateTreeModule/Private/Debugger/StateTreeDebugger.cpp b/Engine/Plugins/Runtime/StateTree/Source/StateTreeModule/Private/Debugger/StateTreeDebugger.cpp
--- a/Engine/Plugins/Runtime/StateTree/Source/StateTreeModule/Private/Debugger/StateTreeDebugger.cpp
+++ b/Engine/Plugins/Runtime/StateTree/Source/StateTreeModule/Private/Debugger/StateTreeDebugger.cpp
@@ -145,6 +145,7 @@ void FStateTreeDebugger::StopSessionAnalysis()
 	bSessionAnalysisActive = false;
 	bSessionAnalysisPaused = false;
 	HitBreakpoint.Reset();
+	LastProcessedRecordedWorldTime = 0;
 }
 
 void FStateTreeDebugger::SyncToCurrentSessionDuration()
@@ -869,6 +870,8 @@ UE::StateTreeDebugger::FInstanceEventCollection* FStateTreeDebugger::GetMutableE
 	{
 		if (EventCollection.InstanceId == InstanceId)
 		{
+			// Note that returning the pointer to the EventCollection element here is fine.
+			// Using a TDereferencingIterator on a TIndirectArray returns us the actual collection instance.
 			return &EventCollection;
 		}
 	}
diff --git a/Engine/Plugins/Runtime/StateTree/Source/StateTreeModule/Private/Debugger/StateTreeDebuggerTypes.cpp b/Engine/Plugins/Runtime/StateTree/Source/StateTreeModule/Private/Debugger/StateTreeDebuggerTypes.cpp
--- a/Engine/Plugins/Runtime/StateTree/Source/StateTreeModule/Private/Debugger/StateTreeDebuggerTypes.cpp
+++ b/Engine/Plugins/Runtime/StateTree/Source/StateTreeModule/Private/Debugger/StateTreeDebuggerTypes.cpp
@@ -88,12 +88,9 @@ void FScrubState::SetEventCollection(const FInstanceEventCollection* InEventColl
 {
 	EventCollection = InEventCollection;
 
-	if (EventCollection)
-	{
-		// Force refresh of internal indices with current time applied on the new event collection.
-		constexpr bool bForceRefresh = true;
-		SetScrubTime(ScrubTime, bForceRefresh);
-	}
+	// Force refresh of internal indices with current time applied on the new event collection.
+	constexpr bool bForceRefresh = true;
+	SetScrubTime(ScrubTime, bForceRefresh);
 }
 
 void FScrubState::SetFrameSpanIndex(const int32 NewFrameSpanIndex)
