[StateTree] Fix localization error
#rnx

diff --git a/Engine/Plugins/Runtime/StateTree/Source/StateTreeEditorModule/Private/Customizations/StateTreeStateDetails.cpp b/Engine/Plugins/Runtime/StateTree/Source/StateTreeEditorModule/Private/Customizations/StateTreeStateDetails.cpp
--- a/Engine/Plugins/Runtime/StateTree/Source/StateTreeEditorModule/Private/Customizations/StateTreeStateDetails.cpp
+++ b/Engine/Plugins/Runtime/StateTree/Source/StateTreeEditorModule/Private/Customizations/StateTreeStateDetails.cpp
@@ -172,7 +172,7 @@ void FStateTreeStateDetails::CustomizeDetails(IDetailLayoutBuilder& DetailBuilde
 	TypeProperty->MarkHiddenByCustomization();
 	{
 		// Restrict IsStateTypeAllowed
-		TSharedRef<FPropertyRestriction> SchemaSupported = MakeShared<FPropertyRestriction>(LOCTEXT("StateSelectionAllowedRestriction", "The schema restricts the state type."));
+		TSharedRef<FPropertyRestriction> SchemaSupported = MakeShared<FPropertyRestriction>(LOCTEXT("StateTypeAllowedRestriction", "The schema restricts the state type."));
 		const UEnum* EnumPtr = StaticEnum<EStateTreeStateType>();
 		for (int32 Index = 0; Index < EnumPtr->NumEnums(); ++Index)
 		{
