[StateTreeDebugger] used on stack copy of the active states for recursive call

diff --git a/Engine/Plugins/Runtime/StateTree/Source/StateTreeDeveloper/Private/Debugger/SCompactTreeDebuggerView.cpp b/Engine/Plugins/Runtime/StateTree/Source/StateTreeDeveloper/Private/Debugger/SCompactTreeDebuggerView.cpp
--- a/Engine/Plugins/Runtime/StateTree/Source/StateTreeDeveloper/Private/Debugger/SCompactTreeDebuggerView.cpp
+++ b/Engine/Plugins/Runtime/StateTree/Source/StateTreeDeveloper/Private/Debugger/SCompactTreeDebuggerView.cpp
@@ -33,8 +33,9 @@ void SCompactTreeDebuggerView::CacheStatesInternal()
 
 		if (PerAssetActiveStates && PerAssetActiveStates->ActiveStates.Num() > 0)
 		{
+			FStateTreeTraceActiveStates::FAssetActiveStates AssetActiveStates = *PerAssetActiveStates;
 			TArray<FProcessedState> ProcessedStates;
-			CacheStateRecursive(PerAssetActiveStates, RootItem, PerAssetActiveStates->ActiveStates[0].Index, ProcessedStates);
+			CacheStateRecursive(&AssetActiveStates, RootItem, AssetActiveStates.ActiveStates[0].Index, ProcessedStates);
 		}
 	}
 }
@@ -101,8 +102,9 @@ void SCompactTreeDebuggerView::CacheStateRecursive(
 					});
 				if (LinkedAssetActiveStates && LinkedAssetActiveStates->ActiveStates.Num() > 0)
 				{
+					FStateTreeTraceActiveStates::FAssetActiveStates AssetActiveStates = *LinkedAssetActiveStates;
 					TArray<FProcessedState> ProcessedStates;
-					CacheStateRecursive(LinkedAssetActiveStates, InParentItem, LinkedAssetActiveStates->ActiveStates[0].Index, OutProcessedStates);
+					CacheStateRecursive(&AssetActiveStates, InParentItem, AssetActiveStates.ActiveStates[0].Index, OutProcessedStates);
 				}
 			}
 		}
