[StateTreeDebugger] moved some breakpoint related functionalities to the ViewModel and updated details view extensions to hide Task and Transition breakpoints actions when debugger window is not available

diff --git a/Engine/Plugins/Runtime/StateTree/Source/StateTreeEditorModule/Private/Customizations/StateTreeEditorNodeDetails.cpp b/Engine/Plugins/Runtime/StateTree/Source/StateTreeEditorModule/Private/Customizations/StateTreeEditorNodeDetails.cpp
--- a/Engine/Plugins/Runtime/StateTree/Source/StateTreeEditorModule/Private/Customizations/StateTreeEditorNodeDetails.cpp
+++ b/Engine/Plugins/Runtime/StateTree/Source/StateTreeEditorModule/Private/Customizations/StateTreeEditorNodeDetails.cpp
@@ -605,7 +605,7 @@ void FStateTreeEditorNodeDetails::CustomizeHeader(TSharedRef<class IPropertyHand
 					.AutoWidth()
 					.VAlign(VAlign_Center)
 					[
-						UE::StateTreeEditor::DebuggerExtensions::CreateEditorNodeWidget(StructPropertyHandle, EditorData.Get())
+						UE::StateTreeEditor::DebuggerExtensions::CreateEditorNodeWidget(StructPropertyHandle, StateTreeViewModel)
 					]
 
 					// Browse To source Button
@@ -2006,7 +2006,7 @@ TSharedRef<SWidget> FStateTreeEditorNodeDetails::GenerateOptionsMenu()
 	MenuBuilder.EndSection();
 
 	// Append debugger items.
-	UE::StateTreeEditor::DebuggerExtensions::AppendEditorNodeMenuItems(MenuBuilder, StructProperty, EditorData.Get());
+	UE::StateTreeEditor::DebuggerExtensions::AppendEditorNodeMenuItems(MenuBuilder, StructProperty, StateTreeViewModel);
 
 	return MenuBuilder.MakeWidget();
 }
diff --git a/Engine/Plugins/Runtime/StateTree/Source/StateTreeEditorModule/Private/Customizations/StateTreeStateDetails.cpp b/Engine/Plugins/Runtime/StateTree/Source/StateTreeEditorModule/Private/Customizations/StateTreeStateDetails.cpp
--- a/Engine/Plugins/Runtime/StateTree/Source/StateTreeEditorModule/Private/Customizations/StateTreeStateDetails.cpp
+++ b/Engine/Plugins/Runtime/StateTree/Source/StateTreeEditorModule/Private/Customizations/StateTreeStateDetails.cpp
@@ -7,6 +7,7 @@
 #include "PropertyBagDetails.h"
 #include "PropertyRestriction.h"
 #include "StateTree.h"
+#include "StateTreeEditingSubsystem.h"
 #include "StateTreeEditor.h"
 #include "StateTreeEditorData.h"
 #include "StateTreeEditorNodeUtils.h"
@@ -35,7 +36,7 @@ void FStateTreeStateDetails::CustomizeDetails(IDetailLayoutBuilder& DetailBuilde
 
 	TArray<TWeakObjectPtr<UObject>> Objects;
 	DetailBuilder.GetObjectsBeingCustomized(Objects);
-	
+
 	// Find StateTreeEditorData associated with this panel.
 	UStateTreeEditorData* EditorData = nullptr;
 	for (TWeakObjectPtr<UObject>& WeakObject : Objects)
@@ -63,11 +64,17 @@ void FStateTreeStateDetails::CustomizeDetails(IDetailLayoutBuilder& DetailBuilde
 			}
 		}
 	}
-	
+
 	const UStateTreeSchema* Schema = EditorData ? EditorData->Schema : nullptr;
 	const FString SchemaPath = Schema ? Schema->GetClass()->GetPathName() : FString();
 	TWeakObjectPtr<UStateTreeEditorData> WeakEditorData = EditorData;
 
+	TSharedPtr<FStateTreeViewModel> ViewModel;
+	if (UStateTreeEditingSubsystem* StateTreeEditingSubsystem = EditorData != nullptr ? GEditor->GetEditorSubsystem<UStateTreeEditingSubsystem>() : nullptr)
+	{
+		ViewModel = StateTreeEditingSubsystem->FindOrAddViewModel(EditorData->GetOuterUStateTree());
+	}
+
 	const TSharedPtr<IPropertyHandle> IDProperty = DetailBuilder.GetProperty(GET_MEMBER_NAME_CHECKED(UStateTreeState, ID));
 	const TSharedPtr<IPropertyHandle> NameProperty = DetailBuilder.GetProperty(GET_MEMBER_NAME_CHECKED(UStateTreeState, Name));
 	const TSharedPtr<IPropertyHandle> TagProperty = DetailBuilder.GetProperty(GET_MEMBER_NAME_CHECKED(UStateTreeState, Tag));
@@ -88,10 +95,10 @@ void FStateTreeStateDetails::CustomizeDetails(IDetailLayoutBuilder& DetailBuilde
 	const TSharedPtr<IPropertyHandle> RequiredEventToEnterProperty = DetailBuilder.GetProperty(GET_MEMBER_NAME_CHECKED(UStateTreeState, RequiredEventToEnter));
 	const TSharedPtr<IPropertyHandle> CheckPrerequisitesProperty = DetailBuilder.GetProperty(GET_MEMBER_NAME_CHECKED(UStateTreeState, bCheckPrerequisitesWhenActivatingChildDirectly));
 	const TSharedPtr<IPropertyHandle> WeightProperty = DetailBuilder.GetProperty(GET_MEMBER_NAME_CHECKED(UStateTreeState, Weight));
-	
+
 	// Never show enabled
 	EnabledProperty->MarkHiddenByCustomization();
-	
+
 	// Show ID only for debugging
 	if (UE::StateTree::Editor::GbDisplayItemIds == false)
 	{
@@ -113,13 +120,13 @@ void FStateTreeStateDetails::CustomizeDetails(IDetailLayoutBuilder& DetailBuilde
 			.VAlign(VAlign_Center)
 			[
 				SNew(SHorizontalBox)
-		
+
 				// Debugger labels
 				+ SHorizontalBox::Slot()
 				.AutoWidth()
 				.VAlign(VAlign_Center)
 				[
-					UE::StateTreeEditor::DebuggerExtensions::CreateStateWidget(EnabledProperty, EditorData)
+					UE::StateTreeEditor::DebuggerExtensions::CreateStateWidget(EnabledProperty, ViewModel)
 				]
 				// Options
 				+ SHorizontalBox::Slot()
@@ -128,12 +135,12 @@ void FStateTreeStateDetails::CustomizeDetails(IDetailLayoutBuilder& DetailBuilde
 				[
 					SNew(SComboButton)
 					.ButtonStyle(FAppStyle::Get(), "SimpleButton")
-					.OnGetMenuContent_Lambda([EnabledProperty, WeakEditorData]()
+					.OnGetMenuContent_Lambda([EnabledProperty, ViewModel]()
 					{
 						FMenuBuilder MenuBuilder(/*ShouldCloseWindowAfterMenuSelection*/true, /*CommandList*/nullptr);
 						// Append debugger items.
-						UE::StateTreeEditor::DebuggerExtensions::AppendStateMenuItems(MenuBuilder, EnabledProperty, WeakEditorData.Get());
-						return MenuBuilder.MakeWidget();	
+						UE::StateTreeEditor::DebuggerExtensions::AppendStateMenuItems(MenuBuilder, EnabledProperty, ViewModel);
+						return MenuBuilder.MakeWidget();
 					})
 					.ToolTipText(LOCTEXT("ItemActions", "Item actions"))
 					.HasDownArrow(false)
@@ -160,7 +167,7 @@ void FStateTreeStateDetails::CustomizeDetails(IDetailLayoutBuilder& DetailBuilde
 	// Color
 	ColorRefProperty->MarkHiddenByCustomization();
 	StateCategory.AddProperty(ColorRefProperty);
-	
+
 	// Custom Tick Rate
 	CustomTickRateProperty->MarkHiddenByCustomization();
 	if (Schema && Schema->IsScheduledTickAllowed())
@@ -252,7 +259,6 @@ void FStateTreeStateDetails::CustomizeDetails(IDetailLayoutBuilder& DetailBuilde
 		bool bFixedLayout = false;
 		ParametersFixedLayoutProperty->GetValue(bFixedLayout);
 
-		
 		const TSharedRef<SHorizontalBox> HeaderContentWidget = SNew(SHorizontalBox)
 			.IsEnabled(PropUtils, &IPropertyUtilities::IsPropertyEditingEnabled);
 
@@ -289,7 +295,7 @@ void FStateTreeStateDetails::CustomizeDetails(IDetailLayoutBuilder& DetailBuilde
 						FLinearColor(UE::StateTree::Colors::Blue)).ToSharedRef()
 				];
 		}
-		
+
 		ParametersCategory.HeaderContent(HeaderContentWidget, /*FullRowContent*/true);
 
 		FGuid ID;
@@ -300,7 +306,7 @@ void FStateTreeStateDetails::CustomizeDetails(IDetailLayoutBuilder& DetailBuilde
 			MakeShareable(new FStateTreeStateParametersInstanceDataDetails(ParametersProperty, ParametersProperty, PropUtils, bFixedLayout, ID, WeakEditorData, WeakState));
 		ParametersCategory.AddCustomBuilder(InstanceDetails);
 	}
-	
+
 	// Enter conditions
 	const FName EnterConditionsCategoryName(TEXT("Enter Conditions"));
 	if (Schema && Schema->AllowEnterConditions())
@@ -378,10 +384,10 @@ void FStateTreeStateDetails::CustomizeDetails(IDetailLayoutBuilder& DetailBuilde
 			const FName TaskCategoryName(TEXT("Task"));
 			IDetailCategoryBuilder& Category = DetailBuilder.EditCategory(TaskCategoryName);
 			Category.SetSortOrder(4);
-			
+
 			IDetailPropertyRow& Row = Category.AddProperty(SingleTaskProperty);
 			Row.ShouldAutoExpand(true);
-			
+
 			TasksProperty->MarkHiddenByCustomization();
 		}
 	}
@@ -403,8 +409,8 @@ void FStateTreeStateDetails::CustomizeDetails(IDetailLayoutBuilder& DetailBuilde
 		LOCTEXT("StateDetailsTransitionsAddTooltip", "Add new Transition"),
 		/*SortOrder*/5);
 
-	// Refresh the UI when the type changes.	
-	TypeProperty->SetOnPropertyValueChanged(FSimpleDelegate::CreateLambda([WeakPropertyUtilities = WeakPropertyUtilities] ()
+	// Refresh the UI when the type changes.
+	TypeProperty->SetOnPropertyValueChanged(FSimpleDelegate::CreateLambda([WeakPropertyUtilities = WeakPropertyUtilities]
 		{
 			if (const TSharedPtr<IPropertyUtilities> PropertyUtilities = WeakPropertyUtilities.Pin())
 			{
@@ -420,7 +426,6 @@ void FStateTreeStateDetails::PostUndo(bool bSuccess)
 	{
 		PropertyUtilities->ForceRefresh();
 	}
-	
 }
 
 void FStateTreeStateDetails::PostRedo(bool bSuccess)
diff --git a/Engine/Plugins/Runtime/StateTree/Source/StateTreeEditorModule/Private/Customizations/StateTreeTransitionDetails.cpp b/Engine/Plugins/Runtime/StateTree/Source/StateTreeEditorModule/Private/Customizations/StateTreeTransitionDetails.cpp
--- a/Engine/Plugins/Runtime/StateTree/Source/StateTreeEditorModule/Private/Customizations/StateTreeTransitionDetails.cpp
+++ b/Engine/Plugins/Runtime/StateTree/Source/StateTreeEditorModule/Private/Customizations/StateTreeTransitionDetails.cpp
@@ -4,28 +4,29 @@
 #include "Debugger/StateTreeDebuggerUIExtensions.h"
 #include "DetailLayoutBuilder.h"
 #include "DetailWidgetRow.h"
+#include "Framework/Application/SlateApplication.h"
 #include "HAL/PlatformApplicationMisc.h"
 #include "IDetailChildrenBuilder.h"
 #include "IPropertyUtilities.h"
 #include "ScopedTransaction.h"
+#include "StateTreeBindingExtension.h"
 #include "StateTreeDescriptionHelpers.h"
+#include "StateTreeEditingSubsystem.h"
 #include "StateTreeEditor.h"
 #include "StateTreeEditorData.h"
 #include "StateTreeEditorNodeUtils.h"
 #include "StateTreeEditorStyle.h"
 #include "StateTreePropertyHelpers.h"
+#include "StateTreeScopedEditorDataFixer.h"
 #include "StateTreeTypes.h"
+#include "TextStyleDecorator.h"
+#include "Widgets/Images/SImage.h"
+#include "Widgets/Input/SButton.h"
+#include "Widgets/Input/SComboButton.h"
 #include "Widgets/Layout/SBox.h"
 #include "Widgets/SBoxPanel.h"
 #include "Widgets/Text/SRichTextBlock.h"
 #include "Widgets/Text/STextBlock.h"
-#include "Widgets/Images/SImage.h"
-#include "Widgets/Input/SComboButton.h"
-#include "Widgets/Input/SButton.h"
-#include "TextStyleDecorator.h"
-#include "Framework/Application/SlateApplication.h"
-#include "StateTreeBindingExtension.h"
-#include "StateTreeScopedEditorDataFixer.h"
 
 #define LOCTEXT_NAMESPACE "StateTreeEditor"
 
@@ -54,6 +55,11 @@ void FStateTreeTransitionDetails::CustomizeHeader(TSharedRef<IPropertyHandle> St
 		}
 	}
 
+	if (UStateTreeEditingSubsystem* StateTreeEditingSubsystem = EditorData != nullptr ? GEditor->GetEditorSubsystem<UStateTreeEditingSubsystem>() : nullptr)
+	{
+		StateTreeViewModel = StateTreeEditingSubsystem->FindOrAddViewModel(EditorData->GetOuterUStateTree());
+	}
+
 	TriggerProperty = StructProperty->GetChildHandle(GET_MEMBER_NAME_CHECKED(FStateTreeTransition, Trigger));
 	PriorityProperty = StructProperty->GetChildHandle(GET_MEMBER_NAME_CHECKED(FStateTreeTransition, Priority));
 	RequiredEventProperty = StructProperty->GetChildHandle(GET_MEMBER_NAME_CHECKED(FStateTreeTransition, RequiredEvent));
@@ -120,9 +126,9 @@ void FStateTreeTransitionDetails::CustomizeHeader(TSharedRef<IPropertyHandle> St
 					.AutoWidth()
 					.VAlign(VAlign_Center)
 					[
-						UE::StateTreeEditor::DebuggerExtensions::CreateTransitionWidget(StructPropertyHandle, EditorData)
+						UE::StateTreeEditor::DebuggerExtensions::CreateTransitionWidget(StructPropertyHandle, StateTreeViewModel)
 					]
-					
+
 					// Options
 					+ SHorizontalBox::Slot()
 					.AutoWidth()
@@ -147,7 +153,7 @@ void FStateTreeTransitionDetails::CustomizeHeader(TSharedRef<IPropertyHandle> St
 		.CopyAction(FUIAction(FExecuteAction::CreateSP(this, &FStateTreeTransitionDetails::OnCopyTransition)))
 		.PasteAction(FUIAction(FExecuteAction::CreateSP(this, &FStateTreeTransitionDetails::OnPasteTransition)));
 }
- 
+
 void FStateTreeTransitionDetails::CustomizeChildren(TSharedRef<class IPropertyHandle> StructPropertyHandle, class IDetailChildrenBuilder& StructBuilder, IPropertyTypeCustomizationUtils& StructCustomizationUtils)
 {
 	check(TriggerProperty);
@@ -174,11 +180,11 @@ void FStateTreeTransitionDetails::CustomizeChildren(TSharedRef<class IPropertyHa
 	{
 		StructBuilder.AddProperty(IDProperty.ToSharedRef());
 	}
-	
+
 	// Trigger
 	StructBuilder.AddProperty(TriggerProperty.ToSharedRef());
 
-	// Show event only when the trigger is set to Event. 
+	// Show event only when the trigger is set to Event.
 	StructBuilder.AddProperty(RequiredEventProperty.ToSharedRef())
 		.Visibility(TAttribute<EVisibility>::Create(TAttribute<EVisibility>::FGetter::CreateLambda([WeakSelf]()
 		{
@@ -288,7 +294,7 @@ FReply FStateTreeTransitionDetails::OnRowMouseUp(const FGeometry& MyGeometry, co
 
 		return FReply::Handled();
 	}
-	
+
 	return FReply::Unhandled();
 }
 
@@ -351,7 +357,7 @@ TSharedRef<SWidget> FStateTreeTransitionDetails::GenerateOptionsMenu()
 	MenuBuilder.EndSection();
 
 	// Append debugger items.
-	UE::StateTreeEditor::DebuggerExtensions::AppendTransitionMenuItems(MenuBuilder, StructProperty, GetEditorData());
+	UE::StateTreeEditor::DebuggerExtensions::AppendTransitionMenuItems(MenuBuilder, StructProperty, StateTreeViewModel);
 
 	return MenuBuilder.MakeWidget();
 }
@@ -484,7 +490,7 @@ void FStateTreeTransitionDetails::OnCopyTransition() const
 	if (StructProperty->GetValueAsFormattedString(Value, PPF_Copy) == FPropertyAccess::Success)
 	{
 		FPlatformApplicationMisc::ClipboardCopy(*Value);
-	} 
+	}
 }
 
 UStateTreeEditorData* FStateTreeTransitionDetails::GetEditorData() const
diff --git a/Engine/Plugins/Runtime/StateTree/Source/StateTreeEditorModule/Private/Customizations/StateTreeTransitionDetails.h b/Engine/Plugins/Runtime/StateTree/Source/StateTreeEditorModule/Private/Customizations/StateTreeTransitionDetails.h
--- a/Engine/Plugins/Runtime/StateTree/Source/StateTreeEditorModule/Private/Customizations/StateTreeTransitionDetails.h
+++ b/Engine/Plugins/Runtime/StateTree/Source/StateTreeEditorModule/Private/Customizations/StateTreeTransitionDetails.h
@@ -5,6 +5,7 @@
 #include "IPropertyTypeCustomization.h"
 #include "Input/Reply.h"
 
+class FStateTreeViewModel;
 class IPropertyHandle;
 class IPropertyUtilities;
 class IDetailChildrenBuilder;
@@ -48,7 +49,8 @@ private:
 	void OnDeleteTransition() const;
 	void OnDeleteAllTransitions() const;
 	void OnDuplicateTransition() const;
-	
+
+	TSharedPtr<FStateTreeViewModel> StateTreeViewModel;
 	TSharedPtr<IPropertyHandle> TriggerProperty;
 	TSharedPtr<IPropertyHandle> PriorityProperty;
 	TSharedPtr<IPropertyHandle> RequiredEventProperty;
@@ -59,8 +61,8 @@ private:
 	TSharedPtr<IPropertyHandle> DelayRandomVarianceProperty;
 	TSharedPtr<IPropertyHandle> ConditionsProperty;
 	TSharedPtr<IPropertyHandle> IDProperty;
-	
-	TSharedPtr<SBorder> RowBorder; 
+
+	TSharedPtr<SBorder> RowBorder;
 
 	TSharedPtr<IPropertyUtilities> PropUtils;
 	TSharedPtr<IPropertyHandle> StructProperty;
diff --git a/Engine/Plugins/Runtime/StateTree/Source/StateTreeEditorModule/Private/Debugger/SStateTreeDebuggerView.cpp b/Engine/Plugins/Runtime/StateTree/Source/StateTreeEditorModule/Private/Debugger/SStateTreeDebuggerView.cpp
--- a/Engine/Plugins/Runtime/StateTree/Source/StateTreeEditorModule/Private/Debugger/SStateTreeDebuggerView.cpp
+++ b/Engine/Plugins/Runtime/StateTree/Source/StateTreeEditorModule/Private/Debugger/SStateTreeDebuggerView.cpp
@@ -71,7 +71,7 @@ void SStateTreeDebuggerView::StopRecording()
 {
 	if (CanStopRecording())
 	{
-		// Calling StopTraces on the module will notify all regsitered views through OnTracesStateChanged delegate.
+		// Calling StopTraces on the module will notify all registered views through OnTracesStateChanged delegate.
 		IStateTreeModule& StateTreeModule = FModuleManager::GetModuleChecked<IStateTreeModule>("StateTreeModule");
 		StateTreeModule.StopTraces();
 	}
diff --git a/Engine/Plugins/Runtime/StateTree/Source/StateTreeEditorModule/Private/Debugger/StateTreeDebuggerUIExtensions.cpp b/Engine/Plugins/Runtime/StateTree/Source/StateTreeEditorModule/Private/Debugger/StateTreeDebuggerUIExtensions.cpp
--- a/Engine/Plugins/Runtime/StateTree/Source/StateTreeEditorModule/Private/Debugger/StateTreeDebuggerUIExtensions.cpp
+++ b/Engine/Plugins/Runtime/StateTree/Source/StateTreeEditorModule/Private/Debugger/StateTreeDebuggerUIExtensions.cpp
@@ -12,6 +12,7 @@
 #include "StateTreeEditorNode.h"
 #include "StateTreeEditorStyle.h"
 #include "StateTreeTaskBase.h"
+#include "StateTreeViewModel.h"
 #include "Styling/StyleColors.h"
 #include "Widgets/Images/SImage.h"
 #include "Widgets/Input/SComboButton.h"
@@ -93,57 +94,35 @@ bool HasStateBreakpoint(const TArray<TWeakObjectPtr<>>& StatesBeingCustomized, c
 	return false;
 }
 
-void OnStateBreakpointToggled(const TArray<TWeakObjectPtr<>>& States, UStateTreeEditorData* EditorData, const EStateTreeBreakpointType Type)
+void OnStateBreakpointToggled(const TArray<TWeakObjectPtr<>>& States, const TSharedPtr<FStateTreeViewModel>& StateTreeViewModel, const EStateTreeBreakpointType Type)
 {
 #if WITH_STATETREE_TRACE_DEBUGGER
-	if (EditorData == nullptr || States.IsEmpty())
+	if (StateTreeViewModel != nullptr)
 	{
-		return;
-	}
-
-	const FScopedTransaction Transaction(LOCTEXT("ToggleStateBreakpoint", "Toggle State Breakpoint"));
-	EditorData->Modify(/*bAlwaysMarkDirty*/false);
-
-	for (const TWeakObjectPtr<>& WeakStateObject : States)
-	{
-		if (const UStateTreeState* State = static_cast<const UStateTreeState*>(WeakStateObject.Get()))
-		{
-			const bool bRemoved = EditorData->RemoveBreakpoint(State->ID, Type);
-			if (bRemoved == false)
-			{
-				EditorData->AddBreakpoint(State->ID, Type);
-			}
-		}
+		StateTreeViewModel->ToggleStateBreakpoints(States, Type);
 	}
 #endif // WITH_STATETREE_TRACE_DEBUGGER
 }
 
-bool HasTaskBreakpoint(TSharedPtr<IPropertyHandle> StructProperty, const UStateTreeEditorData* EditorData, const EStateTreeBreakpointType Type)
+bool HasTaskBreakpoint(TSharedPtr<IPropertyHandle> StructProperty, const TSharedPtr<FStateTreeViewModel>& StateTreeViewModel, const EStateTreeBreakpointType Type)
 {
 #if WITH_STATETREE_TRACE_DEBUGGER
 	const FStateTreeEditorNode* Node = EditorNodeUtils::GetCommonNode(StructProperty);
-	if (EditorData != nullptr && Node != nullptr)
+	if (StateTreeViewModel != nullptr && Node != nullptr)
 	{
-		return EditorData->HasBreakpoint(Node->ID, Type);
+		return StateTreeViewModel->HasBreakpoint(Node->ID, Type);
 	}
 #endif // WITH_STATETREE_TRACE_DEBUGGER
 	return false;
 }
 
-void OnTaskBreakpointToggled(TSharedPtr<IPropertyHandle> StructProperty, UStateTreeEditorData* EditorData, const EStateTreeBreakpointType Type)
+void OnTaskBreakpointToggled(TSharedPtr<IPropertyHandle> StructProperty, const TSharedPtr<FStateTreeViewModel>& StateTreeViewModel, const EStateTreeBreakpointType Type)
 {
 #if WITH_STATETREE_TRACE_DEBUGGER
 	const FStateTreeEditorNode* Node = EditorNodeUtils::GetMutableCommonNode(StructProperty);
-	if (EditorData != nullptr && Node != nullptr)
+	if (StateTreeViewModel != nullptr && Node != nullptr)
 	{
-		const FScopedTransaction Transaction(LOCTEXT("ToggleTaskBreakpoint", "Toggle Task Breakpoint"));
-		EditorData->Modify(/*bAlwaysMarkDirty*/false);
-		
-		const bool bRemoved = EditorData->RemoveBreakpoint(Node->ID, Type);
-		if (bRemoved == false)
-		{
-			EditorData->AddBreakpoint(Node->ID, Type);
-		}
+		StateTreeViewModel->ToggleTaskBreakpoint(Node->ID, Type);
 	}
 #endif // WITH_STATETREE_TRACE_DEBUGGER
 }
@@ -197,10 +176,10 @@ ECheckBoxState GetTransitionBreakpointCheckState(const TSharedPtr<IPropertyHandl
 			if (TreeState->Transitions.IsValidIndex(IndexInArray))
 			{
 				return EditorData.HasBreakpoint(TreeState->Transitions[IndexInArray].ID, EStateTreeBreakpointType::OnTransition)
-					? ECheckBoxState::Checked : ECheckBoxState::Unchecked;	
+					? ECheckBoxState::Checked : ECheckBoxState::Unchecked;
 			}
 
-			return  ECheckBoxState::Unchecked;
+			return ECheckBoxState::Unchecked;
 		});
 #else
 	return  ECheckBoxState::Unchecked;
@@ -265,30 +244,27 @@ void OnTransitionEnableToggled(const TSharedPtr<IPropertyHandle>& StructProperty
 	}
 }
 
-void OnTransitionBreakpointToggled(const TSharedPtr<IPropertyHandle>& StructPropertyHandle, UStateTreeEditorData* EditorData)
-{	
+void OnTransitionBreakpointToggled(const TSharedPtr<IPropertyHandle>& StructPropertyHandle, const TSharedPtr<FStateTreeViewModel>& StateTreeViewModel)
+{
 #if WITH_STATETREE_TRACE_DEBUGGER
 	if (StructPropertyHandle.IsValid() == false
 		|| StructPropertyHandle->IsValidHandle() == false
-		|| EditorData == nullptr)
+		|| StateTreeViewModel == nullptr)
 	{
 		return;
 	}
-	
+
 	// Determine desired toggled state. Undermined will switch all to checked (i.e. Enable breakpoint)
-	const ECheckBoxState CommonState = GetTransitionBreakpointCheckState(StructPropertyHandle, EditorData);
-	const ECheckBoxState ToggledState = (CommonState == ECheckBoxState::Checked) ? ECheckBoxState::Unchecked : ECheckBoxState::Checked; 
-	
+	const ECheckBoxState CommonState = GetTransitionBreakpointCheckState(StructPropertyHandle, StateTreeViewModel->GetStateTreeEditorData());
+	const ECheckBoxState ToggledState = (CommonState == ECheckBoxState::Checked) ? ECheckBoxState::Unchecked : ECheckBoxState::Checked;
 	const IPropertyHandle& PropertyHandle = *(StructPropertyHandle.Get());
 
 	TArray<UObject*> OuterObjects;
+	TArray<TNotNull<const FStateTreeTransition*>> Transitions;
 	PropertyHandle.GetOuterObjects(OuterObjects);
 
 	if (OuterObjects.Num() > 0)
 	{
-		const FScopedTransaction Transaction(LOCTEXT("TransitionBreakpointToggled", "Transition Breakpoint Toggled"));
-		EditorData->Modify(/*bAlwaysMarkDirty*/false);
-		
 		for (const UObject* OuterObject : OuterObjects)
 		{
 			if (const UStateTreeState* TreeState = Cast<UStateTreeState>(OuterObject))
@@ -297,21 +273,14 @@ void OnTransitionBreakpointToggled(const TSharedPtr<IPropertyHandle>& StructProp
 
 				if (TreeState->Transitions.IsValidIndex(IndexInArray))
 				{
-					const FStateTreeTransition& Transition = TreeState->Transitions[IndexInArray];
-					const bool bHasBreakpoint = EditorData->HasBreakpoint(Transition.ID, EStateTreeBreakpointType::OnTransition);
-					EditorData->HasBreakpoint(Transition.ID, EStateTreeBreakpointType::OnTransition);
-					if (ToggledState == ECheckBoxState::Checked && bHasBreakpoint == false)
-					{
-						EditorData->AddBreakpoint(Transition.ID, EStateTreeBreakpointType::OnTransition);	
-					}
-					else if (ToggledState == ECheckBoxState::Unchecked && bHasBreakpoint)
-					{
-						EditorData->RemoveBreakpoint(Transition.ID, EStateTreeBreakpointType::OnTransition);
-					}		
+					Transitions.Push(&TreeState->Transitions[IndexInArray]);
 				}
 			}
 		}
 	}
+
+	StateTreeViewModel->ToggleTransitionBreakpoint(Transitions, ToggledState);
+
 #endif // WITH_STATETREE_TRACE_DEBUGGER
 }
 
@@ -319,7 +288,7 @@ TArray<TWeakObjectPtr<UObject>> GetWeakOuterObjects(const TSharedPtr<IPropertyHa
 {
 	TArray<UObject*> OuterObjects;
 	PropertyHandle->GetOuterObjects(OuterObjects);
-	
+
 	TArray<TWeakObjectPtr<UObject>> WeakOuterObjects;
 	for (UObject* Object : OuterObjects)
 	{
@@ -328,14 +297,14 @@ TArray<TWeakObjectPtr<UObject>> GetWeakOuterObjects(const TSharedPtr<IPropertyHa
 	return WeakOuterObjects;
 }
 
-TSharedRef<SWidget> CreateStateWidget(TSharedPtr<IPropertyHandle> StateEnabledProperty, UStateTreeEditorData* TreeData)
+TSharedRef<SWidget> CreateStateWidget(TSharedPtr<IPropertyHandle> StateEnabledProperty, const TSharedPtr<FStateTreeViewModel>& StateTreeViewModel)
 {
-	TWeakObjectPtr<UStateTreeEditorData> WeakTreeData = TreeData;
+	TWeakObjectPtr<const UStateTreeEditorData> WeakTreeData = StateTreeViewModel != nullptr ? StateTreeViewModel->GetStateTreeEditorData() : nullptr;
 	TArray<TWeakObjectPtr<UObject>> StatesBeingCustomized = GetWeakOuterObjects(StateEnabledProperty);
 
 	const TSharedRef<SHorizontalBox> HeaderContentWidget = SNew(SHorizontalBox)
 		.Clipping(EWidgetClipping::OnDemand)
-		
+
 		// Disabled Label
 		+ SHorizontalBox::Slot()
 		.AutoWidth()
@@ -359,7 +328,7 @@ TSharedRef<SWidget> CreateStateWidget(TSharedPtr<IPropertyHandle> StateEnabledPr
 				.ToolTipText(LOCTEXT("DisabledStateTooltip", "This state has been disabled."))
 			]
 		]
-		
+
 		// Break on enter Label
 		+ SHorizontalBox::Slot()
 		.AutoWidth()
@@ -409,11 +378,11 @@ TSharedRef<SWidget> CreateStateWidget(TSharedPtr<IPropertyHandle> StateEnabledPr
 	return HeaderContentWidget;
 }
 
-void AppendStateMenuItems(FMenuBuilder& InMenuBuilder, TSharedPtr<IPropertyHandle> StateEnabledProperty, UStateTreeEditorData* TreeData)
+void AppendStateMenuItems(FMenuBuilder& InMenuBuilder, TSharedPtr<IPropertyHandle> StateEnabledProperty, const TSharedPtr<FStateTreeViewModel>& StateTreeViewModel)
 {
-	TWeakObjectPtr<UStateTreeEditorData> WeakTreeData = TreeData;
+	TWeakObjectPtr WeakTreeData = StateTreeViewModel != nullptr ? StateTreeViewModel->GetStateTreeEditorData() : nullptr;
 	TArray<TWeakObjectPtr<UObject>> StatesBeingCustomized = GetWeakOuterObjects(StateEnabledProperty);
-	
+
 	InMenuBuilder.BeginSection(FName("Options"), LOCTEXT("StateOptions", "State Debug Options"));
 	InMenuBuilder.AddMenuEntry
 	(
@@ -449,15 +418,18 @@ void AppendStateMenuItems(FMenuBuilder& InMenuBuilder, TSharedPtr<IPropertyHandl
 	InMenuBuilder.AddMenuEntry
 	(
 		LOCTEXT("ToggleOnEnterStateBreakpoint", "Break on Enter"),
-		LOCTEXT("ToggleOnEnterStateBreakpoint_ToolTip", "Enables or disables breakpoint when entering selected state(s)."),
+		LOCTEXT("ToggleOnEnterStateBreakpoint_ToolTip", "Enables or disables breakpoint when entering selected state(s). (Debugger Window required)"),
 		FSlateIcon(),
 		FUIAction
 		(
-			FExecuteAction::CreateLambda([StatesBeingCustomized, WeakTreeData]
+			FExecuteAction::CreateLambda([StatesBeingCustomized, StateTreeViewModel]
 				{
-					OnStateBreakpointToggled(StatesBeingCustomized, WeakTreeData.Get(), EStateTreeBreakpointType::OnEnter);
+					OnStateBreakpointToggled(StatesBeingCustomized, StateTreeViewModel, EStateTreeBreakpointType::OnEnter);
+				}),
+			FCanExecuteAction::CreateLambda([StateTreeViewModel]
+				{
+					return StateTreeViewModel != nullptr && StateTreeViewModel->CanProcessBreakpoints();
 				}),
-			FCanExecuteAction(),
 			FGetActionCheckState::CreateLambda([StatesBeingCustomized, WeakTreeData]
 				{
 					return HasStateBreakpoint(StatesBeingCustomized, WeakTreeData.Get(), EStateTreeBreakpointType::OnEnter) ? ECheckBoxState::Checked : ECheckBoxState::Unchecked;
@@ -469,15 +441,19 @@ void AppendStateMenuItems(FMenuBuilder& InMenuBuilder, TSharedPtr<IPropertyHandl
 	InMenuBuilder.AddMenuEntry
 	(
 		LOCTEXT("ToggleOnExitStateBreakpoint", "Break on Exit"),
-		LOCTEXT("ToggleOnExitStateBreakpoint_ToolTip", "Enables or disables breakpoint when exiting selected state(s)."),
+		LOCTEXT("ToggleOnExitStateBreakpoint_ToolTip", "Enables or disables breakpoint when exiting selected state(s). (Debugger Window required)"),
 		FSlateIcon(),
 		FUIAction
 		(
-			FExecuteAction::CreateLambda([StatesBeingCustomized, WeakTreeData]
+			FExecuteAction::CreateLambda([StatesBeingCustomized, StateTreeViewModel]
 				{
-					OnStateBreakpointToggled(StatesBeingCustomized, WeakTreeData.Get(), EStateTreeBreakpointType::OnExit);
+					OnStateBreakpointToggled(StatesBeingCustomized, StateTreeViewModel, EStateTreeBreakpointType::OnExit);
+				}),
+			FCanExecuteAction::CreateLambda([StateTreeViewModel]
+				{
+					return StateTreeViewModel != nullptr && (StateTreeViewModel->CanAddStateBreakpoint(EStateTreeBreakpointType::OnExit)
+							|| StateTreeViewModel->CanRemoveStateBreakpoint(EStateTreeBreakpointType::OnExit));
 				}),
-			FCanExecuteAction(),
 			FGetActionCheckState::CreateLambda([StatesBeingCustomized, WeakTreeData]
 				{
 					return HasStateBreakpoint(StatesBeingCustomized, WeakTreeData.Get(), EStateTreeBreakpointType::OnExit) ? ECheckBoxState::Checked : ECheckBoxState::Unchecked;
@@ -487,15 +463,13 @@ void AppendStateMenuItems(FMenuBuilder& InMenuBuilder, TSharedPtr<IPropertyHandl
 		EUserInterfaceActionType::Check
 	);
 	InMenuBuilder.EndSection();
-	
+
 }
 
-TSharedRef<SWidget> CreateEditorNodeWidget(const TSharedPtr<IPropertyHandle>& StructPropertyHandle, UStateTreeEditorData* TreeData)
+TSharedRef<SWidget> CreateEditorNodeWidget(const TSharedPtr<IPropertyHandle>& StructPropertyHandle, const TSharedPtr<FStateTreeViewModel>& StateTreeViewModel)
 {
 	using namespace EditorNodeUtils;
 
-	TWeakObjectPtr<UStateTreeEditorData> WeakTreeData = TreeData;
-	
 	bool bIsTask = false;
 	bool bIsCondition = false;
 	if (const FStructProperty* TmpStructProperty = CastField<FStructProperty>(StructPropertyHandle->GetProperty()))
@@ -516,7 +490,7 @@ TSharedRef<SWidget> CreateEditorNodeWidget(const TSharedPtr<IPropertyHandle>& St
 		}
 	}
 
-	// In case there is no common Editor node or not associated to a condition or a task we don't need to add any widget.	
+	// In case there is no common Editor node or not associated to a condition or a task we don't need to add any widget.
 	if (bIsCondition == false && bIsTask == false)
 	{
 		return SNullWidget::NullWidget;
@@ -563,9 +537,9 @@ TSharedRef<SWidget> CreateEditorNodeWidget(const TSharedPtr<IPropertyHandle>& St
 				SNew(SBorder)
 				.Padding(FMargin(6.0f, 2.0f))
 				.BorderImage(FStateTreeEditorStyle::Get().GetBrush("StateTree.Node.Label"))
-				.Visibility_Lambda([StructPropertyHandle, WeakTreeData]
+				.Visibility_Lambda([StructPropertyHandle, StateTreeViewModel]
 					{
-						return HasTaskBreakpoint(StructPropertyHandle, WeakTreeData.Get(), EStateTreeBreakpointType::OnEnter) ? EVisibility::Visible : EVisibility::Collapsed;
+						return HasTaskBreakpoint(StructPropertyHandle, StateTreeViewModel, EStateTreeBreakpointType::OnEnter) ? EVisibility::Visible : EVisibility::Collapsed;
 					})
 				[
 					SNew(STextBlock)
@@ -586,9 +560,9 @@ TSharedRef<SWidget> CreateEditorNodeWidget(const TSharedPtr<IPropertyHandle>& St
 				SNew(SBorder)
 				.Padding(FMargin(6.0f, 2.0f))
 				.BorderImage(FStateTreeEditorStyle::Get().GetBrush("StateTree.Node.Label"))
-				.Visibility_Lambda([StructPropertyHandle, WeakTreeData]
+				.Visibility_Lambda([StructPropertyHandle, StateTreeViewModel]
 					{
-						return HasTaskBreakpoint(StructPropertyHandle, WeakTreeData.Get(), EStateTreeBreakpointType::OnExit) ? EVisibility::Visible : EVisibility::Collapsed;
+						return HasTaskBreakpoint(StructPropertyHandle, StateTreeViewModel, EStateTreeBreakpointType::OnExit) ? EVisibility::Visible : EVisibility::Collapsed;
 					})
 				[
 					SNew(STextBlock)
@@ -648,12 +622,10 @@ TSharedRef<SWidget> CreateEditorNodeWidget(const TSharedPtr<IPropertyHandle>& St
 		];
 }
 
-void AppendEditorNodeMenuItems(FMenuBuilder& InMenuBuilder, const TSharedPtr<IPropertyHandle>& StructPropertyHandle, UStateTreeEditorData* TreeData)
+void AppendEditorNodeMenuItems(FMenuBuilder& InMenuBuilder, const TSharedPtr<IPropertyHandle>& StructPropertyHandle, const TSharedPtr<FStateTreeViewModel>& StateTreeViewModel)
 {
 	using namespace EditorNodeUtils;
 
-	TWeakObjectPtr<UStateTreeEditorData> WeakTreeData = TreeData;
-	
 	bool bIsTask = false;
 	bool bIsCondition = false;
 	if (const FStructProperty* TmpStructProperty = CastField<FStructProperty>(StructPropertyHandle->GetProperty()))
@@ -674,7 +646,7 @@ void AppendEditorNodeMenuItems(FMenuBuilder& InMenuBuilder, const TSharedPtr<IPr
 		}
 	}
 
-	// In case there is no common Editor node or not associated to a condition or a task we don't need to add any widget.	
+	// In case there is no common Editor node or not associated to a condition or a task we don't need to add any widget.
 	if (bIsCondition == false && bIsTask == false)
 	{
 		return;
@@ -704,15 +676,21 @@ void AppendEditorNodeMenuItems(FMenuBuilder& InMenuBuilder, const TSharedPtr<IPr
 		InMenuBuilder.AddMenuEntry
 		(
 			LOCTEXT("ToggleOnEnterTaskBreakpoint", "Break on Enter"),
-			LOCTEXT("ToggleOnEnterTaskBreakpoint_ToolTip", "Enables or disables breakpoint when entering selected task(s)."),
+			LOCTEXT("ToggleOnEnterTaskBreakpoint_ToolTip", "Enables or disables breakpoint when entering selected task(s). (Debugger Window required)"),
 			FSlateIcon(),
 			FUIAction
 			(
-				FExecuteAction::CreateLambda([StructPropertyHandle, WeakTreeData] { OnTaskBreakpointToggled(StructPropertyHandle, WeakTreeData.Get(), EStateTreeBreakpointType::OnEnter); }),
-				FCanExecuteAction(),
-				FGetActionCheckState::CreateLambda([StructPropertyHandle, WeakTreeData]
+				FExecuteAction::CreateLambda([StructPropertyHandle, StateTreeViewModel]
+					{
+						OnTaskBreakpointToggled(StructPropertyHandle, StateTreeViewModel, EStateTreeBreakpointType::OnEnter);
+					}),
+				FCanExecuteAction::CreateLambda([StateTreeViewModel]
 					{
-						return HasTaskBreakpoint(StructPropertyHandle, WeakTreeData.Get(), EStateTreeBreakpointType::OnEnter) ? ECheckBoxState::Checked : ECheckBoxState::Unchecked;
+						return StateTreeViewModel != nullptr && StateTreeViewModel->CanProcessBreakpoints();
+					}),
+				FGetActionCheckState::CreateLambda([StructPropertyHandle, StateTreeViewModel]
+					{
+						return HasTaskBreakpoint(StructPropertyHandle, StateTreeViewModel, EStateTreeBreakpointType::OnEnter) ? ECheckBoxState::Checked : ECheckBoxState::Unchecked;
 					})
 			),
 			NAME_None,
@@ -721,15 +699,21 @@ void AppendEditorNodeMenuItems(FMenuBuilder& InMenuBuilder, const TSharedPtr<IPr
 		InMenuBuilder.AddMenuEntry
 		(
 			LOCTEXT("ToggleOnExitTaskBreakpoint", "Break on Exit"),
-			LOCTEXT("ToggleOnExitTaskBreakpoint_ToolTip", "Enables or disables breakpoint when exiting selected task(s)."),
+			LOCTEXT("ToggleOnExitTaskBreakpoint_ToolTip", "Enables or disables breakpoint when exiting selected task(s). (Debugger Window required)"),
 			FSlateIcon(),
 			FUIAction
 			(
-				FExecuteAction::CreateLambda([StructPropertyHandle, WeakTreeData] { OnTaskBreakpointToggled(StructPropertyHandle, WeakTreeData.Get(), EStateTreeBreakpointType::OnExit); }),
-				FCanExecuteAction(),
-				FGetActionCheckState::CreateLambda([StructPropertyHandle, WeakTreeData]
+				FExecuteAction::CreateLambda([StructPropertyHandle, StateTreeViewModel]
 					{
-						return HasTaskBreakpoint(StructPropertyHandle, WeakTreeData.Get(), EStateTreeBreakpointType::OnExit) ? ECheckBoxState::Checked : ECheckBoxState::Unchecked;
+						OnTaskBreakpointToggled(StructPropertyHandle, StateTreeViewModel, EStateTreeBreakpointType::OnExit);
+					}),
+				FCanExecuteAction::CreateLambda([StateTreeViewModel]
+					{
+						return StateTreeViewModel != nullptr && StateTreeViewModel->CanProcessBreakpoints();
+					}),
+				FGetActionCheckState::CreateLambda([StructPropertyHandle, StateTreeViewModel]
+					{
+						return HasTaskBreakpoint(StructPropertyHandle, StateTreeViewModel, EStateTreeBreakpointType::OnExit) ? ECheckBoxState::Checked : ECheckBoxState::Unchecked;
 					})
 			),
 			NAME_None,
@@ -821,11 +805,11 @@ bool IsEditorNodeEnabled(const TSharedPtr<IPropertyHandle>& StructPropertyHandle
 	return true;
 }
 
-TSharedRef<SWidget> CreateTransitionWidget(const TSharedPtr<IPropertyHandle>& StructPropertyHandle, UStateTreeEditorData* TreeData)
+TSharedRef<SWidget> CreateTransitionWidget(const TSharedPtr<IPropertyHandle>& StructPropertyHandle, const TSharedPtr<FStateTreeViewModel>& StateTreeViewModel)
 {
 	using namespace EditorNodeUtils;
 
-	TWeakObjectPtr<UStateTreeEditorData> WeakTreeData = TreeData;
+	TWeakObjectPtr WeakTreeData = StateTreeViewModel != nullptr ? StateTreeViewModel->GetStateTreeEditorData() : nullptr;
 
 	return SNew(SHorizontalBox)
 		+ SHorizontalBox::Slot()
@@ -893,11 +877,11 @@ TSharedRef<SWidget> CreateTransitionWidget(const TSharedPtr<IPropertyHandle>& St
 		];
 }
 
-void AppendTransitionMenuItems(FMenuBuilder& InMenuBuilder, const TSharedPtr<IPropertyHandle>& StructPropertyHandle, UStateTreeEditorData* TreeData)
+void AppendTransitionMenuItems(FMenuBuilder& InMenuBuilder, const TSharedPtr<IPropertyHandle>& StructPropertyHandle, const TSharedPtr<FStateTreeViewModel>& StateTreeViewModel)
 {
 	using namespace EditorNodeUtils;
 
-	TWeakObjectPtr<UStateTreeEditorData> WeakTreeData = TreeData;
+	TWeakObjectPtr WeakTreeData = StateTreeViewModel != nullptr ? StateTreeViewModel->GetStateTreeEditorData() : nullptr;
 
 	InMenuBuilder.BeginSection(FName("Options"), LOCTEXT("TransitionOptions", "Transition Debug Options"));
 	InMenuBuilder.AddMenuEntry
@@ -918,12 +902,18 @@ void AppendTransitionMenuItems(FMenuBuilder& InMenuBuilder, const TSharedPtr<IPr
 	InMenuBuilder.AddMenuEntry
 	(
 		LOCTEXT("ToggleTransitionBreakpoint", "Break on transition"),
-		LOCTEXT("ToggleTransitionBreakpoint_ToolTip", "Enables or disables breakpoint when executing the transition(s)."),
+		LOCTEXT("ToggleTransitionBreakpoint_ToolTip", "Enables or disables breakpoint when executing the transition(s). (Debugger Window required)"),
 		FSlateIcon(),
 		FUIAction
 		(
-			FExecuteAction::CreateLambda([StructPropertyHandle, WeakTreeData] { OnTransitionBreakpointToggled(StructPropertyHandle, WeakTreeData.Get()); }),
-			FCanExecuteAction(),
+			FExecuteAction::CreateLambda([StructPropertyHandle, StateTreeViewModel]
+			{
+				OnTransitionBreakpointToggled(StructPropertyHandle, StateTreeViewModel);
+			}),
+			FCanExecuteAction::CreateLambda([StateTreeViewModel]
+				{
+					return StateTreeViewModel != nullptr && StateTreeViewModel->CanProcessBreakpoints();
+				}),
 			FGetActionCheckState::CreateLambda([StructPropertyHandle, WeakTreeData]
 				{
 					return GetTransitionBreakpointCheckState(StructPropertyHandle, WeakTreeData.Get());
diff --git a/Engine/Plugins/Runtime/StateTree/Source/StateTreeEditorModule/Private/Debugger/StateTreeDebuggerUIExtensions.h b/Engine/Plugins/Runtime/StateTree/Source/StateTreeEditorModule/Private/Debugger/StateTreeDebuggerUIExtensions.h
--- a/Engine/Plugins/Runtime/StateTree/Source/StateTreeEditorModule/Private/Debugger/StateTreeDebuggerUIExtensions.h
+++ b/Engine/Plugins/Runtime/StateTree/Source/StateTreeEditorModule/Private/Debugger/StateTreeDebuggerUIExtensions.h
@@ -18,15 +18,15 @@ class UStateTreeEditorData;
 namespace UE::StateTreeEditor::DebuggerExtensions
 {
 
-TSharedRef<SWidget> CreateStateWidget(TSharedPtr<IPropertyHandle> StateEnabledProperty, UStateTreeEditorData* TreeData);
-void AppendStateMenuItems(FMenuBuilder& InMenuBuilder, TSharedPtr<IPropertyHandle> StateEnabledProperty, UStateTreeEditorData* TreeData);
+TSharedRef<SWidget> CreateStateWidget(TSharedPtr<IPropertyHandle> StateEnabledProperty, const TSharedPtr<FStateTreeViewModel>& InStateTreeViewModel);
+void AppendStateMenuItems(FMenuBuilder& InMenuBuilder, TSharedPtr<IPropertyHandle> StateEnabledProperty, const TSharedPtr<FStateTreeViewModel>& InStateTreeViewModel);
 
-TSharedRef<SWidget> CreateEditorNodeWidget(const TSharedPtr<IPropertyHandle>& StructPropertyHandle, UStateTreeEditorData* TreeData);
-void AppendEditorNodeMenuItems(FMenuBuilder& InMenuBuilder, const TSharedPtr<IPropertyHandle>& StructPropertyHandle, UStateTreeEditorData* TreeData);
+TSharedRef<SWidget> CreateEditorNodeWidget(const TSharedPtr<IPropertyHandle>& StructPropertyHandle, const TSharedPtr<FStateTreeViewModel>& InStateTreeViewModel);
+void AppendEditorNodeMenuItems(FMenuBuilder& InMenuBuilder, const TSharedPtr<IPropertyHandle>& StructPropertyHandle, const TSharedPtr<FStateTreeViewModel>& InStateTreeViewModel);
 bool IsEditorNodeEnabled(const TSharedPtr<IPropertyHandle>& StructPropertyHandle);
 
-TSharedRef<SWidget> CreateTransitionWidget(const TSharedPtr<IPropertyHandle>& StructPropertyHandle, UStateTreeEditorData* TreeData);
-void AppendTransitionMenuItems(FMenuBuilder& InMenuBuilder, const TSharedPtr<IPropertyHandle>& StructPropertyHandle, UStateTreeEditorData* TreeData);
+TSharedRef<SWidget> CreateTransitionWidget(const TSharedPtr<IPropertyHandle>& StructPropertyHandle, const TSharedPtr<FStateTreeViewModel>& InStateTreeViewModel);
+void AppendTransitionMenuItems(FMenuBuilder& InMenuBuilder, const TSharedPtr<IPropertyHandle>& StructPropertyHandle, const TSharedPtr<FStateTreeViewModel>& InStateTreeViewModel);
 bool IsTransitionEnabled(const TSharedPtr<IPropertyHandle>& StructPropertyHandle);
 
 }; // UE::StateTreeEditor::DebuggerExtensions
diff --git a/Engine/Plugins/Runtime/StateTree/Source/StateTreeEditorModule/Private/StateTreeViewModel.cpp b/Engine/Plugins/Runtime/StateTree/Source/StateTreeEditorModule/Private/StateTreeViewModel.cpp
--- a/Engine/Plugins/Runtime/StateTree/Source/StateTreeEditorModule/Private/StateTreeViewModel.cpp
+++ b/Engine/Plugins/Runtime/StateTree/Source/StateTreeEditorModule/Private/StateTreeViewModel.cpp
@@ -440,9 +440,89 @@ void FStateTreeViewModel::HandleIdentifierChanged(const UStateTree& StateTree) c
 }
 
 #if WITH_STATETREE_TRACE_DEBUGGER
+
+void FStateTreeViewModel::ToggleStateBreakpoints(TConstArrayView<TWeakObjectPtr<>> States, EStateTreeBreakpointType Type)
+{
+	UStateTreeEditorData* EditorData = TreeDataWeak.Get();
+	if (EditorData == nullptr || States.IsEmpty())
+	{
+		return;
+	}
+
+	const FScopedTransaction Transaction(LOCTEXT("ToggleStateBreakpoint", "Toggle State Breakpoint"));
+	EditorData->Modify(/*bAlwaysMarkDirty*/false);
+
+	for (const TWeakObjectPtr<>& WeakStateObject : States)
+	{
+		if (const UStateTreeState* State = static_cast<const UStateTreeState*>(WeakStateObject.Get()))
+		{
+			const bool bRemoved = EditorData->RemoveBreakpoint(State->ID, Type);
+			if (bRemoved == false)
+			{
+				EditorData->AddBreakpoint(State->ID, Type);
+			}
+		}
+	}
+}
+
+void FStateTreeViewModel::ToggleTaskBreakpoint(const FGuid ID, const EStateTreeBreakpointType Type)
+{
+	if (UStateTreeEditorData* EditorData = TreeDataWeak.Get())
+	{
+		const FScopedTransaction Transaction(LOCTEXT("ToggleTaskBreakpoint", "Toggle Task Breakpoint"));
+		EditorData->Modify(/*bAlwaysMarkDirty*/false);
+
+		const bool bRemoved = EditorData->RemoveBreakpoint(ID, Type);
+		if (bRemoved == false)
+		{
+			EditorData->AddBreakpoint(ID, Type);
+		}
+	}
+}
+
+void FStateTreeViewModel::ToggleTransitionBreakpoint(const TConstArrayView<TNotNull<const FStateTreeTransition*>> Transitions, const ECheckBoxState ToggledState)
+{
+	UStateTreeEditorData* EditorData = TreeDataWeak.Get();
+	if (EditorData == nullptr)
+	{
+		return;
+	}
+
+	const FScopedTransaction Transaction(LOCTEXT("TransitionBreakpointToggled", "Transition Breakpoint Toggled"));
+	EditorData->Modify(/*bAlwaysMarkDirty*/false);
+
+	for (const FStateTreeTransition* Transition : Transitions)
+	{
+		const bool bHasBreakpoint = EditorData->HasBreakpoint(Transition->ID, EStateTreeBreakpointType::OnTransition);
+		if (ToggledState == ECheckBoxState::Checked && bHasBreakpoint == false)
+		{
+			EditorData->AddBreakpoint(Transition->ID, EStateTreeBreakpointType::OnTransition);
+		}
+		else if (ToggledState == ECheckBoxState::Unchecked && bHasBreakpoint)
+		{
+			EditorData->RemoveBreakpoint(Transition->ID, EStateTreeBreakpointType::OnTransition);
+		}
+	}
+}
+
+bool FStateTreeViewModel::HasBreakpoint(const FGuid ID, const EStateTreeBreakpointType Type)
+{
+	if (UStateTreeEditorData* TreeData = TreeDataWeak.Get())
+	{
+		return TreeData->HasBreakpoint(ID, Type);
+	}
+
+	return false;
+}
+
+bool FStateTreeViewModel::CanProcessBreakpoints() const
+{
+	return Debugger.Get().CanProcessBreakpoints();
+}
+
 bool FStateTreeViewModel::CanAddStateBreakpoint(const EStateTreeBreakpointType Type) const
 {
-	if (!Debugger.Get().CanProcessBreakpoints())
+	if (!CanProcessBreakpoints())
 	{
 		return false;
 	}
diff --git a/Engine/Plugins/Runtime/StateTree/Source/StateTreeEditorModule/Public/StateTreeViewModel.h b/Engine/Plugins/Runtime/StateTree/Source/StateTreeEditorModule/Public/StateTreeViewModel.h
--- a/Engine/Plugins/Runtime/StateTree/Source/StateTreeEditorModule/Public/StateTreeViewModel.h
+++ b/Engine/Plugins/Runtime/StateTree/Source/StateTreeEditorModule/Public/StateTreeViewModel.h
@@ -4,10 +4,13 @@
 
 #include "UObject/WeakObjectPtr.h"
 #include "EditorUndoClient.h"
+#include "PropertyHandle.h"
+#include "StateTreeEditorNode.h"
 #include "StateTreeViewModel.generated.h"
 
 #define UE_API STATETREEEDITORMODULE_API
 
+struct FStateTreeTransition;
 class FMenuBuilder;
 class UStateTree;
 class UStateTreeEditorData;
@@ -124,10 +127,15 @@ public:
 
 	// Debugging
 #if WITH_STATETREE_TRACE_DEBUGGER
+	UE_API bool HasBreakpoint(FGuid ID, EStateTreeBreakpointType Type);
+	UE_API bool CanProcessBreakpoints() const;
 	UE_API bool CanAddStateBreakpoint(EStateTreeBreakpointType Type) const;
 	UE_API bool CanRemoveStateBreakpoint(EStateTreeBreakpointType Type) const;
 	UE_API ECheckBoxState GetStateBreakpointCheckState(EStateTreeBreakpointType Type) const;
 	UE_API void HandleEnableStateBreakpoint(EStateTreeBreakpointType Type);
+	UE_API void ToggleStateBreakpoints(TConstArrayView<TWeakObjectPtr<>> States, EStateTreeBreakpointType Type);
+	UE_API void ToggleTaskBreakpoint(FGuid ID, EStateTreeBreakpointType Type);
+	UE_API void ToggleTransitionBreakpoint(TConstArrayView<TNotNull<const FStateTreeTransition*>> Transitions, ECheckBoxState ToggledState);
 
 	UE_API UStateTreeState* FindStateAssociatedToBreakpoint(FStateTreeDebuggerBreakpoint Breakpoint) const;
 
@@ -188,7 +196,7 @@ public:
 	{
 		return OnBringNodeToFocus;
 	}
-	 
+
 protected:
 	UE_API void GetExpandedStatesRecursive(UStateTreeState* State, TSet<TWeakObjectPtr<UStateTreeState>>& ExpandedStates);
 
