[StateTree] Move binding refresh from editor data to compiler manager. It prevents multithreaded access to the delegate. It will be optimized later with dependencies.

diff --git a/Engine/Plugins/Runtime/StateTree/Source/StateTreeEditorModule/Private/StateTreeCompilerManager.cpp b/Engine/Plugins/Runtime/StateTree/Source/StateTreeEditorModule/Private/StateTreeCompilerManager.cpp
--- a/Engine/Plugins/Runtime/StateTree/Source/StateTreeEditorModule/Private/StateTreeCompilerManager.cpp
+++ b/Engine/Plugins/Runtime/StateTree/Source/StateTreeEditorModule/Private/StateTreeCompilerManager.cpp
@@ -5,6 +5,7 @@
 #include "StateTreeCompilerLog.h"
 #include "StateTreeDelegates.h"
 #include "StateTreeEditingSubsystem.h"
+#include "StateTreeEditorData.h"
 
 #include "Editor.h"
 #include "StructUtilsDelegates.h"
@@ -26,13 +27,12 @@ public:
 
 private:
 	bool HandleCompileStateTree(UStateTree& StateTree);
+	void UpdateBindingsInstanceStructsIfNeeded(TSet<const UStruct*>& Structs, TNotNull<UStateTreeEditorData*> EditorData);
 
-#if WITH_EDITOR
 	using FReplacementObjectMap = TMap<UObject*, UObject*>;
 	void HandleObjectsReinstanced(const FReplacementObjectMap& ObjectMap);
 	void HandlePreBeginPIE(bool bIsSimulating);
 	void HandleUserDefinedStructReinstanced(const UUserDefinedStruct& UserDefinedStruct);
-#endif
 
 private:
 	FDelegateHandle ObjectsReinstancedHandle;
@@ -57,6 +57,30 @@ FCompilerManagerImpl::~FCompilerManagerImpl()
 	UE::StateTree::Delegates::OnRequestCompile.Unbind();
 }
 
+void FCompilerManagerImpl::UpdateBindingsInstanceStructsIfNeeded(TSet<const UStruct*>& Structs, TNotNull<UStateTreeEditorData*> EditorData)
+{
+	bool bShouldUpdate = false;
+	EditorData->VisitAllNodes([&Structs, &bShouldUpdate](const UStateTreeState* State, const FStateTreeBindableStructDesc& Desc, const FStateTreeDataView Value)
+		{
+			if (Structs.Contains(Value.GetStruct()))
+			{
+				bShouldUpdate = true;
+				return EStateTreeVisitor::Break;
+			}
+			return EStateTreeVisitor::Continue;
+		});
+
+	if (!bShouldUpdate)
+	{
+		bShouldUpdate = EditorData->GetEditorPropertyBindings()->ContainsAnyStruct(Structs);
+	}
+
+	if (bShouldUpdate)
+	{
+		EditorData->UpdateBindingsInstanceStructs();
+	}
+}
+
 bool FCompilerManagerImpl::HandleCompileStateTree(UStateTree& StateTree)
 {
 	FStateTreeCompilerLog Log;
@@ -104,6 +128,12 @@ void FCompilerManagerImpl::HandleObjectsReinstanced(const FReplacementObjectMap&
 		}
 	}
 
+	for (TObjectIterator<UStateTreeEditorData> It; It; ++It)
+	{
+		UStateTreeEditorData* StateTreeEditorData = *It;
+		UpdateBindingsInstanceStructsIfNeeded(Structs, StateTreeEditorData);
+	}
+
 	for (TObjectIterator<UStateTree> It; It; ++It)
 	{
 		UStateTree* StateTree = *It;
@@ -166,6 +196,12 @@ void FCompilerManagerImpl::HandleUserDefinedStructReinstanced(const UUserDefined
 	TSet<const UStruct*> Structs;
 	Structs.Add(&UserDefinedStruct);
 
+	for (TObjectIterator<UStateTreeEditorData> It; It; ++It)
+	{
+		UStateTreeEditorData* StateTreeEditorData = *It;
+		UpdateBindingsInstanceStructsIfNeeded(Structs, StateTreeEditorData);
+	}
+
 	for (TObjectIterator<UStateTree> It; It; ++It)
 	{
 		UStateTree* StateTree = *It;
diff --git a/Engine/Plugins/Runtime/StateTree/Source/StateTreeEditorModule/Private/StateTreeEditorData.cpp b/Engine/Plugins/Runtime/StateTree/Source/StateTreeEditorModule/Private/StateTreeEditorData.cpp
--- a/Engine/Plugins/Runtime/StateTree/Source/StateTreeEditorModule/Private/StateTreeEditorData.cpp
+++ b/Engine/Plugins/Runtime/StateTree/Source/StateTreeEditorModule/Private/StateTreeEditorData.cpp
@@ -92,8 +92,6 @@ void UStateTreeEditorData::PostInitProperties()
 #if WITH_EDITOR
 	if (!HasAnyFlags(RF_ClassDefaultObject))
 	{
-		OnObjectsReinstancedHandle = FCoreUObjectDelegates::OnObjectsReinstanced.AddUObject(this, &UStateTreeEditorData::OnObjectsReinstanced);
-		OnUserDefinedStructReinstancedHandle = UE::StructUtils::Delegates::OnUserDefinedStructReinstanced.AddUObject(this, &UStateTreeEditorData::OnUserDefinedStructReinstanced);
 		OnParametersChangedHandle = UE::StateTree::Delegates::OnParametersChanged.AddUObject(this, &UStateTreeEditorData::OnParametersChanged);
 		OnStateParametersChangedHandle = UE::StateTree::Delegates::OnStateParametersChanged.AddUObject(this, &UStateTreeEditorData::OnStateParametersChanged);
 	}
@@ -111,16 +109,6 @@ void UStateTreeEditorData::Serialize(FArchive& Ar)
 
 void UStateTreeEditorData::BeginDestroy()
 {
-	if (OnObjectsReinstancedHandle.IsValid())
-	{
-		FCoreUObjectDelegates::OnObjectsReinstanced.Remove(OnObjectsReinstancedHandle);
-		OnObjectsReinstancedHandle.Reset();
-	}
-	if (OnUserDefinedStructReinstancedHandle.IsValid())
-	{
-		UE::StructUtils::Delegates::OnUserDefinedStructReinstanced.Remove(OnUserDefinedStructReinstancedHandle);
-		OnUserDefinedStructReinstancedHandle.Reset();
-	}
 	if (OnParametersChangedHandle.IsValid())
 	{
 		UE::StateTree::Delegates::OnParametersChanged.Remove(OnParametersChangedHandle);
@@ -135,71 +123,6 @@ void UStateTreeEditorData::BeginDestroy()
 	Super::BeginDestroy();
 }
 
-void UStateTreeEditorData::OnObjectsReinstanced(const FReplacementObjectMap& ObjectMap)
-{
-	if (ObjectMap.IsEmpty())
-	{
-		return;
-	}
-	
-	TSet<const UStruct*> Structs;
-	for (TMap<UObject*, UObject*>::TConstIterator It(ObjectMap); It; ++It)
-	{
-		if (const UObject* ObjectToBeReplaced = It->Value)
-		{
-			Structs.Add(ObjectToBeReplaced->GetClass());
-		}
-	}
-
-	bool bShouldUpdate = false;
-	VisitAllNodes([&Structs, &bShouldUpdate](const UStateTreeState* State, const FStateTreeBindableStructDesc& Desc, const FStateTreeDataView Value)
-	{
-		if (Structs.Contains(Value.GetStruct()))
-		{
-			bShouldUpdate = true;
-			return EStateTreeVisitor::Break; 
-		}
-		return EStateTreeVisitor::Continue;
-	});
-
-	if (!bShouldUpdate)
-	{
-		bShouldUpdate = EditorBindings.ContainsAnyStruct(Structs);
-	}
-	
-	if (bShouldUpdate)
-	{
-		UpdateBindingsInstanceStructs();
-	}
-}
-
-void UStateTreeEditorData::OnUserDefinedStructReinstanced(const UUserDefinedStruct& UserDefinedStruct)
-{
-	TSet<const UStruct*> Structs;
-	Structs.Add(&UserDefinedStruct);
-
-	bool bShouldUpdate = false;
-	VisitAllNodes([&Structs, &bShouldUpdate](const UStateTreeState* State, const FStateTreeBindableStructDesc& Desc, const FStateTreeDataView Value)
-	{
-		if (Structs.Contains(Value.GetStruct()))
-		{
-			bShouldUpdate = true;
-			return EStateTreeVisitor::Break; 
-		}
-		return EStateTreeVisitor::Continue;
-	});
-
-	if (!bShouldUpdate)
-	{
-		bShouldUpdate = EditorBindings.ContainsAnyStruct(Structs);
-	}
-	
-	if (bShouldUpdate)
-	{
-		UpdateBindingsInstanceStructs();
-	}
-}
-
 void UStateTreeEditorData::OnParametersChanged(const UStateTree& StateTree)
 {
 	if (const UStateTree* OwnerStateTree = GetTypedOuter<UStateTree>())
diff --git a/Engine/Plugins/Runtime/StateTree/Source/StateTreeEditorModule/Public/StateTreeEditorData.h b/Engine/Plugins/Runtime/StateTree/Source/StateTreeEditorModule/Public/StateTreeEditorData.h
--- a/Engine/Plugins/Runtime/StateTree/Source/StateTreeEditorModule/Public/StateTreeEditorData.h
+++ b/Engine/Plugins/Runtime/StateTree/Source/StateTreeEditorModule/Public/StateTreeEditorData.h
@@ -23,6 +23,10 @@ namespace UE::StateTree::Editor
 	extern STATETREEEDITORMODULE_API const FString PropertyFunctionStateName;
 }
 
+namespace UE::StateTree::Compiler::Private
+{
+	class FCompilerManagerImpl;
+}
 
 USTRUCT()
 struct FStateTreeEditorBreakpoint
@@ -106,9 +110,6 @@ public:
 	UE_API FText GetNodeDescription(const FStateTreeEditorNode& Node, const EStateTreeNodeFormatting Formatting = EStateTreeNodeFormatting::Text) const;
 
 #if WITH_EDITOR
-	using FReplacementObjectMap = TMap<UObject*, UObject*>;
-	UE_API void OnObjectsReinstanced(const FReplacementObjectMap& ObjectMap);
-	UE_API void OnUserDefinedStructReinstanced(const UUserDefinedStruct& UserDefinedStruct);
 	UE_API void OnParametersChanged(const UStateTree& StateTree);
 	UE_API void OnStateParametersChanged(const UStateTree& StateTree, const FGuid StateID);
 	UE_API virtual void BeginDestroy() override;
@@ -335,8 +336,6 @@ private:
 	UE_API void CallPostLoadOnNodes();
 
 #if WITH_EDITORONLY_DATA
-	FDelegateHandle OnObjectsReinstancedHandle;
-	FDelegateHandle OnUserDefinedStructReinstancedHandle;
 	FDelegateHandle OnParametersChangedHandle;
 	FDelegateHandle OnStateParametersChangedHandle;
 #endif
@@ -398,6 +397,7 @@ public:
 	TArray<FStateTreeEditorDelegateDispatcherCompiledBinding> CompiledDispatchers;
 
 	friend class FStateTreeEditorDataDetails;
+	friend class UE::StateTree::Compiler::Private::FCompilerManagerImpl;
 };
 
 
