[StateTree] Use NonconcurrentLinearArrayAllocator for scoped allocation in execution context.

diff --git a/Engine/Plugins/Runtime/StateTree/Source/StateTreeModule/Private/StateTreeExecutionContext.cpp b/Engine/Plugins/Runtime/StateTree/Source/StateTreeModule/Private/StateTreeExecutionContext.cpp
--- a/Engine/Plugins/Runtime/StateTree/Source/StateTreeModule/Private/StateTreeExecutionContext.cpp
+++ b/Engine/Plugins/Runtime/StateTree/Source/StateTreeModule/Private/StateTreeExecutionContext.cpp
@@ -129,7 +129,8 @@ namespace UE::StateTree::ExecutionContext::Private
 
 	// From metric, paths are usually of length < 8. We put a large enough number to not have to worry about complex assets.
 	static constexpr int32 ExpectedActiveStatePathLength = 24;
-	void GetActiveStatePath(TArrayView<const FStateTreeExecutionFrame> Frames, TArray<UE::StateTree::FActiveState, TInlineAllocator<ExpectedActiveStatePathLength>>& OutResult)
+	using FActiveStateInlineArray = TArray<FActiveState, TInlineAllocator<ExpectedActiveStatePathLength, FNonconcurrentLinearArrayAllocator>>;
+	void GetActiveStatePath(TArrayView<const FStateTreeExecutionFrame> Frames, FActiveStateInlineArray& OutResult)
 	{
 		if (Frames.Num() == 0 || Frames[0].StateTree == nullptr)
 		{
@@ -145,7 +146,8 @@ namespace UE::StateTree::ExecutionContext::Private
 		}
 	}
 
-	bool GetStatesListToState(FStateHandleContext State, TArray<FStateHandleContext, TInlineAllocator<FStateTreeActiveStates::MaxStates>>& OutStates)
+	using FStateHandleContextInlineArray = TArray<FStateHandleContext, TInlineAllocator<FStateTreeActiveStates::MaxStates>>;
+	bool GetStatesListToState(FStateHandleContext State, FStateHandleContextInlineArray& OutStates)
 	{
 		// Walk towards the root from the state.
 		FStateHandleContext CurrentState = State;
@@ -163,7 +165,8 @@ namespace UE::StateTree::ExecutionContext::Private
 		return true;
 	}
 
-	using FSharedEventInlineArray = TArray<FStateTreeSharedEvent, TInlineAllocator<8>>;
+	static constexpr int32 ExpectedAmountTransitionEvent = 8;
+	using FSharedEventInlineArray = TArray<FStateTreeSharedEvent, TInlineAllocator<ExpectedAmountTransitionEvent, FNonconcurrentLinearArrayAllocator>>;
 	void GetTriggerTransitionEvent(const FCompactStateTransition& Transition, FStateTreeInstanceStorage& Storage, const FStateTreeSharedEvent& TransitionEvent, const TArrayView<const FStateTreeSharedEvent> EventsToProcess, FSharedEventInlineArray& OutTransitionEvents)
 	{
 		if (Transition.Trigger == EStateTreeTransitionTrigger::OnEvent)
@@ -2138,14 +2141,14 @@ void FStateTreeExecutionContext::UpdateInstanceData(FSelectStateResult& Args)
 		}
 	}
 
-	TArray<FConstStructView> InstanceStructs;
+	TArray<FConstStructView, FNonconcurrentLinearArrayAllocator> InstanceStructs;
 	InstanceStructs.Reserve(EstimatedNumStructs);
 
-	TArray<FInstancedStruct*> TempInstanceStructs;
+	TArray<FInstancedStruct*, FNonconcurrentLinearArrayAllocator> TempInstanceStructs;
 	TempInstanceStructs.Reserve(EstimatedNumStructs);
 
 	constexpr int32 ExpectedAmountOfFrames = 8;
-	TArray<FCompactStateTreeParameters, TInlineAllocator<ExpectedAmountOfFrames>> TempParams;
+	TArray<FCompactStateTreeParameters, TInlineAllocator<ExpectedAmountOfFrames, FNonconcurrentLinearArrayAllocator>> TempParams;
 
 	TArrayView<FStateTreeTemporaryInstanceData> TempInstances = Storage.GetMutableTemporaryInstances();
 	auto FindInstanceTempData = [&TempInstances](FActiveFrameID FrameID, FStateTreeDataHandle DataHandle)
@@ -2482,7 +2485,7 @@ EStateTreeRunStatus FStateTreeExecutionContext::ForceTransition(const FRecordedS
 	using namespace UE::StateTree::ExecutionContext;
 	using namespace UE::StateTree::ExecutionContext::Private;
 
-	TArray<FStateHandleContext> StateContexts;
+	TArray<FStateHandleContext, FNonconcurrentLinearArrayAllocator> StateContexts;
 	StateContexts.Reserve(RecordedTransition.States.Num());
 	for (const FRecordedActiveState& RecordedActiveState : RecordedTransition.States)
 	{
@@ -2506,7 +2509,7 @@ EStateTreeRunStatus FStateTreeExecutionContext::ForceTransition(const FRecordedS
 		return EStateTreeRunStatus::Unset;
 	}
 
-	TArray<FStateTreeSharedEvent> SharedEvents;
+	TArray<FStateTreeSharedEvent, FNonconcurrentLinearArrayAllocator> SharedEvents;
 	SharedEvents.Reserve(RecordedTransition.Events.Num());
 	for (const FStateTreeEvent& SelectedEvent : RecordedTransition.Events)
 	{
@@ -2652,7 +2655,7 @@ bool FStateTreeExecutionContext::ForceTransitionInternal(const TArrayView<const
 		return false;
 	}
 
-	TArray<FActiveState, TInlineAllocator<ExpectedActiveStatePathLength>> CurrentActiveStatePath;
+	FActiveStateInlineArray CurrentActiveStatePath;
 	GetActiveStatePath(Exec.ActiveFrames, CurrentActiveStatePath);
 	if (CurrentActiveStatePath.Num() == 0)
 	{
@@ -3437,7 +3440,7 @@ void FStateTreeExecutionContext::ExitState(const FSelectStateResult& Args, const
 	STATETREE_LOG(Log, TEXT("Exit state '%s' (%d)"), *DebugGetStatePath(Exec.ActiveFrames), Exec.StateChangeCount);
 	UE_STATETREE_DEBUG_SCOPED_PHASE(this, EStateTreeUpdatePhase::ExitStates);
 
-	TArray<FActiveState, TInlineAllocator<ExpectedActiveStatePathLength>> ActiveStates;
+	FActiveStateInlineArray ActiveStates;
 	GetActiveStatePath(GetExecState().ActiveFrames, ActiveStates);
 
 	const TArrayView<const FActiveState> CommonStates = FActiveStatePath::Intersect(MakeConstArrayView(ActiveStates), MakeConstArrayView(Args.SelectedStates));
@@ -4834,7 +4837,7 @@ bool FStateTreeExecutionContext::RequestTransitionInternal(
 		return true;
 	}
 
-	TArray<FActiveState, TInlineAllocator<ExpectedActiveStatePathLength>> CurrentActiveStatePath;
+	FActiveStateInlineArray CurrentActiveStatePath;
 	GetActiveStatePath(GetExecState().ActiveFrames, CurrentActiveStatePath);
 
 	FSelectStateArguments SelectStateArgs;
@@ -5825,7 +5828,7 @@ bool FStateTreeExecutionContext::SelectState(const FStateTreeExecutionFrame& Cur
 	using namespace UE::StateTree::ExecutionContext;
 	using namespace UE::StateTree::ExecutionContext::Private;
 
-	TArray<FActiveState, TInlineAllocator<ExpectedActiveStatePathLength>> CurrentActiveStatePath;
+	FActiveStateInlineArray CurrentActiveStatePath;
 	GetActiveStatePath(GetExecState().ActiveFrames, CurrentActiveStatePath);
 
 	FSelectStateArguments SelectStateArgs;
