[StateTree] Can't clear the schema from in the editor.
#rnx

diff --git a/Engine/Plugins/Runtime/StateTree/Source/StateTreeEditorModule/Public/StateTreeEditorData.h b/Engine/Plugins/Runtime/StateTree/Source/StateTreeEditorModule/Public/StateTreeEditorData.h
--- a/Engine/Plugins/Runtime/StateTree/Source/StateTreeEditorModule/Public/StateTreeEditorData.h
+++ b/Engine/Plugins/Runtime/StateTree/Source/StateTreeEditorModule/Public/StateTreeEditorData.h
@@ -342,7 +342,7 @@ private:
 
 public:
 	/** Schema describing which inputs, evaluators, and tasks a StateTree can contain */	
-	UPROPERTY(EditDefaultsOnly, Instanced, Category = Common)
+	UPROPERTY(EditDefaultsOnly, Instanced, Category = Common, NoClear)
 	TObjectPtr<UStateTreeSchema> Schema = nullptr;
 
 	/** Public parameters that could be used for bindings within the Tree. */
