[StateTree] Add warning when a StateTree node uses deprecated instance data. The deprecated class won't be saved. It will break the cooked version.

diff --git a/Engine/Plugins/Runtime/StateTree/Source/StateTreeEditorModule/Private/StateTreeCompiler.cpp b/Engine/Plugins/Runtime/StateTree/Source/StateTreeEditorModule/Private/StateTreeCompiler.cpp
--- a/Engine/Plugins/Runtime/StateTree/Source/StateTreeEditorModule/Private/StateTreeCompiler.cpp
+++ b/Engine/Plugins/Runtime/StateTree/Source/StateTreeEditorModule/Private/StateTreeCompiler.cpp
@@ -2029,6 +2029,11 @@ bool FStateTreeCompiler::CompileAndValidateNode(const UStateTreeState* SourceSta
 		return false;
 	}
 	
+	if (const UClass* InstanceDataClass = Cast<UClass>(InstanceData.GetStruct()); InstanceDataClass && InstanceDataClass->HasAllClassFlags(CLASS_Deprecated))
+	{
+		Log.Report(EMessageSeverity::Warning, NodeDesc, TEXT("The instance data class is deprecated. It won't work in a cooked build."));
+	}
+
 	return Result != EDataValidationResult::Invalid;
 }
 
diff --git a/Engine/Plugins/Runtime/StateTree/Source/StateTreeModule/Private/StateTree.cpp b/Engine/Plugins/Runtime/StateTree/Source/StateTreeModule/Private/StateTree.cpp
--- a/Engine/Plugins/Runtime/StateTree/Source/StateTreeModule/Private/StateTree.cpp
+++ b/Engine/Plugins/Runtime/StateTree/Source/StateTreeModule/Private/StateTree.cpp
@@ -449,8 +449,9 @@ void UStateTree::ResetLinked()
 bool UStateTree::ValidateInstanceData()
 {
 	bool bResult = true;
-	for (FConstStructView NodeView : Nodes)
+	for (int32 Index = 0; Index < Nodes.Num(); ++Index)
 	{
+		const FConstStructView& NodeView = Nodes[Index];
 		const FStateTreeNodeBase* Node = NodeView.GetPtr<const FStateTreeNodeBase>();
 		if (Node && Node->InstanceTemplateIndex.IsValid())
 		{
@@ -471,7 +472,12 @@ bool UStateTree::ValidateInstanceData()
 
 			if (CurrentInstanceDataType == nullptr)
 			{
-				UE_LOG(LogStateTree, Error, TEXT("%s: node '%s' failed. Missing instance value, possibly due to Blueprint class or C++ class/struct template deletion."), *GetPathName(), *WriteToString<64>(Node->StaticStruct()->GetFName()));
+				UE_LOG(LogStateTree, Error, TEXT("%s: node (%d) '%s' with name '%s' failed. Missing instance value, possibly due to Blueprint class or C++ class/struct template deletion.")
+					, *GetPathName()
+					, Index
+					, *WriteToString<64>(Node->StaticStruct()->GetFName())
+					, *WriteToString<128>(Node->Name)
+					);
 
 				bResult = false;
 				continue;
