[StateTreeDebugger] Update state tree rewind debugger tracing
- use Instance ids instead of child elements (which makes state tree tracks display as a child of the graph track for anim next state trees)
- refactored StateTree traces to use same pattern as ObjectTrace to avoid buffering lifetime events (i.e., all traces will trace the instance information if needed)

diff --git a/Engine/Plugins/AI/MassAI/Source/MassAIBehavior/Private/MassStateTreeProcessors.cpp b/Engine/Plugins/AI/MassAI/Source/MassAIBehavior/Private/MassStateTreeProcessors.cpp
--- a/Engine/Plugins/AI/MassAI/Source/MassAIBehavior/Private/MassStateTreeProcessors.cpp
+++ b/Engine/Plugins/AI/MassAI/Source/MassAIBehavior/Private/MassStateTreeProcessors.cpp
@@ -16,6 +16,7 @@
 #include "MassStateTreeFragments.h"
 #include "MassStateTreeSubsystem.h"
 #include "MassZoneGraphAnnotationTypes.h"
+#include "ObjectTrace.h"
 #include "ProfilingDebugging/CsvProfiler.h"
 #include "StateTree.h"
 #include "VisualLogger/VisualLogger.h"
@@ -49,8 +50,13 @@ void ForEachEntityInChunk(
 		FStateTreeInstanceData* InstanceData = MassStateTreeSubsystem.GetInstanceData(StateTreeFragment.InstanceHandle);
 		if (InstanceData != nullptr)
 		{
+			const UObject* TraceOwner = Context.GetEntityManagerChecked().GetOwner();
+			TRACE_OBJECT(TraceOwner);
+			TRACE_INSTANCE(TraceOwner, Entity.AsNumber(), FObjectTrace::GetObjectId(TraceOwner), FMassEntityHandle::StaticStruct(), LexToString(Entity));
+
 			FMassStateTreeExecutionContext StateTreeContext(MassStateTreeSubsystem, *StateTree, *InstanceData, Context);
 			StateTreeContext.SetEntity(Entity);
+			StateTreeContext.SetOuterTraceId(Entity.AsNumber());
 
 			// Make sure all required external data are set.
 			{
diff --git a/Engine/Plugins/Experimental/UAF/UAFStateTree/Source/UAFStateTree/Private/AnimStateTreeTrait.cpp b/Engine/Plugins/Experimental/UAF/UAFStateTree/Source/UAFStateTree/Private/AnimStateTreeTrait.cpp
--- a/Engine/Plugins/Experimental/UAF/UAFStateTree/Source/UAFStateTree/Private/AnimStateTreeTrait.cpp
+++ b/Engine/Plugins/Experimental/UAF/UAFStateTree/Source/UAFStateTree/Private/AnimStateTreeTrait.cpp
@@ -96,14 +96,15 @@ namespace UE::UAF
 			{
 				return;
 			}
-#endif // ENABLE_ANIM_DEBUG 
+#endif // ENABLE_ANIM_DEBUG
 
 			FStateTreeExecutionContext StateTreeExecutionContext(*Owner, *InstanceData->StateTree, InstanceData->InstanceData);
 
 			FAnimNextStateTreeTraitContext TraitContext(Context, Binding);
-			StateTreeExecutionContext.SetContextDataByName(UStateTreeAnimNextSchema::AnimStateTreeExecutionContextName, FStateTreeDataView(FAnimNextStateTreeTraitContext::StaticStruct(), reinterpret_cast<uint8*>(&TraitContext)));
+			StateTreeExecutionContext.SetContextDataByName(UStateTreeAnimNextSchema::AnimStateTreeExecutionContextName, FStateTreeDataView(FAnimNextStateTreeTraitContext::StaticStruct(), &TraitContext));
 
 			FAnimNextGraphInstance& OwnerGraphInstance = Binding.GetTraitPtr().GetNodeInstance()->GetOwner();
+			StateTreeExecutionContext.SetOuterTraceId(OwnerGraphInstance.GetUniqueId());
 #if DO_CHECK
 			check(OwnerGraphInstance.LayoutMatches(InstanceData->StateTree->GetDefaultParameters()));
 #endif
@@ -212,13 +213,13 @@ namespace UE::UAF
 #endif // ENABLE_ANIM_DEBUG 
 			
 			FStateTreeExecutionContext StateTreeExecutionContext(*Owner, *InstanceData->StateTree, InstanceData->InstanceData);
-			FAnimNextGraphInstance* OwnerGraphInstance = &Binding.GetTraitPtr().GetNodeInstance()->GetOwner();
-			if(OwnerGraphInstance && StateTreeExecutionContext.GetLastTickStatus() != EStateTreeRunStatus::Failed)
+			FAnimNextGraphInstance& OwnerGraphInstance = Binding.GetTraitPtr().GetNodeInstance()->GetOwner();
+			if (StateTreeExecutionContext.GetLastTickStatus() != EStateTreeRunStatus::Failed)
 			{
 				FAnimNextStateTreeTraitContext TraitContext(Context, Binding);
 				StateTreeExecutionContext.SetContextDataByName(UStateTreeAnimNextSchema::AnimStateTreeExecutionContextName, FStateTreeDataView(FAnimNextStateTreeTraitContext::StaticStruct(), reinterpret_cast<uint8*>(&TraitContext)));
 				StateTreeExecutionContext.SetExternalGlobalParameters(&InstanceData->StateTreeExternalParameters);
-
+				StateTreeExecutionContext.SetOuterTraceId(OwnerGraphInstance.GetUniqueId());
 				StateTreeExecutionContext.SetCollectExternalDataCallback(FOnCollectStateTreeExternalData::CreateLambda([this, &TraitContext]( const FStateTreeExecutionContext& Context, const UStateTree* StateTree, TArrayView<const FStateTreeExternalDataDesc> ExternalDataDescs, TArrayView<FStateTreeDataView> OutDataViews) -> bool
 				{
 					for (int32 Index = 0; Index < ExternalDataDescs.Num(); Index++)
diff --git a/Engine/Plugins/Runtime/StateTree/Source/StateTreeEditorModule/Private/Debugger/StateTreeRewindDebuggerExtensions.cpp b/Engine/Plugins/Runtime/StateTree/Source/StateTreeEditorModule/Private/Debugger/StateTreeRewindDebuggerExtensions.cpp
--- a/Engine/Plugins/Runtime/StateTree/Source/StateTreeEditorModule/Private/Debugger/StateTreeRewindDebuggerExtensions.cpp
+++ b/Engine/Plugins/Runtime/StateTree/Source/StateTreeEditorModule/Private/Debugger/StateTreeRewindDebuggerExtensions.cpp
@@ -1,29 +1,67 @@
 // Copyright Epic Games, Inc. All Rights Reserved.
 
-#if WITH_STATETREE_TRACE_DEBUGGER
-
 #include "Debugger/StateTreeRewindDebuggerExtensions.h"
 
-#include "Debugger/StateTreeTrace.h"
-#include "Debugger/StateTreeTraceProvider.h"
+#if WITH_STATETREE_TRACE || WITH_STATETREE_TRACE_DEBUGGER
+#include "Debugger/StateTreeTraceTypes.h"
 #include "StateTreeDelegates.h"
-#include "Trace/Trace.h"
+#endif
+
+#if WITH_STATETREE_TRACE
+#include "Debugger/StateTreeTrace.h"
+#endif
+
+#if WITH_STATETREE_TRACE_DEBUGGER
+#include "IRewindDebugger.h"
+#endif
 
 namespace UE::StateTreeDebugger
 {
 
-void FRewindDebuggerExtension::Update(float DeltaTime, IRewindDebugger* RewindDebugger)
+#if WITH_STATETREE_TRACE
+//----------------------------------------------------------------------//
+// FRewindDebuggerRecordingExtension
+//----------------------------------------------------------------------//
+void FRewindDebuggerRecordingExtension::RecordingStarted()
+{
+	if (!UE_TRACE_CHANNELEXPR_IS_ENABLED(StateTreeDebugChannel))
+	{
+		Trace::ToggleChannel(TEXT("StateTreeDebugChannel"), true);
+		StateTree::Delegates::OnTracingStateChanged.Broadcast(EStateTreeTraceStatus::TracesStarted);
+	}
+}
+
+void FRewindDebuggerRecordingExtension::RecordingStopped()
+{
+	// Shouldn't fire as channel will already be disabled when rewind debugger stops, but just as safeguard
+	if (UE_TRACE_CHANNELEXPR_IS_ENABLED(StateTreeDebugChannel))
+	{
+		StateTree::Delegates::OnTracingStateChanged.Broadcast(EStateTreeTraceStatus::StoppingTrace);
+		Trace::ToggleChannel(TEXT("StateTreeDebugChannel"), false);
+	}
+}
+
+void FRewindDebuggerRecordingExtension::Clear()
+{
+	StateTree::Delegates::OnTracingStateChanged.Broadcast(EStateTreeTraceStatus::Cleared);
+}
+#endif // WITH_STATETREE_TRACE
+
+
+#if WITH_STATETREE_TRACE_DEBUGGER
+//----------------------------------------------------------------------//
+// FRewindDebuggerPlaybackExtension
+//----------------------------------------------------------------------//
+void FRewindDebuggerPlaybackExtension::Update(float DeltaTime, IRewindDebugger* RewindDebugger)
 {
 	const TraceServices::IAnalysisSession* AnalysisSession = RewindDebugger->GetAnalysisSession();
 	TraceServices::FAnalysisSessionReadScope SessionReadScope(*AnalysisSession);
-	
-	const IStateTreeTraceProvider* TraceProvider = AnalysisSession->ReadProvider<IStateTreeTraceProvider>(FStateTreeTraceProvider::ProviderName);
 
 	const TraceServices::IFrameProvider& FrameProvider = TraceServices::ReadFrameProvider(*RewindDebugger->GetAnalysisSession());
 	TraceServices::FFrame Frame;
 
 	// Require some debug frame to exist before doing any processing, currently the frame itself is not used
-	if (FrameProvider.GetFrameFromTime(ETraceFrameType::TraceFrameType_Game, RewindDebugger->CurrentTraceTime(), Frame))
+	if (FrameProvider.GetFrameFromTime(TraceFrameType_Game, RewindDebugger->CurrentTraceTime(), Frame))
 	{
 		auto SetScrubberTimeline = [&]()
 		{
@@ -35,7 +73,7 @@ void FRewindDebuggerExtension::Update(float DeltaTime, IRewindDebugger* RewindDe
 			const double CurrentScrubTime = RewindDebugger->GetScrubTime();
 			if (LastScrubTime != CurrentScrubTime)
 			{
-				UE::StateTree::Delegates::OnTracingTimelineScrubbed.Broadcast(CurrentScrubTime);
+				StateTree::Delegates::OnTracingTimelineScrubbed.Broadcast(CurrentScrubTime);
 				LastScrubTime = CurrentScrubTime;
 			}
 		};
@@ -43,26 +81,20 @@ void FRewindDebuggerExtension::Update(float DeltaTime, IRewindDebugger* RewindDe
 	}
 }
 
-void FRewindDebuggerRuntimeExtension::RecordingStarted()
-{ 
-	if (!UE_TRACE_CHANNELEXPR_IS_ENABLED(StateTreeDebugChannel))
-	{
-		UE::Trace::ToggleChannel(TEXT("StateTreeDebugChannel"), true);
-		UE::StateTree::Delegates::OnTracingStateChanged.Broadcast(EStateTreeTraceStatus::TracesStarted);
-	}
+void FRewindDebuggerPlaybackExtension::RecordingStarted(IRewindDebugger* RewindDebugger)
+{
+	StateTree::Delegates::OnTraceAnalysisStateChanged.Broadcast(EStateTreeTraceAnalysisStatus::Started);
 }
 
-void FRewindDebuggerRuntimeExtension::RecordingStopped()
+void FRewindDebuggerPlaybackExtension::RecordingStopped(IRewindDebugger* RewindDebugger)
 {
-	// Shouldn't fire as channel will already be disabled when rewind debugger stops, but just as safeguard
-	if (UE_TRACE_CHANNELEXPR_IS_ENABLED(StateTreeDebugChannel))
-	{
-		UE::StateTree::Delegates::OnTracingStateChanged.Broadcast(EStateTreeTraceStatus::StoppingTrace);
-		UE::Trace::ToggleChannel(TEXT("StateTreeDebugChannel"), false);
-	}
+	StateTree::Delegates::OnTraceAnalysisStateChanged.Broadcast(EStateTreeTraceAnalysisStatus::Stopped);
 }
 
+void FRewindDebuggerPlaybackExtension::Clear(IRewindDebugger* RewindDebugger)
+{
+	StateTree::Delegates::OnTraceAnalysisStateChanged.Broadcast(EStateTreeTraceAnalysisStatus::Cleared);
+}
+#endif // WITH_STATETREE_TRACE_DEBUGGER
 
-} // namespace UE::StateTreeDebugger
-
-#endif // WITH_STATETREE_TRACE_DEBUGGER
\ No newline at end of file
+} // namespace UE::StateTreeDebugger
\ No newline at end of file
diff --git a/Engine/Plugins/Runtime/StateTree/Source/StateTreeEditorModule/Private/Debugger/StateTreeRewindDebuggerExtensions.h b/Engine/Plugins/Runtime/StateTree/Source/StateTreeEditorModule/Private/Debugger/StateTreeRewindDebuggerExtensions.h
--- a/Engine/Plugins/Runtime/StateTree/Source/StateTreeEditorModule/Private/Debugger/StateTreeRewindDebuggerExtensions.h
+++ b/Engine/Plugins/Runtime/StateTree/Source/StateTreeEditorModule/Private/Debugger/StateTreeRewindDebuggerExtensions.h
@@ -1,35 +1,39 @@
 // Copyright Epic Games, Inc. All Rights Reserved.
 #pragma once
 
-#if WITH_STATETREE_TRACE_DEBUGGER
+#if WITH_STATETREE_TRACE
+#include "RewindDebuggerRuntimeInterface/IRewindDebuggerRuntimeExtension.h"
+#endif
 
-#include "IRewindDebugger.h"
+#if WITH_STATETREE_TRACE_DEBUGGER
 #include "IRewindDebuggerExtension.h"
-#include "RewindDebuggerRuntimeInterface/IRewindDebuggerRuntimeExtension.h"
+
+class IRewindDebugger;
+#endif
 
 namespace UE::StateTreeDebugger
 {
 
-class FRewindDebuggerExtension : public IRewindDebuggerExtension
+#if WITH_STATETREE_TRACE
+class FRewindDebuggerRecordingExtension final : public IRewindDebuggerRuntimeExtension
 {
-public:
+	virtual void RecordingStarted() override;
+	virtual void RecordingStopped() override;
+	virtual void Clear() override;
+};
+#endif // WITH_STATETREE_TRACE
 
+#if WITH_STATETREE_TRACE_DEBUGGER
+class FRewindDebuggerPlaybackExtension final : public IRewindDebuggerExtension
+{
 	virtual void Update(float DeltaTime, IRewindDebugger* RewindDebugger) override;
-
-private:
+	virtual void RecordingStarted(IRewindDebugger* RewindDebugger) override;
+	virtual void RecordingStopped(IRewindDebugger* RewindDebugger) override;
+	virtual void Clear(IRewindDebugger* RewindDebugger) override;
 
 	/** Last scrub time we received. Used to avoid redundant updates. */
 	double LastScrubTime = 0.0;
 };
+#endif // WITH_STATETREE_TRACE_DEBUGGER
 
-class FRewindDebuggerRuntimeExtension : public IRewindDebuggerRuntimeExtension
-{
-public:
-
-	virtual void RecordingStarted() override;
-	virtual void RecordingStopped() override;
-};
-
-} // namespace UE::StateTreeDebugger
-
-#endif // WITH_STATETREE_TRACE_DEBUGGER
\ No newline at end of file
+} // UE::StateTreeDebugger
\ No newline at end of file
diff --git a/Engine/Plugins/Runtime/StateTree/Source/StateTreeEditorModule/Private/Debugger/StateTreeRewindDebuggerTrack.cpp b/Engine/Plugins/Runtime/StateTree/Source/StateTreeEditorModule/Private/Debugger/StateTreeRewindDebuggerTrack.cpp
--- a/Engine/Plugins/Runtime/StateTree/Source/StateTreeEditorModule/Private/Debugger/StateTreeRewindDebuggerTrack.cpp
+++ b/Engine/Plugins/Runtime/StateTree/Source/StateTreeEditorModule/Private/Debugger/StateTreeRewindDebuggerTrack.cpp
@@ -23,7 +23,7 @@ static const FLazyName StateTreeInstancesName("StateTreeInstances");
 //----------------------------------------------------------------------//
 void FRewindDebuggerTrackCreator::GetTrackTypesInternal(TArray<RewindDebugger::FRewindDebuggerTrackType>& Types) const
 {
-	Types.Add({ StateTreeInstancesName, LOCTEXT("StateTree", "StateTreeInstance") });
+	Types.Add({StateTreeInstancesName, LOCTEXT("StateTree", "StateTreeInstance")});
 }
 
 TSharedPtr<RewindDebugger::FRewindDebuggerTrack> FRewindDebuggerTrackCreator::CreateTrackInternal(const RewindDebugger::FObjectId& InObjectId) const
@@ -32,29 +32,6 @@ TSharedPtr<RewindDebugger::FRewindDebuggerTrack> FRewindDebuggerTrackCreator::Cr
 }
 
 bool FRewindDebuggerTrackCreator::HasDebugInfoInternal(const RewindDebugger::FObjectId& InObjectId) const
-{
-	TRACE_CPUPROFILER_EVENT_SCOPE(FStateTreeRewindDebuggerTrack::HasDebugInfoInternal);
-
-	if (!InObjectId.IsChildElement())
-	{
-		return false;
-	}
-
-	if (const IRewindDebugger* RewindDebugger = IRewindDebugger::Instance())
-	{
-		if (const TraceServices::IAnalysisSession* Session = RewindDebugger->GetAnalysisSession())
-		{
-			if (const IStateTreeTraceProvider* StateTreeTraceProvider = Session->ReadProvider<IStateTreeTraceProvider>(FStateTreeTraceProvider::ProviderName))
-			{
-				return StateTreeTraceProvider->GetInstanceDescriptor(FStateTreeInstanceDebugId(InObjectId.GetChildElementId())).IsValid();
-			}
-		}
-	}
-
-	return false;
-}
-
-bool FRewindDebuggerTrackCreator::IsChildElementSupportedInternal() const
 {
 	return true;
 }
@@ -66,10 +43,7 @@ bool FRewindDebuggerTrackCreator::IsCreatingPrimaryChildTrackInternal() const
 
 FName FRewindDebuggerTrackCreator::GetTargetTypeNameInternal() const
 {
-	// Statetree instance can be executed by multiple systems so 
-	// we want to process all objects tracks but we early out in 
-	// FRewindDebuggerTrackCreator::HasDebugInfoInternal if not a child element
-	static const FName AssociatedObjectName("Object");
+	static const FName AssociatedObjectName("StateTreeInstanceData");
 	return AssociatedObjectName;
 }
 
@@ -77,7 +51,7 @@ FName FRewindDebuggerTrackCreator::GetTargetTypeNameInternal() const
 // FRewindDebuggerTrack
 //----------------------------------------------------------------------//
 FRewindDebuggerTrack::FRewindDebuggerTrack(const RewindDebugger::FObjectId& InObjectId)
-	: EventCollection(FStateTreeInstanceDebugId(InObjectId.GetChildElementId()))
+	: EventCollection(FStateTreeInstanceDebugId(InObjectId.GetMainId()))
 	, InstanceTrackHelper(new FInstanceTrackHelper(EventCollection.InstanceId, IRewindDebugger::Instance()->GetCurrentViewRange()))
 	, Icon(FSlateIconFinder::FindIconForClass(UStateTree::StaticClass()))
 	, ObjectId(InObjectId)
@@ -189,14 +163,16 @@ bool FRewindDebuggerTrack::UpdateInternal()
 				InstanceTrackHelper.Get().SetActiveStates(ActiveStates);
 			}
 
+			// Events need to be updated for a new scrub time
+			FScrubState ScrubState(&EventCollection);
+			ScrubState.SetScrubTime(ScrubTime);
+			EventsView->RequestRefresh(ScrubState);
+
+			// Tree view only need to be updated on new timeline data or new active states
 			if (bTimelineDataChanged
 				|| bActivateStatesChanged)
 			{
 				CompactTreeView->Refresh();
-
-				FScrubState ScrubState(&EventCollection);
-				ScrubState.SetScrubTime(ScrubTime);
-				EventsView->RequestRefresh(ScrubState);
 				return true;
 			}
 		}
diff --git a/Engine/Plugins/Runtime/StateTree/Source/StateTreeEditorModule/Private/Debugger/StateTreeRewindDebuggerTrack.h b/Engine/Plugins/Runtime/StateTree/Source/StateTreeEditorModule/Private/Debugger/StateTreeRewindDebuggerTrack.h
--- a/Engine/Plugins/Runtime/StateTree/Source/StateTreeEditorModule/Private/Debugger/StateTreeRewindDebuggerTrack.h
+++ b/Engine/Plugins/Runtime/StateTree/Source/StateTreeEditorModule/Private/Debugger/StateTreeRewindDebuggerTrack.h
@@ -30,7 +30,6 @@ protected:
 	virtual void GetTrackTypesInternal(TArray<RewindDebugger::FRewindDebuggerTrackType>& Types) const override;
 	virtual TSharedPtr<RewindDebugger::FRewindDebuggerTrack> CreateTrackInternal(const RewindDebugger::FObjectId& InObjectId) const override;
 	virtual bool HasDebugInfoInternal(const RewindDebugger::FObjectId& InObjectId) const override;
-	virtual bool IsChildElementSupportedInternal() const override;
 	virtual bool IsCreatingPrimaryChildTrackInternal() const override;
 };
 
diff --git a/Engine/Plugins/Runtime/StateTree/Source/StateTreeEditorModule/Private/StateTreeEditorModule.cpp b/Engine/Plugins/Runtime/StateTree/Source/StateTreeEditorModule/Private/StateTreeEditorModule.cpp
--- a/Engine/Plugins/Runtime/StateTree/Source/StateTreeEditorModule/Private/StateTreeEditorModule.cpp
+++ b/Engine/Plugins/Runtime/StateTree/Source/StateTreeEditorModule/Private/StateTreeEditorModule.cpp
@@ -93,11 +93,11 @@ void FStateTreeEditorModule::StartupModule()
 #if WITH_STATETREE_TRACE_DEBUGGER
 	FStateTreeDebuggerCommands::Register();
 
-	RewindDebuggerExtension = MakePimpl<UE::StateTreeDebugger::FRewindDebuggerExtension>();
-	IModularFeatures::Get().RegisterModularFeature(IRewindDebuggerExtension::ModularFeatureName, RewindDebuggerExtension.Get());
+	RewindDebuggerPlaybackExtension = MakePimpl<UE::StateTreeDebugger::FRewindDebuggerPlaybackExtension>();
+	IModularFeatures::Get().RegisterModularFeature(IRewindDebuggerExtension::ModularFeatureName, RewindDebuggerPlaybackExtension.Get());
 
-	RewindDebuggerRuntimeExtension = MakePimpl<UE::StateTreeDebugger::FRewindDebuggerRuntimeExtension>();
-	IModularFeatures::Get().RegisterModularFeature(IRewindDebuggerRuntimeExtension::ModularFeatureName, RewindDebuggerRuntimeExtension.Get());
+	RewindDebuggerRecordingExtension = MakePimpl<UE::StateTreeDebugger::FRewindDebuggerRecordingExtension>();
+	IModularFeatures::Get().RegisterModularFeature(IRewindDebuggerRuntimeExtension::ModularFeatureName, RewindDebuggerRecordingExtension.Get());
 
 	RewindDebuggerTrackCreator = MakePimpl<UE::StateTreeDebugger::FRewindDebuggerTrackCreator>();
 	IModularFeatures::Get().RegisterModularFeature(RewindDebugger::IRewindDebuggerTrackCreator::ModularFeatureName, RewindDebuggerTrackCreator.Get());
@@ -159,8 +159,8 @@ void FStateTreeEditorModule::ShutdownModule()
 
 #if WITH_STATETREE_TRACE_DEBUGGER
 	IModularFeatures::Get().UnregisterModularFeature(RewindDebugger::IRewindDebuggerTrackCreator::ModularFeatureName, RewindDebuggerTrackCreator.Get());
-	IModularFeatures::Get().UnregisterModularFeature(IRewindDebuggerRuntimeExtension::ModularFeatureName, RewindDebuggerRuntimeExtension.Get());
-	IModularFeatures::Get().UnregisterModularFeature(IRewindDebuggerExtension::ModularFeatureName, RewindDebuggerExtension.Get());
+	IModularFeatures::Get().UnregisterModularFeature(IRewindDebuggerRuntimeExtension::ModularFeatureName, RewindDebuggerRecordingExtension.Get());
+	IModularFeatures::Get().UnregisterModularFeature(IRewindDebuggerExtension::ModularFeatureName, RewindDebuggerPlaybackExtension.Get());
 	FStateTreeDebuggerCommands::Unregister();
 #endif // WITH_STATETREE_TRACE_DEBUGGER
 
diff --git a/Engine/Plugins/Runtime/StateTree/Source/StateTreeEditorModule/Public/StateTreeEditorModule.h b/Engine/Plugins/Runtime/StateTree/Source/StateTreeEditorModule/Public/StateTreeEditorModule.h
--- a/Engine/Plugins/Runtime/StateTree/Source/StateTreeEditorModule/Public/StateTreeEditorModule.h
+++ b/Engine/Plugins/Runtime/StateTree/Source/StateTreeEditorModule/Public/StateTreeEditorModule.h
@@ -19,8 +19,8 @@ struct FStateTreeNodeClassCache;
 
 namespace UE::StateTreeDebugger
 {
-	class FRewindDebuggerExtension;
-	class FRewindDebuggerRuntimeExtension;
+	class FRewindDebuggerPlaybackExtension;
+	class FRewindDebuggerRecordingExtension;
 	struct FRewindDebuggerTrackCreator;
 }
 
@@ -77,8 +77,8 @@ protected:
 	TArray<FEditorDataType> EditorDataTypes;
 
 #if WITH_STATETREE_TRACE_DEBUGGER
-	TPimplPtr<UE::StateTreeDebugger::FRewindDebuggerExtension> RewindDebuggerExtension;
-	TPimplPtr<UE::StateTreeDebugger::FRewindDebuggerRuntimeExtension> RewindDebuggerRuntimeExtension;
+	TPimplPtr<UE::StateTreeDebugger::FRewindDebuggerPlaybackExtension> RewindDebuggerPlaybackExtension;
+	TPimplPtr<UE::StateTreeDebugger::FRewindDebuggerRecordingExtension> RewindDebuggerRecordingExtension;
 	TPimplPtr<UE::StateTreeDebugger::FRewindDebuggerTrackCreator> RewindDebuggerTrackCreator;
 #endif  // WITH_STATETREE_TRACE_DEBUGGER
 
diff --git a/Engine/Plugins/Runtime/StateTree/Source/StateTreeModule/Private/Debugger/StateTreeTrace.cpp b/Engine/Plugins/Runtime/StateTree/Source/StateTreeModule/Private/Debugger/StateTreeTrace.cpp
--- a/Engine/Plugins/Runtime/StateTree/Source/StateTreeModule/Private/Debugger/StateTreeTrace.cpp
+++ b/Engine/Plugins/Runtime/StateTree/Source/StateTreeModule/Private/Debugger/StateTreeTrace.cpp
@@ -12,6 +12,7 @@
 #include "Serialization/BufferArchive.h"
 #include "StateTree.h"
 #include "StateTreeDelegates.h"
+#include "StateTreeExecutionContext.h"
 #include "StateTreeExecutionTypes.h"
 #include "UObject/Package.h"
 #include "Trace/Trace.inl"
@@ -138,7 +139,7 @@ struct FPhaseTraceStatusPair
 	}
 
 	EStateTreeUpdatePhase Phase = EStateTreeUpdatePhase::Unset;
-	FStateTreeStateHandle StateHandle = FStateTreeStateHandle::Invalid; 
+	FStateTreeStateHandle StateHandle = FStateTreeStateHandle::Invalid;
 	bool bTraced = false;
 };
 
@@ -152,17 +153,6 @@ struct FAssetDebugIdEventBufferedData
 	{
 	}
 
-	void Trace() const
-	{
-		if (ensureMsgf(UE_TRACE_CHANNELEXPR_IS_ENABLED(StateTreeDebugChannel), TEXT("Tracing a buffered data is expected only if channel is enabled.")))
-		{
-			if (const UStateTree* StateTree = WeakStateTree.Get())
-			{
-				OutputAssetDebugIdEvent(StateTree, AssetDebugId);
-			}
-		}
-	}
-
 	TWeakObjectPtr<const UStateTree> WeakStateTree;
 	FStateTreeIndex16 AssetDebugId;
 };
@@ -192,15 +182,8 @@ struct FInstanceEventBufferedData
 	};
 
 	FInstanceEventBufferedData() = default;
-	explicit FInstanceEventBufferedData(
-		const UObject* Owner,
-		const UStateTree* StateTree,
-		const FStateTreeInstanceDebugId InstanceId,
-		const FString& InstanceName)
-		: WeakOwner(Owner)
-		, InstanceName(InstanceName)
-		, WeakStateTree(StateTree)
-		, InstanceId(InstanceId)
+	explicit FInstanceEventBufferedData(const FStateTreeInstanceDebugId InstanceId)
+		: InstanceId(InstanceId)
 	{
 	}
 
@@ -212,26 +195,6 @@ struct FInstanceEventBufferedData
 		EmptyState.Output(InstanceId);
 	}
 
-	void Trace() const
-	{
-		if (ensureMsgf(UE_TRACE_CHANNELEXPR_IS_ENABLED(StateTreeDebugChannel), TEXT("Tracing a buffered data is expected only if channel is enabled.")))
-		{
-			if (const UStateTree* StateTree = WeakStateTree.Get())
-			{
-				OutputInstanceLifetimeEvent(WeakOwner.Get(), InstanceId, StateTree, *InstanceName, EStateTreeTraceEventType::Push);
-
-				if (ActiveStates.IsValid())
-				{
-					ActiveStates.Output(InstanceId);
-				}
-			}
-		}
-	}
-
-	TWeakObjectPtr<const UObject> WeakOwner;
-	FActiveStates ActiveStates;
-	FString InstanceName;
-	TWeakObjectPtr<const UStateTree> WeakStateTree;
 	FStateTreeInstanceDebugId InstanceId;
 
 	/** Stack to keep track of all received phase events so matching "closing" events (i.e., Pop) will control "IF" and "WHEN" a given phase will be sent. */
@@ -262,15 +225,40 @@ struct FBufferedDataList
 	/** Converts the provided list of execution frames to traceable data format and outputs it to the trace. */
 	void OutputActiveStates_AnyThread(const FStateTreeInstanceDebugId InstanceId, const TConstArrayView<FStateTreeExecutionFrame> ActiveFrames);
 
-	/** Converts the provided list of execution frames to traceable data format and updates the buffered data for the provided instance. */
-	void UpdateActiveStates_AnyThread(const FStateTreeInstanceDebugId InstanceId, const TConstArrayView<FStateTreeExecutionFrame> ActiveFrames);
-
-	/** Keeps track of created/destroyed statetree instances */
-	void UpdateInstanceLifetime_AnyThread(const UObject* Owner, const FStateTreeInstanceDebugId InstanceId, const UStateTree* StateTree, const TCHAR* InstanceName, const EStateTreeTraceEventType EventType);
-
 	/** Sets the world time that might need to be traced by worker threads. */
 	void SetWorldTime_GameThread(const double WorldTime);
 
+	void Reset()
+	{
+		{
+			FRWScopeLock RWScopeLock(TracedInstancesRWLock, SLT_Write);
+			TracedInstances.Empty();
+		}
+
+		RecordingWorldTime = UNINITIALIZED_WORLD_TIME;
+		TracedRecordingWorldTime = UNINITIALIZED_WORLD_TIME;
+
+		AssetDebugIdMap.Reset();
+		BufferedEventsMap.Reset();
+		NextAssetDebugId = 1;
+	}
+
+	bool ShouldTraceInstance(const FStateTreeInstanceDebugId InstanceId) const
+	{
+		// Note this pattern is the same as AnimNextTrace; both should be updated if a change/fix is needed
+		FRWScopeLock RWScopeLock(TracedInstancesRWLock, SLT_ReadOnly);
+		if (!TracedInstances.Contains(InstanceId))
+		{
+			RWScopeLock.ReleaseReadOnlyLockAndAcquireWriteLock_USE_WITH_CAUTION();
+
+			// Note that we might double add here due to 2x people writing at the same time, but it doesn't matter
+			TracedInstances.Add(InstanceId);
+			return true;
+		}
+
+		return false;
+	}
+
 private:
 	static constexpr double UNINITIALIZED_WORLD_TIME = -1;
 
@@ -280,6 +268,10 @@ private:
 	/** World time traced only once by the first worker thread that needs to trace an event */
 	std::atomic<double> TracedRecordingWorldTime = UNINITIALIZED_WORLD_TIME;
 
+	/** List of traced instances */
+	mutable TSet<FStateTreeInstanceDebugId> TracedInstances;
+	mutable FRWLock TracedInstancesRWLock;
+
 	/** List of asset debug ids events that will be output if channel gets enabled. */
 	TArray<FAssetDebugIdEventBufferedData> AssetDebugIdMap;
 	FTransactionallySafeRWLock AssetDebugIdMapLock;
@@ -294,13 +286,7 @@ private:
 	/** MT access detector used to validate that AssetDebugIdEvents and InstanceLifetimeEvents are never flushed while getting accessed by worker threads. */
 	UE_MT_DECLARE_TS_RW_ACCESS_DETECTOR(BufferedEventsMapMTDetector);
 
-	/** Flag use to prevent reentrant calls when buffered events gets flushed when starting a new trace (i.e., OnTracesStarted) */
-	bool bFlushingLifetimeEvents = false;
-
 	uint16 NextAssetDebugId = 1;
-	int32 CurrentVersion = 0;
-	int32 FlushedVersion = -1;
-
 
 	/** Called by TraceBufferedEvents from the OutputXYZ methods to make sure current world time was traced. */
 	void TraceWorldTimeIfNeeded_AnyThread();
@@ -313,6 +299,13 @@ private:
 	void TraceStackedPhases_AnyThread(const FStateTreeInstanceDebugId InstanceId);
 };
 
+/**
+ * Buffered events (e.g. lifetime, active state, scoped phase) in case channel is not active yet or phase are empty and don't need to be traced.
+ * This can be accessed from any thread and based on the assumption that a given instance data can never be accessed by multiple thread simultaneously in the same frame
+ * so we only need to protect access to the main maps.
+ */
+FBufferedDataList GBufferedEvents;
+
 FInstanceEventBufferedData::FActiveStates::FActiveStates(FBufferedDataList& BufferedData, const TConstArrayView<FStateTreeExecutionFrame> ActiveFrames)
 {
 	for (const FStateTreeExecutionFrame& Frame : ActiveFrames)
@@ -343,12 +336,12 @@ void FBufferedDataList::Flush_AnyThread(const FStateTreeInstanceDebugId Instance
 void FBufferedDataList::TraceWorldTimeIfNeeded_AnyThread()
 {
 	checkf(RecordingWorldTime != UNINITIALIZED_WORLD_TIME
-	       , TEXT("Expecting world time to always be set at the beginning of the frame before any worker thread attempts to trace."));
+		, TEXT("Expecting world time to always be set at the beginning of the frame before any worker thread attempts to trace."));
 	double Expected = UNINITIALIZED_WORLD_TIME;
 	if (TracedRecordingWorldTime.compare_exchange_strong(Expected, RecordingWorldTime))
 	{
 		UE_TRACE_LOG(StateTreeDebugger, WorldTimestampEvent, StateTreeDebugChannel)
-				<< WorldTimestampEvent.WorldTime(RecordingWorldTime);
+			<< WorldTimestampEvent.WorldTime(RecordingWorldTime);
 	}
 }
 
@@ -371,7 +364,7 @@ void FBufferedDataList::TraceStackedPhases_AnyThread(const FStateTreeInstanceDeb
 		UE_MT_SCOPED_READ_ACCESS(BufferedEventsMapMTDetector);
 		if (FInstanceEventBufferedData* Data = BufferedEventsMap.Find(InstanceId))
 		{
-			// Keep local copy to trace outside the critical section and only perform the 
+			// Keep local copy to trace outside the critical section and only perform the
 			// minimal required changes while we hold the lock (i.e., updating that status)
 			PhaseStackToTrace = Data->PhaseStack;
 			for (FPhaseTraceStatusPair& StackEntry : Data->PhaseStack)
@@ -401,23 +394,6 @@ void FBufferedDataList::OnTracesStarted_GameThread()
 {
 	checkf(IsInGameThread(), TEXT("Expecting to only be called by the statetree delegate on the main thread before worker threads trace events."));
 
-	// Trace asset events first since they are required for instance lifetime event types.
-	// Events are preserved in case the trace session is stopped and then a new one gets started
-	// in the same game session. In which case we need to output AssetDebugIdEvents to that new trace.
-	const bool bFlushingLifetimeEventsRequired = FlushedVersion != CurrentVersion;
-	TGuardValue<bool> GuardReentry(bFlushingLifetimeEvents, bFlushingLifetimeEventsRequired);
-
-	if (bFlushingLifetimeEventsRequired)
-	{
-		FlushedVersion = CurrentVersion;
-
-		UE_MT_SCOPED_READ_ACCESS(BufferedEventsMapMTDetector);
-		for (const FAssetDebugIdEventBufferedData& AssetDebugIdEventData : AssetDebugIdMap)
-		{
-			AssetDebugIdEventData.Trace();
-		}
-	}
-
 	// Delegate can be received before first call to SetWorldTime_GameThread()
 	// In that case we set the time to 0 for any pending lifetime events
 	if (RecordingWorldTime == UNINITIALIZED_WORLD_TIME)
@@ -426,19 +402,6 @@ void FBufferedDataList::OnTracesStarted_GameThread()
 	}
 
 	TraceWorldTimeIfNeeded_AnyThread();
-
-	// Then trace instance lifetime events since they are required for other event types.
-	// They are associated to an older world time but to simplify the logic on the analysis side
-	// we will send them as if the instances were created at the beginning of the recording.
-	// Those events are also preserved for the same reason as AssetDebugIdEvents.
-	if (bFlushingLifetimeEventsRequired)
-	{
-		UE_MT_SCOPED_READ_ACCESS(BufferedEventsMapMTDetector);
-		for (TPair<FStateTreeInstanceDebugId, FInstanceEventBufferedData>& Pair : BufferedEventsMap)
-		{
-			Pair.Value.Trace();
-		}
-	}
 }
 
 void FBufferedDataList::OnStoppingTraces_GameThread()
@@ -451,9 +414,6 @@ void FBufferedDataList::OnStoppingTraces_GameThread()
 		Pair.Value.CloseRecording(TracedRecordingWorldTime);
 	}
 
-	// Bump version so shareable data will be flush in the next trace (e.g. Asset ids, instance lifetime events, etc.)
-	++CurrentVersion;
-
 	// Reset for next trace
 	RecordingWorldTime = UNINITIALIZED_WORLD_TIME;
 }
@@ -538,44 +498,6 @@ void FBufferedDataList::OutputActiveStates_AnyThread(const FStateTreeInstanceDeb
 	ActiveStates.Output(InstanceId);
 }
 
-void FBufferedDataList::UpdateActiveStates_AnyThread(const FStateTreeInstanceDebugId InstanceId, const TConstArrayView<FStateTreeExecutionFrame> ActiveFrames)
-{
-	TReadScopeLock ReadSection(BufferedEventsMapLock);
-	UE_MT_SCOPED_READ_ACCESS(BufferedEventsMapMTDetector);
-
-	// We keep only the most recent active states since all we need to know is the last active states of the instance
-	// when we start receiving events once the channel gets enabled.
-	if (FInstanceEventBufferedData* ExistingBufferedData = BufferedEventsMap.Find(InstanceId))
-	{
-		ExistingBufferedData->ActiveStates = FInstanceEventBufferedData::FActiveStates(*this, ActiveFrames);
-	}
-}
-
-void FBufferedDataList::UpdateInstanceLifetime_AnyThread(const UObject* Owner, const FStateTreeInstanceDebugId InstanceId, const UStateTree* StateTree, const TCHAR* InstanceName
-	, const EStateTreeTraceEventType EventType)
-{
-	// This can be called by worker threads from statetree trace macros and by the game thread during the flush.
-	// For the latter we don't want to queue again events that are getting flushed.
-	if (!bFlushingLifetimeEvents)
-	{
-		TWriteScopeLock WriteSection(BufferedEventsMapLock);
-		UE_MT_SCOPED_WRITE_ACCESS(BufferedEventsMapMTDetector);
-
-		if (EventType == EStateTreeTraceEventType::Push)
-		{
-			BufferedEventsMap.Emplace(InstanceId, FInstanceEventBufferedData(Owner, StateTree, InstanceId, InstanceName));
-		}
-		else if (EventType == EStateTreeTraceEventType::Pop)
-		{
-			BufferedEventsMap.Remove(InstanceId);
-		}
-		else
-		{
-			ensureMsgf(false, TEXT("Unexpected EventType '%s' for instance lifetime event."), *UEnum::GetDisplayValueAsText(EventType).ToString());
-		}
-	}
-}
-
 void FBufferedDataList::SetWorldTime_GameThread(const double WorldTime)
 {
 	checkf(IsInGameThread(), TEXT("Expecting to only be called by the world delegate on the main thread before worker threads trace events."));
@@ -585,24 +507,18 @@ void FBufferedDataList::SetWorldTime_GameThread(const double WorldTime)
 	TracedRecordingWorldTime = UNINITIALIZED_WORLD_TIME;
 }
 
-/**
- * Buffered events (e.g. lifetime, active state, scoped phase) in case channel is not active yet or phase are empty and don't need to be traced.
- * This can be accessed from any thread and based on the assumption that a given instance data can never be accessed by multiple thread simultaneously in the same frame
- * so we only need to protect access to the main maps.
- */
-FBufferedDataList GBufferedEvents;
-
 /**
  * Pushed or pops an entry on the Phase stack for a given Instance.
  * Will send the Pop events for phases popped if their associated Push events were sent.
  */
 void OutputPhaseScopeEvent(
-	const UObject* Owner,
-	const FStateTreeInstanceDebugId InstanceId,
+	TNotNull<const FStateTreeReadOnlyExecutionContext*> Context,
 	const EStateTreeUpdatePhase Phase,
 	const EStateTreeTraceEventType EventType,
 	const FStateTreeStateHandle StateHandle)
 {
+	const FStateTreeInstanceDebugId InstanceId = Context->GetInstanceDebugId();
+
 	const TOptional<FPhaseTraceStatusPair> RemovedPair = GBufferedEvents.UpdatePhaseScope_AnyThread(InstanceId, Phase, EventType, StateHandle);
 
 	// Phase was previously traced (i.e., other events were traced in that scope so we need to trace the closing (i.e., Pop) event)
@@ -623,15 +539,51 @@ void OutputPhaseScopeEvent(
 /**
  * Called by the OutputXYZ methods to flush pending events (e.g. Push or WorldTime).
  */
-void TraceBufferedEvents(const FStateTreeInstanceDebugId InstanceId)
+void TraceBufferedEvents(TNotNull<const FStateTreeReadOnlyExecutionContext*> Context)
 {
+	const UObject* Owner = Context->GetOwner();
+	if (CANNOT_TRACE_OBJECT(Owner))
+	{
+		return;
+	}
+
+	const FStateTreeInstanceDebugId InstanceId = Context->GetInstanceDebugId();
+	if (GBufferedEvents.ShouldTraceInstance(InstanceId))
+	{
+		const UStateTree* StateTree = Context->GetStateTree();
+		const FString InstanceName = Context->GetInstanceDebugDescription();
+		const uint64 OuterTraceId = Context->GetOuterTraceId();
+		const uint64 OuterId = OuterTraceId == 0 ? FObjectTrace::GetObjectId(Owner) : OuterTraceId;
+
+		// Trace the owner (Object trace specific)
+		TRACE_OBJECT(Owner);
+
+		// Trace the struct instance (Object trace specific)
+		TRACE_INSTANCE(Owner, InstanceId.ToUint64(), OuterId, FStateTreeInstanceData::StaticStruct(), StateTree->GetName());
+
+		// Will trace Asset debug Id if needed (StateTree trace specific)
+		const FStateTreeIndex16 AssetDebugId = GBufferedEvents.FindOrAddDebugIdForAsset_AnyThread(StateTree);
+
+		// Trace instance (StateTree trace specific)
+		UE_TRACE_LOG(StateTreeDebugger, InstanceEvent, StateTreeDebugChannel)
+			<< InstanceEvent.Cycle(FPlatformTime::Cycles64())
+			<< InstanceEvent.InstanceId(InstanceId.Id)
+			<< InstanceEvent.InstanceSerial(InstanceId.SerialNumber)
+			<< InstanceEvent.InstanceName(*InstanceName, InstanceName.Len())
+			<< InstanceEvent.EventType(static_cast<std::underlying_type_t<EStateTreeTraceEventType>>(EStateTreeTraceEventType::Push))
+			<< InstanceEvent.AssetDebugId(AssetDebugId.Get());
+
+		// Trace active states (StateTree trace specific)
+		OutputActiveStatesEventTrace(Context, Context->GetActiveFrames());
+	}
+
 	GBufferedEvents.Flush_AnyThread(InstanceId);
 }
 
 void SerializeDebugDataToArchive(FBufferArchive& Ar, FNodeCustomDebugData&& CustomDebugData, const FStateTreeDataView DataView)
 {
-	constexpr uint32 PortFlags = 
-		PPF_PropertyWindow // limit to properties visible in Editor 
+	constexpr uint32 PortFlags =
+		PPF_PropertyWindow // limit to properties visible in Editor
 		| PPF_ExportsNotFullyQualified
 		| PPF_Delimited // property data should be wrapped in quotes
 		| PPF_ExternalEditor // uses authored names instead of internal names and default values are always written out
@@ -647,7 +599,7 @@ void SerializeDebugDataToArchive(FBufferArchive& Ar, FNodeCustomDebugData&& Cust
 	{
 		TRACE_CPUPROFILER_EVENT_SCOPE(UE::StateTree::ExportStructAsText);
 		TypePath = ScriptStruct->GetPathName();
-		
+
 		if (!CustomDebugData.ShouldOverrideDataView())
 		{
 			ScriptStruct->ExportText(InstanceDataAsText, DataView.GetMemory(), DataView.GetMemory(), /*OwnerObject*/nullptr, PortFlags | PPF_SeparateDefine, /*ExportRootScope*/nullptr);
@@ -695,12 +647,15 @@ void RegisterGlobalDelegates()
 		{
 			GBufferedEvents.OnTracesStarted_GameThread();
 		}
-
 		// Traces are about to be stopped so allow the buffered events to react.
-		if (TraceStatus == EStateTreeTraceStatus::StoppingTrace)
+		else if (TraceStatus == EStateTreeTraceStatus::StoppingTrace)
 		{
 			GBufferedEvents.OnStoppingTraces_GameThread();
 		}
+		else if (TraceStatus == EStateTreeTraceStatus::Cleared)
+		{
+			GBufferedEvents.Reset();
+		}
 	});
 }
 
@@ -713,15 +668,22 @@ void UnregisterGlobalDelegates()
 	GTracingStateChangedDelegateHandle.Reset();
 }
 
+void TraceInstance(TNotNull<const FStateTreeReadOnlyExecutionContext*> Context)
+{
+	if (!UE_TRACE_CHANNELEXPR_IS_ENABLED(StateTreeDebugChannel))
+	{
+		return;
+	}
+
+	TraceBufferedEvents(Context);
+}
+
 void OutputAssetDebugIdEvent(
 	const UStateTree* StateTree,
-	const FStateTreeIndex16 AssetDebugId
-	)
+	const FStateTreeIndex16 AssetDebugId)
 {
 	if (UE_TRACE_CHANNELEXPR_IS_ENABLED(StateTreeDebugChannel))
 	{
-		TraceBufferedEvents(FStateTreeInstanceDebugId::Invalid);
-
 		check(StateTree);
 		const FString TreeName = StateTree->GetName();
 		const FString TreePath = StateTree->GetPathName();
@@ -736,56 +698,35 @@ void OutputAssetDebugIdEvent(
 }
 
 void OutputInstanceLifetimeEvent(
-	const UObject* Owner,
-	const FStateTreeInstanceDebugId InstanceId,
-	const UStateTree* StateTree,
-	const TCHAR* InstanceName,
-	const EStateTreeTraceEventType EventType
-	)
+	TNotNull<const FStateTreeReadOnlyExecutionContext*> Context,
+	const EStateTreeTraceEventType EventType)
 {
+	const UObject* Owner = Context->GetOwner();
+
 	if (UE_TRACE_CHANNELEXPR_IS_ENABLED(StateTreeDebugChannel))
 	{
-		const uint64 ChildId(InstanceId.ToUint64());
 		if (EventType == EStateTreeTraceEventType::Push)
 		{
-			TRACE_CHILD_ELEMENT_LIFETIME_BEGIN(Owner, ChildId);
+			TraceInstance(Context);
 		}
 		else if (EventType == EStateTreeTraceEventType::Pop)
 		{
-			TRACE_CHILD_ELEMENT_LIFETIME_END(Owner, ChildId);
+			TRACE_INSTANCE_LIFETIME_END(Owner, Context->GetInstanceDebugId().ToUint64());
 		}
-
-		TraceBufferedEvents(InstanceId);
-
-		const FStateTreeIndex16 AssetDebugId = GBufferedEvents.FindOrAddDebugIdForAsset_AnyThread(StateTree);
-
-		UE_TRACE_LOG(StateTreeDebugger, InstanceEvent, StateTreeDebugChannel)
-			<< InstanceEvent.Cycle(FPlatformTime::Cycles64())
-			<< InstanceEvent.InstanceId(InstanceId.Id)
-			<< InstanceEvent.InstanceSerial(InstanceId.SerialNumber)
-			<< InstanceEvent.InstanceName(InstanceName)
-			<< InstanceEvent.EventType(static_cast<std::underlying_type_t<EStateTreeTraceEventType>>(EventType))
-			<< InstanceEvent.AssetDebugId(AssetDebugId.Get());
 	}
-
-	// Update buffered events regardless of the status of the channel since they will be used
-	// when flushing buffered event when a late recording is started or more than one trace are started
-	// during the same game session (i.e. Start Traces -> Stop Traces -> Start Traces).
-	GBufferedEvents.UpdateInstanceLifetime_AnyThread(Owner, InstanceId, StateTree, InstanceName, EventType);
 }
 
 void OutputInstanceAssetEvent(
-	const UObject* Owner,
-	const FStateTreeInstanceDebugId InstanceId,
-	const UStateTree* StateTree
-	)
+	TNotNull<const FStateTreeReadOnlyExecutionContext*> Context,
+	const UStateTree* StateTree)
 {
 	check(StateTree != nullptr);
 
 	if (UE_TRACE_CHANNELEXPR_IS_ENABLED(StateTreeDebugChannel))
 	{
-		TRACE_CHILD_ELEMENT(Owner, InstanceId.ToUint64());
-		TraceBufferedEvents(InstanceId);
+		const FStateTreeInstanceDebugId InstanceId = Context->GetInstanceDebugId();
+
+		TraceInstance(Context);
 
 		const FStateTreeIndex16 AssetDebugId = GBufferedEvents.FindOrAddDebugIdForAsset_AnyThread(StateTree);
 
@@ -795,34 +736,27 @@ void OutputInstanceAssetEvent(
 			<< InstanceFrameEvent.InstanceSerial(InstanceId.SerialNumber)
 			<< InstanceFrameEvent.AssetDebugId(AssetDebugId.Get());
 	}
-
-	// No need to buffer anything here since these events are sent each time a FCurrentlyProcessedFrameScope
-	// is used by a FStateTreeExecutionContext, and we don't expect the trace channel to be enabled/disabled
-	// during a single execution context update.
 }
 
 void OutputInstanceFrameEvent(
-	const UObject* Owner,
-	const FStateTreeInstanceDebugId InstanceId,
-	const FStateTreeExecutionFrame* Frame
-	)
+	TNotNull<const FStateTreeReadOnlyExecutionContext*> Context,
+	const FStateTreeExecutionFrame* Frame)
 {
 	check(Frame != nullptr);
-	OutputInstanceAssetEvent(Owner, InstanceId, Frame->StateTree.Get());
+	OutputInstanceAssetEvent(Context, Frame->StateTree.Get());
 }
 
 void OutputLogEventTrace(
-	const UObject* Owner,
-	const FStateTreeInstanceDebugId InstanceId,
+	TNotNull<const FStateTreeReadOnlyExecutionContext*> Context,
 	const ELogVerbosity::Type Verbosity,
-	const TCHAR* Fmt, ...
-	)
+	const TCHAR* Fmt, ...)
 {
+	const FStateTreeInstanceDebugId InstanceId = Context->GetInstanceDebugId();
+
 	static TCHAR TraceStaticBuffer[8192];
 	GET_TYPED_VARARGS(TCHAR, TraceStaticBuffer, UE_ARRAY_COUNT(TraceStaticBuffer), UE_ARRAY_COUNT(TraceStaticBuffer) - 1, Fmt, Fmt);
 
-	TRACE_CHILD_ELEMENT(Owner, InstanceId.ToUint64());
-	TraceBufferedEvents(InstanceId);
+	TraceInstance(Context);
 
 	UE_TRACE_LOG(StateTreeDebugger, LogEvent, StateTreeDebugChannel)
 		<< LogEvent.Cycle(FPlatformTime::Cycles64())
@@ -833,14 +767,13 @@ void OutputLogEventTrace(
 }
 
 void OutputStateEventTrace(
-	const UObject* Owner,
-	const FStateTreeInstanceDebugId InstanceId,
+	TNotNull<const FStateTreeReadOnlyExecutionContext*> Context,
 	const FStateTreeStateHandle StateHandle,
-	const EStateTreeTraceEventType EventType
-	)
+	const EStateTreeTraceEventType EventType)
 {
-	TRACE_CHILD_ELEMENT(Owner, InstanceId.ToUint64());
-	TraceBufferedEvents(InstanceId);
+	const FStateTreeInstanceDebugId InstanceId = Context->GetInstanceDebugId();
+
+	TraceInstance(Context);
 
 	UE_TRACE_LOG(StateTreeDebugger, StateEvent, StateTreeDebugChannel)
 		<< StateEvent.Cycle(FPlatformTime::Cycles64())
@@ -851,17 +784,18 @@ void OutputStateEventTrace(
 }
 
 void OutputTaskEventTrace(
-	const UObject* Owner,
-	const FStateTreeInstanceDebugId InstanceId,
-	FNodeCustomDebugData&& CustomDebugData,
-	const FStateTreeIndex16 TaskIdx, const FStateTreeDataView DataView,
-	EStateTreeTraceEventType EventType, EStateTreeRunStatus Status)
+	TNotNull<const FStateTreeReadOnlyExecutionContext*> Context,
+	const FStateTreeIndex16 TaskIdx,
+	const FStateTreeDataView DataView,
+	EStateTreeTraceEventType EventType,
+	EStateTreeRunStatus Status)
 {
+	const FStateTreeInstanceDebugId InstanceId = Context->GetInstanceDebugId();
+
 	FBufferArchive Archive;
-	SerializeDebugDataToArchive(Archive, MoveTemp(CustomDebugData), DataView);
+	SerializeDebugDataToArchive(Archive, Context->StealNodeCustomDebugTraceData(), DataView);
 
-	TRACE_CHILD_ELEMENT(Owner, InstanceId.ToUint64());
-	TraceBufferedEvents(InstanceId);
+	TraceInstance(Context);
 
 	UE_TRACE_LOG(StateTreeDebugger, TaskEvent, StateTreeDebugChannel)
 		<< TaskEvent.Cycle(FPlatformTime::Cycles64())
@@ -874,18 +808,17 @@ void OutputTaskEventTrace(
 }
 
 void OutputEvaluatorEventTrace(
-	const UObject* Owner,
-	const FStateTreeInstanceDebugId InstanceId,
-	FNodeCustomDebugData&& CustomDebugData,
+	TNotNull<const FStateTreeReadOnlyExecutionContext*> Context,
 	const FStateTreeIndex16 EvaluatorIdx,
 	const FStateTreeDataView DataView,
 	EStateTreeTraceEventType EventType)
 {
+	const FStateTreeInstanceDebugId InstanceId = Context->GetInstanceDebugId();
+
 	FBufferArchive Archive;
-	SerializeDebugDataToArchive(Archive, MoveTemp(CustomDebugData), DataView);
+	SerializeDebugDataToArchive(Archive, Context->StealNodeCustomDebugTraceData(), DataView);
 
-	TRACE_CHILD_ELEMENT(Owner, InstanceId.ToUint64());
-	TraceBufferedEvents(InstanceId);
+	TraceInstance(Context);
 
 	UE_TRACE_LOG(StateTreeDebugger, EvaluatorEvent, StateTreeDebugChannel)
 		<< EvaluatorEvent.Cycle(FPlatformTime::Cycles64())
@@ -897,23 +830,23 @@ void OutputEvaluatorEventTrace(
 }
 
 void OutputTransitionEventTrace(
-	const UObject* Owner,
-	const FStateTreeInstanceDebugId InstanceId,
+	TNotNull<const FStateTreeReadOnlyExecutionContext*> Context,
 	const FStateTreeTransitionSource Source,
 	const EStateTreeTraceEventType EventType
-	)
+)
 {
+	const FStateTreeInstanceDebugId InstanceId = Context->GetInstanceDebugId();
+
 	// Output an instance Frame/Asset event to put the transition data in the proper context when it gets read.
 	if (const UStateTree* StateTree = Source.Asset.Get())
 	{
-		OutputInstanceAssetEvent(Owner, InstanceId, StateTree);
+		OutputInstanceAssetEvent(Context, StateTree);
 	}
 
 	FBufferArchive Archive;
 	Archive << EventType;
 
-	TRACE_CHILD_ELEMENT(Owner, InstanceId.ToUint64());
-	TraceBufferedEvents(InstanceId);
+	TraceInstance(Context);
 
 	UE_TRACE_LOG(StateTreeDebugger, TransitionEvent, StateTreeDebugChannel)
 	<< TransitionEvent.Cycle(FPlatformTime::Cycles64())
@@ -927,18 +860,17 @@ void OutputTransitionEventTrace(
 }
 
 void OutputConditionEventTrace(
-	const UObject* Owner,
-	const FStateTreeInstanceDebugId InstanceId,
-	FNodeCustomDebugData&& CustomDebugData,
+	TNotNull<const FStateTreeReadOnlyExecutionContext*> Context,
 	const FStateTreeIndex16 ConditionIdx,
 	const FStateTreeDataView DataView,
 	EStateTreeTraceEventType EventType)
 {
+	const FStateTreeInstanceDebugId InstanceId = Context->GetInstanceDebugId();
+
 	FBufferArchive Archive;
-	SerializeDebugDataToArchive(Archive, MoveTemp(CustomDebugData), DataView);
+	SerializeDebugDataToArchive(Archive, Context->StealNodeCustomDebugTraceData(), DataView);
 
-	TRACE_CHILD_ELEMENT(Owner, InstanceId.ToUint64());
-	TraceBufferedEvents(InstanceId);
+	TraceInstance(Context);
 
 	UE_TRACE_LOG(StateTreeDebugger, ConditionEvent, StateTreeDebugChannel)
 	<< ConditionEvent.Cycle(FPlatformTime::Cycles64())
@@ -950,22 +882,18 @@ void OutputConditionEventTrace(
 }
 
 void OutputActiveStatesEventTrace(
-	const UObject* Owner,
-	const FStateTreeInstanceDebugId InstanceId,
+	TNotNull<const FStateTreeReadOnlyExecutionContext*> Context,
 	const TConstArrayView<FStateTreeExecutionFrame> ActiveFrames
-	)
+)
 {
+	const FStateTreeInstanceDebugId InstanceId = Context->GetInstanceDebugId();
+
 	if (UE_TRACE_CHANNELEXPR_IS_ENABLED(StateTreeDebugChannel))
 	{
-		TRACE_CHILD_ELEMENT(Owner, InstanceId.ToUint64());
-		TraceBufferedEvents(InstanceId);
+		TraceInstance(Context);
 
 		GBufferedEvents.OutputActiveStates_AnyThread(InstanceId, ActiveFrames);
 	}
-	else
-	{
-		GBufferedEvents.UpdateActiveStates_AnyThread(InstanceId, ActiveFrames);
-	}
 }
 
 // Deprecated
@@ -977,69 +905,55 @@ FStateTreeIndex16 FindOrAddDebugIdForAsset(const UStateTree* StateTree)
 // Deprecated
 void OutputPhaseScopeEvent(const FStateTreeInstanceDebugId InstanceId, const EStateTreeUpdatePhase Phase, const EStateTreeTraceEventType EventType, const FStateTreeStateHandle StateHandle)
 {
-	OutputPhaseScopeEvent(nullptr, InstanceId, Phase, EventType, StateHandle);
 }
 
 // Deprecated
 void OutputInstanceLifetimeEvent(const FStateTreeInstanceDebugId InstanceId, const UStateTree* StateTree, const TCHAR* InstanceName, const EStateTreeTraceEventType EventType)
 {
-	OutputInstanceLifetimeEvent(nullptr, InstanceId, StateTree, InstanceName, EventType);
 }
 
 // Deprecated
 void OutputInstanceAssetEvent(const FStateTreeInstanceDebugId InstanceId, const UStateTree* StateTree)
 {
-	OutputInstanceAssetEvent(nullptr, InstanceId, StateTree);
 }
 // Deprecated
 void OutputInstanceFrameEvent(const FStateTreeInstanceDebugId InstanceId, const FStateTreeExecutionFrame* Frame)
 {
-	OutputInstanceFrameEvent(nullptr, InstanceId, Frame);
 }
 
 // Deprecated
 void OutputLogEventTrace(const FStateTreeInstanceDebugId InstanceId, const ELogVerbosity::Type Verbosity, const TCHAR* Fmt, ...)
 {
-	va_list Args;
-	va_start(Args, Fmt);
-	OutputLogEventTrace(nullptr, InstanceId, Verbosity, Fmt, Args);
-	va_end(Args);
 }
 
 // Deprecated
 void OutputStateEventTrace(const FStateTreeInstanceDebugId InstanceId, const FStateTreeStateHandle StateHandle, const EStateTreeTraceEventType EventType)
 {
-	OutputStateEventTrace(nullptr, InstanceId, StateHandle, EventType);
 }
 
 // Deprecated
 void OutputTaskEventTrace(const FStateTreeInstanceDebugId InstanceId, FNodeCustomDebugData&& DebugData, const FStateTreeIndex16 TaskIdx, const FStateTreeDataView DataView, const EStateTreeTraceEventType EventType, const EStateTreeRunStatus Status)
 {
-	OutputTaskEventTrace(nullptr, InstanceId, MoveTemp(DebugData), TaskIdx, DataView, EventType, Status);
 }
 
 // Deprecated
 void OutputEvaluatorEventTrace(const FStateTreeInstanceDebugId InstanceId, FNodeCustomDebugData&& DebugData, const FStateTreeIndex16 EvaluatorIdx, const FStateTreeDataView DataView, const EStateTreeTraceEventType EventType)
 {
-	OutputEvaluatorEventTrace(nullptr, InstanceId, MoveTemp(DebugData), EvaluatorIdx, DataView, EventType);
 }
 
 // Deprecated
 void OutputConditionEventTrace(const FStateTreeInstanceDebugId InstanceId, FNodeCustomDebugData&& DebugData, const FStateTreeIndex16 ConditionIdx, const FStateTreeDataView DataView, const EStateTreeTraceEventType EventType)
 {
-	OutputConditionEventTrace(nullptr, InstanceId, MoveTemp(DebugData), ConditionIdx, DataView, EventType);
 }
 
 // Deprecated
 void OutputTransitionEventTrace(const FStateTreeInstanceDebugId InstanceId, const FStateTreeTransitionSource TransitionSource, const EStateTreeTraceEventType EventType)
 {
-	OutputTransitionEventTrace(nullptr, InstanceId, TransitionSource, EventType);
 }
 
 // Deprecated
 void OutputActiveStatesEventTrace(const FStateTreeInstanceDebugId InstanceId, const TConstArrayView<FStateTreeExecutionFrame> ActiveFrames)
 {
-	OutputActiveStatesEventTrace(nullptr, InstanceId, ActiveFrames);
 }
 
 } // UE::StateTreeTrace
diff --git a/Engine/Plugins/Runtime/StateTree/Source/StateTreeModule/Private/StateTreeDelegates.cpp b/Engine/Plugins/Runtime/StateTree/Source/StateTreeModule/Private/StateTreeDelegates.cpp
--- a/Engine/Plugins/Runtime/StateTree/Source/StateTreeModule/Private/StateTreeDelegates.cpp
+++ b/Engine/Plugins/Runtime/StateTree/Source/StateTreeModule/Private/StateTreeDelegates.cpp
@@ -4,7 +4,7 @@
 
 namespace UE::StateTree::Delegates
 {
-	
+
 #if WITH_EDITOR
 FOnIdentifierChanged OnIdentifierChanged;
 FOnSchemaChanged OnSchemaChanged;
@@ -23,6 +23,7 @@ FOnTracingStateChanged OnTracingStateChanged;
 #endif // WITH_STATETREE_TRACE
 
 #if WITH_STATETREE_TRACE_DEBUGGER
+FOnTraceAnalysisStateChanged OnTraceAnalysisStateChanged;
 FOnTracingTimelineScrubbed OnTracingTimelineScrubbed;
 #endif // WITH_STATETREE_TRACE_DEBUGGER
 
diff --git a/Engine/Plugins/Runtime/StateTree/Source/StateTreeModule/Private/StateTreeExecutionContext.cpp b/Engine/Plugins/Runtime/StateTree/Source/StateTreeModule/Private/StateTreeExecutionContext.cpp
--- a/Engine/Plugins/Runtime/StateTree/Source/StateTreeModule/Private/StateTreeExecutionContext.cpp
+++ b/Engine/Plugins/Runtime/StateTree/Source/StateTreeModule/Private/StateTreeExecutionContext.cpp
@@ -11,6 +11,7 @@
 #include "StateTreePropertyFunctionBase.h"
 #include "StateTreeReference.h"
 #include "AutoRTFM.h"
+#include "ObjectTrace.h"
 #include "Containers/StaticArray.h"
 #include "CrashReporter/StateTreeCrashReporterHandler.h"
 #include "Debugger/StateTreeDebug.h"
@@ -26,9 +27,9 @@
 
 namespace UE::StateTree::Debug
 {
-	constexpr int32 IndentSize = 2;	// Debug printing indent for hierarchical data.
-	static std::atomic<uint32> InstanceSerialNumber = 0;
-}; // namespace
+	// Debug printing indent for hierarchical data.
+	constexpr int32 IndentSize = 2;
+};
 
 namespace UE::StateTree::ExecutionContext::Private
 {
@@ -725,19 +726,13 @@ FString FStateTreeReadOnlyExecutionContext::GetInstanceDescriptionInternal() con
 }
 
 #if WITH_STATETREE_TRACE
-UE_AUTORTFM_ALWAYS_OPEN
-static uint32 GetNextInstanceSerialNumber()
-{
-	// The instance serial number is only used to synthesize unique instance debug IDs; rollback isn't needed.
-	return ++UE::StateTree::Debug::InstanceSerialNumber;
-}
-
 FStateTreeInstanceDebugId FStateTreeReadOnlyExecutionContext::GetInstanceDebugId() const
 {
 	FStateTreeInstanceDebugId& InstanceDebugId = Storage.GetMutableExecutionState().InstanceDebugId;
 	if (!InstanceDebugId.IsValid())
 	{
-		InstanceDebugId = FStateTreeInstanceDebugId(GetTypeHash(GetInstanceDescriptionInternal()), GetNextInstanceSerialNumber());
+		// Using an Id from the object trace pool to allow each instance to be a unique debuggable object in the Traces and RewindDebugger
+		InstanceDebugId = FStateTreeInstanceDebugId(FObjectTrace::AllocateInstanceId());
 	}
 	return InstanceDebugId;
 }
diff --git a/Engine/Plugins/Runtime/StateTree/Source/StateTreeModule/Public/Debugger/StateTreeDebug.h b/Engine/Plugins/Runtime/StateTree/Source/StateTreeModule/Public/Debugger/StateTreeDebug.h
--- a/Engine/Plugins/Runtime/StateTree/Source/StateTreeModule/Public/Debugger/StateTreeDebug.h
+++ b/Engine/Plugins/Runtime/StateTree/Source/StateTreeModule/Public/Debugger/StateTreeDebug.h
@@ -31,59 +31,56 @@ enum class EStateTreeTraceEventType : uint8;
 
 #define UE_STATETREE_DEBUG_ENTER_PHASE(InCtx, InPhase) \
 		::UE::StateTree::Debug::EnterPhase(*(InCtx), (InPhase), FStateTreeStateHandle::Invalid); \
-		TRACE_STATETREE_PHASE_EVENT((InCtx)->GetOwner(), (InCtx)->GetInstanceDebugId(), (InPhase), EStateTreeTraceEventType::Push, FStateTreeStateHandle::Invalid)
+		TRACE_STATETREE_PHASE_EVENT((InCtx), (InPhase), EStateTreeTraceEventType::Push, FStateTreeStateHandle::Invalid)
 
 #define UE_STATETREE_DEBUG_EXIT_PHASE(InCtx, InPhase) \
 		::UE::StateTree::Debug::ExitPhase(*(InCtx), (InPhase), FStateTreeStateHandle::Invalid); \
-		TRACE_STATETREE_PHASE_EVENT((InCtx)->GetOwner(), (InCtx)->GetInstanceDebugId(), (InPhase), EStateTreeTraceEventType::Pop, FStateTreeStateHandle::Invalid)
+		TRACE_STATETREE_PHASE_EVENT((InCtx), (InPhase), EStateTreeTraceEventType::Pop, FStateTreeStateHandle::Invalid)
 
 #define UE_STATETREE_DEBUG_STATE_EVENT(InCtx, InStateHandle, InEventType) \
 		::UE::StateTree::Debug::StateEvent(*(InCtx), (InStateHandle), (InEventType)); \
-		TRACE_STATETREE_STATE_EVENT((InCtx)->GetOwner(), (InCtx)->GetInstanceDebugId(), (InStateHandle), (InEventType))
+		TRACE_STATETREE_STATE_EVENT((InCtx), (InStateHandle), (InEventType))
 
 #define UE_STATETREE_DEBUG_TRANSITION_EVENT(InCtx, InTransitionSource, InEventType) \
 		::UE::StateTree::Debug::TransitionEvent(*(InCtx), (InTransitionSource), (InEventType)); \
-		TRACE_STATETREE_TRANSITION_EVENT((InCtx)->GetOwner(), (InCtx)->GetInstanceDebugId(), (InTransitionSource), (InEventType))
+		TRACE_STATETREE_TRANSITION_EVENT((InCtx), (InTransitionSource), (InEventType))
 
-#define UE_STATETREE_DEBUG_LOG_EVENT(InCtx, InLogVerbosity, InFormat, ...) TRACE_STATETREE_LOG_EVENT((InCtx)->GetOwner(), (InCtx)->GetInstanceDebugId(), InLogVerbosity, InFormat, ##__VA_ARGS__)
-#define UE_STATETREE_DEBUG_CONDITION_EVENT(InCtx, InIndex, InDataView, InEventType) TRACE_STATETREE_CONDITION_EVENT((InCtx)->GetOwner(), (InCtx)->GetInstanceDebugId(), (InCtx)->StealNodeCustomDebugTraceData(), FStateTreeIndex16(InIndex), (InDataView), (InEventType));
-#define UE_STATETREE_DEBUG_INSTANCE_EVENT(InCtx, InEventType) TRACE_STATETREE_INSTANCE_EVENT((InCtx)->GetOwner(), (InCtx)->GetInstanceDebugId(), (InCtx)->GetStateTree(), *((InCtx)->GetInstanceDescriptionInternal()), (InEventType));
-#define UE_STATETREE_DEBUG_INSTANCE_FRAME_EVENT(InCtx, InFrame) TRACE_STATETREE_INSTANCE_FRAME_EVENT((InCtx)->GetOwner(), (InCtx)->GetInstanceDebugId(), (InFrame));
-#define UE_STATETREE_DEBUG_ACTIVE_STATES_EVENT(InCtx, InActiveFrames) TRACE_STATETREE_ACTIVE_STATES_EVENT((InCtx)->GetOwner(), (InCtx)->GetInstanceDebugId(), InActiveFrames);
-#define UE_STATETREE_DEBUG_TASK_EVENT(InCtx, InIndex, InDataView, InEventType, InStatus) TRACE_STATETREE_TASK_EVENT((InCtx)->GetOwner(), (InCtx)->GetInstanceDebugId(), (InCtx)->StealNodeCustomDebugTraceData(), FStateTreeIndex16(InIndex), (InDataView), (InEventType), (InStatus));
-#define UE_STATETREE_DEBUG_EVALUATOR_EVENT(InCtx, InIndex, InDataView, InEventType) TRACE_STATETREE_EVALUATOR_EVENT((InCtx)->GetOwner(), (InCtx)->GetInstanceDebugId(), (InCtx)->StealNodeCustomDebugTraceData(), FStateTreeIndex16(InIndex), (InDataView), (InEventType));
+#define UE_STATETREE_DEBUG_LOG_EVENT(InCtx, InLogVerbosity, InFormat, ...) TRACE_STATETREE_LOG_EVENT((InCtx), InLogVerbosity, InFormat, ##__VA_ARGS__)
+#define UE_STATETREE_DEBUG_CONDITION_EVENT(InCtx, InIndex, InDataView, InEventType) TRACE_STATETREE_CONDITION_EVENT((InCtx), FStateTreeIndex16(InIndex), (InDataView), (InEventType));
+#define UE_STATETREE_DEBUG_INSTANCE_EVENT(InCtx, InEventType) TRACE_STATETREE_INSTANCE_EVENT((InCtx), (InEventType));
+#define UE_STATETREE_DEBUG_INSTANCE_FRAME_EVENT(InCtx, InFrame) TRACE_STATETREE_INSTANCE_FRAME_EVENT((InCtx), (InFrame));
+#define UE_STATETREE_DEBUG_ACTIVE_STATES_EVENT(InCtx, InActiveFrames) TRACE_STATETREE_ACTIVE_STATES_EVENT((InCtx), InActiveFrames);
+#define UE_STATETREE_DEBUG_TASK_EVENT(InCtx, InIndex, InDataView, InEventType, InStatus) TRACE_STATETREE_TASK_EVENT((InCtx), FStateTreeIndex16(InIndex), (InDataView), (InEventType), (InStatus));
+#define UE_STATETREE_DEBUG_EVALUATOR_EVENT(InCtx, InIndex, InDataView, InEventType) TRACE_STATETREE_EVALUATOR_EVENT((InCtx), FStateTreeIndex16(InIndex), (InDataView), (InEventType));
 
 #define UE_STATETREE_ID_NAME PREPROCESSOR_JOIN(InstanceId,__LINE__) \
 
 /** Scope based macros captures the instance ID since it might not be accessible when exiting the scope */
 #define UE_STATETREE_DEBUG_SCOPED_PHASE(InCtx, InPhase) \
-		FStateTreeInstanceDebugId UE_STATETREE_ID_NAME = (InCtx)->GetInstanceDebugId(); \
 		UE_STATETREE_DEBUG_ENTER_PHASE((InCtx), (InPhase)); \
 		ON_SCOPE_EXIT \
 		{ \
 			::UE::StateTree::Debug::ExitPhase(*(InCtx), (InPhase), FStateTreeStateHandle::Invalid); \
-			TRACE_STATETREE_PHASE_EVENT((InCtx)->GetOwner(), UE_STATETREE_ID_NAME, (InPhase), EStateTreeTraceEventType::Pop, FStateTreeStateHandle::Invalid); \
+			TRACE_STATETREE_PHASE_EVENT((InCtx), (InPhase), EStateTreeTraceEventType::Pop, FStateTreeStateHandle::Invalid); \
 		}
 
 /** Scope based macros captures the instance ID since it might not be accessible when exiting the scope */
 #define UE_STATETREE_DEBUG_SCOPED_STATE(InCtx, InStateHandle) \
-		FStateTreeInstanceDebugId UE_STATETREE_ID_NAME = (InCtx)->GetInstanceDebugId(); \
 		UE_STATETREE_DEBUG_STATE_EVENT((InCtx), (InStateHandle), EStateTreeTraceEventType::Push); \
 		ON_SCOPE_EXIT \
 		{ \
 			::UE::StateTree::Debug::StateEvent(*(InCtx), (InStateHandle), EStateTreeTraceEventType::Pop); \
-			TRACE_STATETREE_STATE_EVENT((InCtx)->GetOwner(), UE_STATETREE_ID_NAME, (InStateHandle), EStateTreeTraceEventType::Pop); \
+			TRACE_STATETREE_STATE_EVENT((InCtx), (InStateHandle), EStateTreeTraceEventType::Pop); \
 		}
 
 /** Scope based macros captures the instance ID since it might not be accessible when exiting the scope */
 #define UE_STATETREE_DEBUG_SCOPED_STATE_PHASE(InCtx, InStateHandle, InPhase) \
-		FStateTreeInstanceDebugId UE_STATETREE_ID_NAME = (InCtx)->GetInstanceDebugId(); \
 		::UE::StateTree::Debug::EnterPhase(*(InCtx), (InPhase), (InStateHandle)); \
-		TRACE_STATETREE_PHASE_EVENT((InCtx)->GetOwner(), UE_STATETREE_ID_NAME, (InPhase), EStateTreeTraceEventType::Push, (InStateHandle)); \
+		TRACE_STATETREE_PHASE_EVENT((InCtx), (InPhase), EStateTreeTraceEventType::Push, (InStateHandle)); \
 		ON_SCOPE_EXIT \
 		{ \
 			::UE::StateTree::Debug::ExitPhase(*(InCtx), (InPhase), (InStateHandle)); \
-			TRACE_STATETREE_PHASE_EVENT((InCtx)->GetOwner(), UE_STATETREE_ID_NAME, (InPhase), EStateTreeTraceEventType::Pop, (InStateHandle)); \
+			TRACE_STATETREE_PHASE_EVENT((InCtx), (InPhase), EStateTreeTraceEventType::Pop, (InStateHandle)); \
 		}
 
 namespace UE::StateTree::Debug
diff --git a/Engine/Plugins/Runtime/StateTree/Source/StateTreeModule/Public/Debugger/StateTreeTrace.h b/Engine/Plugins/Runtime/StateTree/Source/StateTreeModule/Public/Debugger/StateTreeTrace.h
--- a/Engine/Plugins/Runtime/StateTree/Source/StateTreeModule/Public/Debugger/StateTreeTrace.h
+++ b/Engine/Plugins/Runtime/StateTree/Source/StateTreeModule/Public/Debugger/StateTreeTrace.h
@@ -8,14 +8,15 @@
 #include "Containers/ContainersFwd.h"
 #include "Containers/UnrealString.h"
 #include "Containers/StringView.h"
+#include "Misc/NotNull.h"
 
-
-struct FStateTreeExecutionFrame;
 class UStateTree;
-struct FStateTreeDataView;
 struct FStateTreeActiveStates;
-struct FStateTreeInstanceDebugId;
+struct FStateTreeDataView;
+struct FStateTreeExecutionFrame;
 struct FStateTreeIndex16;
+struct FStateTreeInstanceDebugId;
+struct FStateTreeReadOnlyExecutionContext;
 struct FStateTreeStateHandle;
 struct FStateTreeTransitionSource;
 enum class EStateTreeStateSelectionBehavior : uint8;
@@ -26,6 +27,8 @@ namespace ELogVerbosity { enum Type : uint8; }
 
 UE_TRACE_CHANNEL_EXTERN(StateTreeDebugChannel, STATETREEMODULE_API)
 
+#define UE_API STATETREEMODULE_API
+
 namespace UE::StateTreeTrace
 {
 
@@ -98,92 +101,94 @@ private:
 void RegisterGlobalDelegates();
 void UnregisterGlobalDelegates();
 void OutputAssetDebugIdEvent(const UStateTree* StateTree, FStateTreeIndex16 AssetDebugId);
-void OutputPhaseScopeEvent(const UObject* Owner, FStateTreeInstanceDebugId InstanceId, EStateTreeUpdatePhase Phase, EStateTreeTraceEventType EventType, FStateTreeStateHandle StateHandle);
-void OutputInstanceLifetimeEvent(const UObject* Owner, FStateTreeInstanceDebugId InstanceId, const UStateTree* StateTree, const TCHAR* InstanceName, EStateTreeTraceEventType EventType);
-void OutputInstanceAssetEvent(const UObject* Owner, FStateTreeInstanceDebugId InstanceId, const UStateTree* StateTree);
-void OutputInstanceFrameEvent(const UObject* Owner, FStateTreeInstanceDebugId InstanceId, const FStateTreeExecutionFrame* Frame);
-void OutputLogEventTrace(const UObject* Owner, FStateTreeInstanceDebugId InstanceId, ELogVerbosity::Type Verbosity, const TCHAR* Fmt, ...);
-void OutputStateEventTrace(const UObject* Owner, FStateTreeInstanceDebugId InstanceId, FStateTreeStateHandle StateHandle, EStateTreeTraceEventType EventType);
-void OutputTaskEventTrace(const UObject* Owner, FStateTreeInstanceDebugId InstanceId, FNodeCustomDebugData&& DebugData, FStateTreeIndex16 TaskIdx, FStateTreeDataView DataView, EStateTreeTraceEventType EventType, EStateTreeRunStatus Status);
-void OutputEvaluatorEventTrace(const UObject* Owner, FStateTreeInstanceDebugId InstanceId, FNodeCustomDebugData&& DebugData, FStateTreeIndex16 EvaluatorIdx, FStateTreeDataView DataView, EStateTreeTraceEventType EventType);
-void OutputConditionEventTrace(const UObject* Owner, FStateTreeInstanceDebugId InstanceId, FNodeCustomDebugData&& DebugData, FStateTreeIndex16 ConditionIdx, FStateTreeDataView DataView, EStateTreeTraceEventType EventType);
-void OutputTransitionEventTrace(const UObject* Owner, FStateTreeInstanceDebugId InstanceId, FStateTreeTransitionSource TransitionSource, EStateTreeTraceEventType EventType);
-void OutputActiveStatesEventTrace(const UObject* Owner, FStateTreeInstanceDebugId InstanceId, TConstArrayView<FStateTreeExecutionFrame> ActiveFrames);
+void OutputPhaseScopeEvent(TNotNull<const FStateTreeReadOnlyExecutionContext*> Context, EStateTreeUpdatePhase Phase, EStateTreeTraceEventType EventType, FStateTreeStateHandle StateHandle);
+void OutputInstanceLifetimeEvent(TNotNull<const FStateTreeReadOnlyExecutionContext*> Context, EStateTreeTraceEventType EventType);
+void OutputInstanceAssetEvent(TNotNull<const FStateTreeReadOnlyExecutionContext*> Context, const UStateTree* StateTree);
+void OutputInstanceFrameEvent(TNotNull<const FStateTreeReadOnlyExecutionContext*> Context, const FStateTreeExecutionFrame* Frame);
+void OutputLogEventTrace(TNotNull<const FStateTreeReadOnlyExecutionContext*> Context, ELogVerbosity::Type Verbosity, const TCHAR* Fmt, ...);
+void OutputStateEventTrace(TNotNull<const FStateTreeReadOnlyExecutionContext*> Context, FStateTreeStateHandle StateHandle, EStateTreeTraceEventType EventType);
+void OutputTaskEventTrace(TNotNull<const FStateTreeReadOnlyExecutionContext*> Context, FStateTreeIndex16 TaskIdx, FStateTreeDataView DataView, EStateTreeTraceEventType EventType, EStateTreeRunStatus Status);
+void OutputEvaluatorEventTrace(TNotNull<const FStateTreeReadOnlyExecutionContext*> Context, FStateTreeIndex16 EvaluatorIdx, FStateTreeDataView DataView, EStateTreeTraceEventType EventType);
+void OutputConditionEventTrace(TNotNull<const FStateTreeReadOnlyExecutionContext*> Context, FStateTreeIndex16 ConditionIdx, FStateTreeDataView DataView, EStateTreeTraceEventType EventType);
+void OutputTransitionEventTrace(TNotNull<const FStateTreeReadOnlyExecutionContext*> Context, FStateTreeTransitionSource TransitionSource, EStateTreeTraceEventType EventType);
+void OutputActiveStatesEventTrace(TNotNull<const FStateTreeReadOnlyExecutionContext*> Context, TConstArrayView<FStateTreeExecutionFrame> ActiveFrames);
 
 UE_DEPRECATED(5.6, "This method will no longer be exposed publicly.")
 FStateTreeIndex16 FindOrAddDebugIdForAsset(const UStateTree* StateTree);
-UE_DEPRECATED(5.7, "Use the overload taking a pointer to the UObject owner instead.")
+UE_DEPRECATED(5.7, "Use the overload taking a pointer to the execution context instead.")
 void OutputPhaseScopeEvent(FStateTreeInstanceDebugId InstanceId, EStateTreeUpdatePhase Phase, EStateTreeTraceEventType EventType, FStateTreeStateHandle StateHandle);
-UE_DEPRECATED(5.7, "Use the overload taking a pointer to the UObject owner instead.")
+UE_DEPRECATED(5.7, "Use the overload taking a pointer to the execution context instead.")
 void OutputInstanceLifetimeEvent(FStateTreeInstanceDebugId InstanceId, const UStateTree* StateTree, const TCHAR* InstanceName, EStateTreeTraceEventType EventType);
-UE_DEPRECATED(5.7, "Use the overload taking a pointer to the UObject owner instead.")
+UE_DEPRECATED(5.7, "Use the overload taking a pointer to the execution context instead.")
 void OutputInstanceAssetEvent(FStateTreeInstanceDebugId InstanceId, const UStateTree* StateTree);
-UE_DEPRECATED(5.7, "Use the overload taking a pointer to the UObject owner instead.")
+UE_DEPRECATED(5.7, "Use the overload taking a pointer to the execution context instead.")
 void OutputInstanceFrameEvent(FStateTreeInstanceDebugId InstanceId, const FStateTreeExecutionFrame* Frame);
-UE_DEPRECATED(5.7, "Use the overload taking a pointer to the UObject owner instead.")
+UE_DEPRECATED(5.7, "Use the overload taking a pointer to the execution context instead.")
 void OutputLogEventTrace(FStateTreeInstanceDebugId InstanceId, ELogVerbosity::Type Verbosity, const TCHAR* Fmt, ...);
-UE_DEPRECATED(5.7, "Use the overload taking a pointer to the UObject owner instead.")
+UE_DEPRECATED(5.7, "Use the overload taking a pointer to the execution context instead.")
 void OutputStateEventTrace(FStateTreeInstanceDebugId InstanceId, FStateTreeStateHandle StateHandle, EStateTreeTraceEventType EventType);
-UE_DEPRECATED(5.7, "Use the overload taking a pointer to the UObject owner instead.")
+UE_DEPRECATED(5.7, "Use the overload taking a pointer to the execution context instead.")
 void OutputTaskEventTrace(FStateTreeInstanceDebugId InstanceId, FNodeCustomDebugData&& DebugData, FStateTreeIndex16 TaskIdx, FStateTreeDataView DataView, EStateTreeTraceEventType EventType, EStateTreeRunStatus Status);
-UE_DEPRECATED(5.7, "Use the overload taking a pointer to the UObject owner instead.")
+UE_DEPRECATED(5.7, "Use the overload taking a pointer to the execution context instead.")
 void OutputEvaluatorEventTrace(FStateTreeInstanceDebugId InstanceId, FNodeCustomDebugData&& DebugData, FStateTreeIndex16 EvaluatorIdx, FStateTreeDataView DataView, EStateTreeTraceEventType EventType);
-UE_DEPRECATED(5.7, "Use the overload taking a pointer to the UObject owner instead.")
+UE_DEPRECATED(5.7, "Use the overload taking a pointer to the execution context instead.")
 void OutputConditionEventTrace(FStateTreeInstanceDebugId InstanceId, FNodeCustomDebugData&& DebugData, FStateTreeIndex16 ConditionIdx, FStateTreeDataView DataView, EStateTreeTraceEventType EventType);
-UE_DEPRECATED(5.7, "Use the overload taking a pointer to the UObject owner instead.")
+UE_DEPRECATED(5.7, "Use the overload taking a pointer to the execution context instead.")
 void OutputTransitionEventTrace(FStateTreeInstanceDebugId InstanceId, FStateTreeTransitionSource TransitionSource, EStateTreeTraceEventType EventType);
-UE_DEPRECATED(5.7, "Use the overload taking a pointer to the UObject owner instead.")
+UE_DEPRECATED(5.7, "Use the overload taking a pointer to the execution context instead.")
 void OutputActiveStatesEventTrace(FStateTreeInstanceDebugId InstanceId, TConstArrayView<FStateTreeExecutionFrame> ActiveFrames);
 
 } // UE::StateTreeTrace
 
-#define TRACE_STATETREE_INSTANCE_EVENT(Owner, InstanceID, StateTree, InstanceName, EventType) \
-	UE::StateTreeTrace::OutputInstanceLifetimeEvent(Owner, InstanceID, StateTree, InstanceName, EventType);
+#define TRACE_STATETREE_INSTANCE_EVENT(Context, EventType) \
+	UE::StateTreeTrace::OutputInstanceLifetimeEvent(Context, EventType);
 
-#define TRACE_STATETREE_INSTANCE_FRAME_EVENT(Owner, InstanceID, Frame) \
-	UE::StateTreeTrace::OutputInstanceFrameEvent(Owner, InstanceID, Frame);
+#define TRACE_STATETREE_INSTANCE_FRAME_EVENT(Context, Frame) \
+	UE::StateTreeTrace::OutputInstanceFrameEvent(Context, Frame);
 
-#define TRACE_STATETREE_PHASE_EVENT(Owner, InstanceID, Phase, EventType, StateHandle) \
-	UE::StateTreeTrace::OutputPhaseScopeEvent(Owner, InstanceID, Phase, EventType, StateHandle); \
+#define TRACE_STATETREE_PHASE_EVENT(Context, Phase, EventType, StateHandle) \
+	UE::StateTreeTrace::OutputPhaseScopeEvent(Context, Phase, EventType, StateHandle); \
 
-#define TRACE_STATETREE_LOG_EVENT(Owner, InstanceId, TraceVerbosity, Format, ...) \
+#define TRACE_STATETREE_LOG_EVENT(Context, TraceVerbosity, Format, ...) \
 	if (UE_TRACE_CHANNELEXPR_IS_ENABLED(StateTreeDebugChannel)) \
 	{ \
-		UE::StateTreeTrace::OutputLogEventTrace(Owner, InstanceId, ELogVerbosity::TraceVerbosity, Format, ##__VA_ARGS__); \
+		UE::StateTreeTrace::OutputLogEventTrace(Context, ELogVerbosity::TraceVerbosity, Format, ##__VA_ARGS__); \
 	}
 
-#define TRACE_STATETREE_STATE_EVENT(Owner, InstanceId, StateHandle, EventType) \
+#define TRACE_STATETREE_STATE_EVENT(Context, StateHandle, EventType) \
 	if (UE_TRACE_CHANNELEXPR_IS_ENABLED(StateTreeDebugChannel)) \
 	{ \
-		UE::StateTreeTrace::OutputStateEventTrace(Owner, InstanceId, StateHandle, EventType); \
+		UE::StateTreeTrace::OutputStateEventTrace(Context, StateHandle, EventType); \
 	}
 
-#define TRACE_STATETREE_TASK_EVENT(Owner, InstanceId, DebugData, TaskIdx, DataView, EventType, Status) \
+#define TRACE_STATETREE_TASK_EVENT(Context, TaskIdx, DataView, EventType, Status) \
 	if (UE_TRACE_CHANNELEXPR_IS_ENABLED(StateTreeDebugChannel)) \
 	{ \
-		UE::StateTreeTrace::OutputTaskEventTrace(Owner, InstanceId, DebugData, TaskIdx, DataView, EventType, Status); \
+		UE::StateTreeTrace::OutputTaskEventTrace(Context, TaskIdx, DataView, EventType, Status); \
 	}
 
-#define TRACE_STATETREE_EVALUATOR_EVENT(Owner, InstanceId, DebugData, EvaluatorIdx, DataView, EventType) \
+#define TRACE_STATETREE_EVALUATOR_EVENT(Context, EvaluatorIdx, DataView, EventType) \
 	if (UE_TRACE_CHANNELEXPR_IS_ENABLED(StateTreeDebugChannel)) \
 	{ \
-		UE::StateTreeTrace::OutputEvaluatorEventTrace(Owner, InstanceId, DebugData, EvaluatorIdx, DataView, EventType); \
+		UE::StateTreeTrace::OutputEvaluatorEventTrace(Context, EvaluatorIdx, DataView, EventType); \
 	}
 
-#define TRACE_STATETREE_CONDITION_EVENT(Owner, InstanceId, DebugData, ConditionIdx, DataView, EventType) \
+#define TRACE_STATETREE_CONDITION_EVENT(Context, ConditionIdx, DataView, EventType) \
 	if (UE_TRACE_CHANNELEXPR_IS_ENABLED(StateTreeDebugChannel)) \
 	{ \
-		UE::StateTreeTrace::OutputConditionEventTrace(Owner, InstanceId, DebugData, ConditionIdx, DataView, EventType); \
+		UE::StateTreeTrace::OutputConditionEventTrace(Context, ConditionIdx, DataView, EventType); \
 	}
 
-#define TRACE_STATETREE_TRANSITION_EVENT(Owner, InstanceId, TransitionIdx, EventType) \
+#define TRACE_STATETREE_TRANSITION_EVENT(Context, TransitionIdx, EventType) \
 	if (UE_TRACE_CHANNELEXPR_IS_ENABLED(StateTreeDebugChannel)) \
 	{ \
-		UE::StateTreeTrace::OutputTransitionEventTrace(Owner, InstanceId, TransitionIdx, EventType); \
+		UE::StateTreeTrace::OutputTransitionEventTrace(Context, TransitionIdx, EventType); \
 	}
 
-#define TRACE_STATETREE_ACTIVE_STATES_EVENT(Owner, InstanceId, ActivateFrames) \
-		UE::StateTreeTrace::OutputActiveStatesEventTrace(Owner, InstanceId, ActivateFrames);
+#define TRACE_STATETREE_ACTIVE_STATES_EVENT(Context, ActivateFrames) \
+		UE::StateTreeTrace::OutputActiveStatesEventTrace(Context, ActivateFrames);
+
+#undef UE_API
 
 #else //WITH_STATETREE_TRACE
 
diff --git a/Engine/Plugins/Runtime/StateTree/Source/StateTreeModule/Public/Debugger/StateTreeTraceTypes.h b/Engine/Plugins/Runtime/StateTree/Source/StateTreeModule/Public/Debugger/StateTreeTraceTypes.h
--- a/Engine/Plugins/Runtime/StateTree/Source/StateTreeModule/Public/Debugger/StateTreeTraceTypes.h
+++ b/Engine/Plugins/Runtime/StateTree/Source/StateTreeModule/Public/Debugger/StateTreeTraceTypes.h
@@ -17,7 +17,16 @@ enum class EStateTreeTraceStatus : uint8
 {
 	TracesStarted,
 	StoppingTrace,
-	TracesStopped
+	TracesStopped,
+	Cleared
+};
+
+UENUM()
+enum class EStateTreeTraceAnalysisStatus : uint8
+{
+	Started,
+	Stopped,
+	Cleared
 };
 
 UENUM()
diff --git a/Engine/Plugins/Runtime/StateTree/Source/StateTreeModule/Public/StateTreeDelegates.h b/Engine/Plugins/Runtime/StateTree/Source/StateTreeModule/Public/StateTreeDelegates.h
--- a/Engine/Plugins/Runtime/StateTree/Source/StateTreeModule/Public/StateTreeDelegates.h
+++ b/Engine/Plugins/Runtime/StateTree/Source/StateTreeModule/Public/StateTreeDelegates.h
@@ -4,8 +4,11 @@
 
 #include "Delegates/Delegate.h"
 
+#define UE_API STATETREEMODULE_API
+
 class UStateTree;
 enum class EStateTreeTraceStatus : uint8;
+enum class EStateTreeTraceAnalysisStatus : uint8;
 
 namespace UE::StateTree::Delegates
 {
@@ -14,7 +17,7 @@ namespace UE::StateTree::Delegates
 
 /** Called when linkable name in a StateTree has changed. */
 DECLARE_MULTICAST_DELEGATE_OneParam(FOnIdentifierChanged, const UStateTree& /*StateTree*/);
-extern STATETREEMODULE_API FOnIdentifierChanged OnIdentifierChanged;
+extern UE_API FOnIdentifierChanged OnIdentifierChanged;
 
 /**
  * Called when schema of the StateTree EditorData has changed.
@@ -22,55 +25,55 @@ extern STATETREEMODULE_API FOnIdentifierChanged OnIdentifierChanged;
  * Note that this is NOT called when updating the StateTree schema from the EditorData on successful compilation.
  */
 DECLARE_MULTICAST_DELEGATE_OneParam(FOnSchemaChanged, const UStateTree& /*StateTree*/);
-extern STATETREEMODULE_API FOnSchemaChanged OnSchemaChanged;
+extern UE_API FOnSchemaChanged OnSchemaChanged;
 
 /**
  * Called when parameters of the StateTree EditorData changed.
- * This should mainly used by the asset editor to maintain consistency in the UI for manipulations on the EditorData
+ * This should mainly be used by the asset editor to maintain consistency in the UI for manipulations on the EditorData
  * until the tree gets compiled.
  */
 DECLARE_MULTICAST_DELEGATE_OneParam(FOnParametersChanged, const UStateTree& /*StateTree*/);
-extern STATETREEMODULE_API FOnParametersChanged OnParametersChanged;
+extern UE_API FOnParametersChanged OnParametersChanged;
 
 /**
  * Called when parameters of a StateTree State changed.
- * This should mainly used by the asset editor to maintain consistency in the UI for manipulations.
+ * This should mainly be used by the asset editor to maintain consistency in the UI for manipulations.
  */
 DECLARE_MULTICAST_DELEGATE_TwoParams(FOnStateParametersChanged, const UStateTree& /*StateTree*/, const FGuid /*StateID*/);
-extern STATETREEMODULE_API FOnStateParametersChanged OnStateParametersChanged;
+extern UE_API FOnStateParametersChanged OnStateParametersChanged;
 
 /**
  * Called when Global Tasks or Evaluators of the StateTree EditorData changed.
- * This should mainly used by the asset editor to maintain consistency in the UI for manipulations on the EditorData.
+ * This should mainly be used by the asset editor to maintain consistency in the UI for manipulations on the EditorData.
  */
 DECLARE_MULTICAST_DELEGATE_OneParam(FOnGlobalDataChanged, const UStateTree& /*StateTree*/);
-extern STATETREEMODULE_API FOnGlobalDataChanged OnGlobalDataChanged;
+extern UE_API FOnGlobalDataChanged OnGlobalDataChanged;
 
 /**
  * Called when the theme colors change.
- * This should mainly used by the asset editor to maintain consistency in the UI for manipulations on the EditorData.
+ * This should mainly be used by the asset editor to maintain consistency in the UI for manipulations on the EditorData.
  */
 DECLARE_MULTICAST_DELEGATE_OneParam(FOnVisualThemeChanged, const UStateTree& /*StateTree*/);
-extern STATETREEMODULE_API FOnVisualThemeChanged OnVisualThemeChanged;
+extern UE_API FOnVisualThemeChanged OnVisualThemeChanged;
 
 /**
  * Called when breakpoints of the StateTree EditorData changed.
- * This should mainly used by the asset editor to update the debugger.
+ * This should mainly be used by the asset editor to update the debugger.
  */
 DECLARE_MULTICAST_DELEGATE_OneParam(FOnBreakpointsChanged, const UStateTree& /*StateTree*/);
-extern STATETREEMODULE_API FOnBreakpointsChanged OnBreakpointsChanged;
+extern UE_API FOnBreakpointsChanged OnBreakpointsChanged;
 
 /** Called when compilation succeeds */
 DECLARE_MULTICAST_DELEGATE_OneParam(FOnPostCompile, const UStateTree& /*StateTree*/);
-extern STATETREEMODULE_API FOnPostCompile OnPostCompile;
+extern UE_API FOnPostCompile OnPostCompile;
 
 /** Request StateTree compilation. Works only in editor. */
 DECLARE_DELEGATE_RetVal_OneParam(bool, FOnRequestCompile, UStateTree& /*StateTreeToCompile*/);
-extern STATETREEMODULE_API FOnRequestCompile OnRequestCompile;
+extern UE_API FOnRequestCompile OnRequestCompile;
 
 /** Request the editor hash of the given state tree. */
 DECLARE_DELEGATE_RetVal_OneParam(uint32, FOnRequestEditorHash, const UStateTree& /*StateTreeToHash*/);
-extern STATETREEMODULE_API FOnRequestEditorHash OnRequestEditorHash;
+extern UE_API FOnRequestEditorHash OnRequestEditorHash;
 
 #endif // WITH_EDITOR
 
@@ -78,16 +81,22 @@ extern STATETREEMODULE_API FOnRequestEditorHash OnRequestEditorHash;
 
 /** Called by the StateTree module when StateTree traces are enabled/disabled. */
 DECLARE_MULTICAST_DELEGATE_OneParam(FOnTracingStateChanged, EStateTreeTraceStatus);
-extern STATETREEMODULE_API FOnTracingStateChanged OnTracingStateChanged;
+extern UE_API FOnTracingStateChanged OnTracingStateChanged;
 
 #endif // WITH_STATETREE_TRACE
 
 #if WITH_STATETREE_TRACE_DEBUGGER
 
+/** Called by the StateTree module when StateTree traces analysis is started/stopped. */
+DECLARE_MULTICAST_DELEGATE_OneParam(FOnTraceAnalysisStateChanged, EStateTreeTraceAnalysisStatus);
+extern UE_API FOnTraceAnalysisStateChanged OnTraceAnalysisStateChanged;
+
 /** Called by the StateTree module whenever tracing timeline is scrubbed on Rewind Debugger */
 DECLARE_MULTICAST_DELEGATE_OneParam(FOnTracingTimelineScrubbed, const double /* InScrubTime */);
-extern STATETREEMODULE_API FOnTracingTimelineScrubbed OnTracingTimelineScrubbed;
+extern UE_API FOnTracingTimelineScrubbed OnTracingTimelineScrubbed;
 
 #endif // WITH_STATETREE_TRACE_DEBUGGER
 
 }; // UE::StateTree::Delegates
+
+#undef UE_API
diff --git a/Engine/Plugins/Runtime/StateTree/Source/StateTreeModule/Public/StateTreeExecutionContext.h b/Engine/Plugins/Runtime/StateTree/Source/StateTreeModule/Public/StateTreeExecutionContext.h
--- a/Engine/Plugins/Runtime/StateTree/Source/StateTreeModule/Public/StateTreeExecutionContext.h
+++ b/Engine/Plugins/Runtime/StateTree/Source/StateTreeModule/Public/StateTreeExecutionContext.h
@@ -157,16 +157,53 @@ public:
 	UE_API void DebugPrintInternalLayout();
 #endif
 
-protected:
-	/** @return Prefix that will be used by STATETREE_LOG and STATETREE_CLOG, Owner name by default. */
-	UE_API FString GetInstanceDescriptionInternal() const;
-
 #if WITH_STATETREE_TRACE
 	/** A unique Id used by debugging tools to identify the instance. */
 	UE_API FStateTreeInstanceDebugId GetInstanceDebugId() const;
-#endif
+
+	/** @return Short description of the instance for debug/trace purposes */
+	FString GetInstanceDebugDescription() const
+	{
+		return GetInstanceDescriptionInternal();
+	}
+
+	void SetOuterTraceId(const uint64 Id) const
+	{
+		OuterTraceId = Id;
+	}
+
+	uint64 GetOuterTraceId() const
+	{
+		return OuterTraceId;
+	}
+
+	void SetNodeCustomDebugTraceData(UE::StateTreeTrace::FNodeCustomDebugData&& DebugData) const
+	{
+		ensureMsgf(!NodeCustomDebugTraceData .IsSet()
+			, TEXT("CustomData is not expected to be already set."
+				" This might indicate nested calls to SetNodeCustomDebugTraceData without calls to a trace macro"));
+		NodeCustomDebugTraceData = MoveTemp(DebugData);
+	}
+
+	UE::StateTreeTrace::FNodeCustomDebugData StealNodeCustomDebugTraceData() const
+	{
+		return MoveTemp(NodeCustomDebugTraceData);
+	}
+#else
+	void SetOuterTraceId(const uint64 Id) const
+	{
+	}
+
+	uint64 GetOuterTraceId() const
+	{
+		return 0;
+	}
+#endif // WITH_STATETREE_TRACE
 
 protected:
+	/** @return Description used as prefix by STATETREE_LOG and STATETREE_CLOG, Owner name by default. */
+	UE_API FString GetInstanceDescriptionInternal() const;
+
 	/** Owner of the instance data. */
 	UObject& Owner;
 
@@ -175,6 +212,11 @@ protected:
 
 	/** Data storage of the instance data. */
 	FStateTreeInstanceStorage& Storage;
+
+#if WITH_STATETREE_TRACE
+	mutable UE::StateTreeTrace::FNodeCustomDebugData NodeCustomDebugTraceData;
+	mutable uint64 OuterTraceId = 0;
+#endif
 };
 
 /**
@@ -216,25 +258,7 @@ public:
 	/** Sends event for the StateTree. */
 	UE_API void SendEvent(const FGameplayTag Tag, const FConstStructView Payload = FConstStructView(), const FName Origin = FName());
 
-#if WITH_STATETREE_TRACE
-	void SetNodeCustomDebugTraceData(UE::StateTreeTrace::FNodeCustomDebugData&& DebugData)
-	{
-		ensureMsgf(!NodeCustomDebugTraceData .IsSet()
-			, TEXT("CustomData is not expected to be already set."
-				" This might indicate nested calls to SetNodeCustomDebugTraceData without calls to a trace macro"));
-		NodeCustomDebugTraceData = MoveTemp(DebugData);
-	}
-#endif //WITH_STATETREE_TRACE
-
 protected:
-#if WITH_STATETREE_TRACE
-	UE::StateTreeTrace::FNodeCustomDebugData StealNodeCustomDebugTraceData()
-	{
-		return MoveTemp(NodeCustomDebugTraceData);
-	}
-
-	UE::StateTreeTrace::FNodeCustomDebugData NodeCustomDebugTraceData;
-#endif // WITH_STATETREE_TRACE
 
 	/** Informs the owner when the instance of the tree must woke up from a scheduled tick sleep. */
 	UE_API void ScheduleNextTick();
