[StateTreeDebugger] quick fix for CIS Issue 966541 (trace defines need to be revisited for shipping editor targets)

diff --git a/Engine/Plugins/Runtime/StateTree/Source/StateTreeModule/Private/Debugger/StateTreeTrace.cpp b/Engine/Plugins/Runtime/StateTree/Source/StateTreeModule/Private/Debugger/StateTreeTrace.cpp
--- a/Engine/Plugins/Runtime/StateTree/Source/StateTreeModule/Private/Debugger/StateTreeTrace.cpp
+++ b/Engine/Plugins/Runtime/StateTree/Source/StateTreeModule/Private/Debugger/StateTreeTrace.cpp
@@ -553,7 +553,12 @@ void TraceBufferedEvents(TNotNull<const FStateTreeReadOnlyExecutionContext*> Con
 		const UStateTree* StateTree = Context->GetStateTree();
 		const FString InstanceName = Context->GetInstanceDebugDescription();
 		const uint64 OuterTraceId = Context->GetOuterTraceId();
+
+#if OBJECT_TRACE_ENABLED
 		const uint64 OuterId = OuterTraceId == 0 ? FObjectTrace::GetObjectId(Owner) : OuterTraceId;
+#else
+		const uint64 OuterId = OuterTraceId;
+#endif
 
 		// Trace the owner (Object trace specific)
 		TRACE_OBJECT(Owner);
diff --git a/Engine/Plugins/Runtime/StateTree/Source/StateTreeModule/Private/StateTreeExecutionContext.cpp b/Engine/Plugins/Runtime/StateTree/Source/StateTreeModule/Private/StateTreeExecutionContext.cpp
--- a/Engine/Plugins/Runtime/StateTree/Source/StateTreeModule/Private/StateTreeExecutionContext.cpp
+++ b/Engine/Plugins/Runtime/StateTree/Source/StateTreeModule/Private/StateTreeExecutionContext.cpp
@@ -732,7 +732,11 @@ FStateTreeInstanceDebugId FStateTreeReadOnlyExecutionContext::GetInstanceDebugId
 	if (!InstanceDebugId.IsValid())
 	{
 		// Using an Id from the object trace pool to allow each instance to be a unique debuggable object in the Traces and RewindDebugger
+#if OBJECT_TRACE_ENABLED
 		InstanceDebugId = FStateTreeInstanceDebugId(FObjectTrace::AllocateInstanceId());
+#else
+		InstanceDebugId = FStateTreeInstanceDebugId(0);
+#endif
 	}
 	return InstanceDebugId;
 }
