[StateTree] Prevent crash when the schema is invalid. That can occur if the type doesn't exist anymore.

diff --git a/Engine/Plugins/Runtime/StateTree/Source/StateTreeEditorModule/Private/Customizations/StateTreeStateDetails.cpp b/Engine/Plugins/Runtime/StateTree/Source/StateTreeEditorModule/Private/Customizations/StateTreeStateDetails.cpp
--- a/Engine/Plugins/Runtime/StateTree/Source/StateTreeEditorModule/Private/Customizations/StateTreeStateDetails.cpp
+++ b/Engine/Plugins/Runtime/StateTree/Source/StateTreeEditorModule/Private/Customizations/StateTreeStateDetails.cpp
@@ -177,6 +177,7 @@ void FStateTreeStateDetails::CustomizeDetails(IDetailLayoutBuilder& DetailBuilde
 
 	// Type
 	TypeProperty->MarkHiddenByCustomization();
+	if (Schema)
 	{
 		// Restrict IsStateTypeAllowed
 		TSharedRef<FPropertyRestriction> SchemaSupported = MakeShared<FPropertyRestriction>(LOCTEXT("StateTypeAllowedRestriction", "The schema restricts the state type."));
@@ -202,17 +203,20 @@ void FStateTreeStateDetails::CustomizeDetails(IDetailLayoutBuilder& DetailBuilde
 	if (StateType == EStateTreeStateType::State || StateType == EStateTreeStateType::Subtree)
 	{
 		// Restrict IsStateSelectionAllowed
-		TSharedRef<FPropertyRestriction> SchemaSupported = MakeShared<FPropertyRestriction>(LOCTEXT("StateSelectionAllowedRestriction", "The schema restricts the selection behavior."));
-		const UEnum* EnumPtr = StaticEnum<EStateTreeStateSelectionBehavior>();
-		for (int32 Index = 0; Index < EnumPtr->NumEnums(); ++Index)
+		if (Schema)
 		{
-			const EStateTreeStateSelectionBehavior EnumValue = static_cast<EStateTreeStateSelectionBehavior>(EnumPtr->GetValueByIndex(Index));
-			if (!Schema->IsStateSelectionAllowed(EnumValue))
+			TSharedRef<FPropertyRestriction> SchemaSupported = MakeShared<FPropertyRestriction>(LOCTEXT("StateSelectionAllowedRestriction", "The schema restricts the selection behavior."));
+			const UEnum* EnumPtr = StaticEnum<EStateTreeStateSelectionBehavior>();
+			for (int32 Index = 0; Index < EnumPtr->NumEnums(); ++Index)
 			{
-				SchemaSupported->AddHiddenValue(EnumPtr->GetNameStringByIndex(Index));
+				const EStateTreeStateSelectionBehavior EnumValue = static_cast<EStateTreeStateSelectionBehavior>(EnumPtr->GetValueByIndex(Index));
+				if (!Schema->IsStateSelectionAllowed(EnumValue))
+				{
+					SchemaSupported->AddHiddenValue(EnumPtr->GetNameStringByIndex(Index));
+				}
 			}
+			SelectionBehaviorProperty->AddRestriction(SchemaSupported);
 		}
-		SelectionBehaviorProperty->AddRestriction(SchemaSupported);
 		StateCategory.AddProperty(SelectionBehaviorProperty);
 	}
 	else if (StateType == EStateTreeStateType::Linked)
