[StateTreeDebugger] transitioning debugging information to RewindDebugger
- Debugger window now only provides support for breakpoints, active states highlight and resume analysis after hitting a breakpoint.
- New window adds possibility to open RewindDebugger tabs, start recording and select instance to debug (which will also Select the matching track in RewindDebugger)
- added temporary option to toggle legacy debugger view: StateTreeEditorSettings -> EnableLegacyDebuggerWindow
- Updated IRewindDebugger interface to allow external request to StartRecording, set the object to debug and to select the track associated to a given object
- Fixed RecordingStopped callback never called in FRewindDebuggerRuntime

diff --git a/Engine/Plugins/Animation/GameplayInsights/Source/RewindDebugger/Private/RewindDebugger.cpp b/Engine/Plugins/Animation/GameplayInsights/Source/RewindDebugger/Private/RewindDebugger.cpp
--- a/Engine/Plugins/Animation/GameplayInsights/Source/RewindDebugger/Private/RewindDebugger.cpp
+++ b/Engine/Plugins/Animation/GameplayInsights/Source/RewindDebugger/Private/RewindDebugger.cpp
@@ -77,7 +77,7 @@ FRewindDebugger::FRewindDebugger()
 	{
 		Runtime->ClearRecording.AddRaw(this, &FRewindDebugger::OnClearRecording);
 		Runtime->RecordingStarted.AddRaw(this, &FRewindDebugger::OnRecordingStarted);
-		Runtime->RecordingStarted.AddRaw(this, &FRewindDebugger::OnRecordingStopped);
+		Runtime->RecordingStopped.AddRaw(this, &FRewindDebugger::OnRecordingStopped);
 	}
 
 	RewindDebugger::FRewindDebuggerTrackCreators::EnumerateCreators([this](const RewindDebugger::IRewindDebuggerTrackCreator* Creator)
@@ -163,6 +163,16 @@ void FRewindDebugger::Shutdown()
 	delete InternalInstance;
 }
 
+FName FRewindDebugger::GetMainTabName() const
+{
+	return FRewindDebuggerModule::MainTabName;
+}
+
+FName FRewindDebugger::GetDetailsTabName() const
+{
+	return FRewindDebuggerModule::DetailsTabName;
+}
+
 void FRewindDebugger::SetTrackListChangedDelegate(const FOnTrackListChanged& InTrackListChangedDelegate)
 {
 	TrackListChangedDelegate = InTrackListChangedDelegate;
@@ -1662,11 +1672,50 @@ TSharedPtr<FDebugObjectInfo> FRewindDebugger::GetSelectedObject() const
 	return TSharedPtr<FDebugObjectInfo>();
 }
 
+void FRewindDebugger::SetObjectToDebug(const RewindDebugger::FObjectId ObjectId)
+{
+	if (IsObjectCurrentlyDebugged(ObjectId.GetMainId()))
+	{
+		return;
+	}
+
+	CandidateIds.Reset(1);
+	CandidateIds.Push(ObjectId);
+
+	RefreshDebugTracks();
+}
+
 TSharedPtr<RewindDebugger::FRewindDebuggerTrack> FRewindDebugger::GetSelectedTrack() const
 {
 	return SelectedTrack;
 }
 
+void FRewindDebugger::SelectTrack(RewindDebugger::FObjectId ObjectId)
+{
+	for (const TSharedPtr<RewindDebugger::FRewindDebuggerTrack>& Track : DebugTracks)
+	{
+		if (Track->GetAssociatedObjectId() == ObjectId)
+		{
+			TrackSelectionChanged(Track);
+			break;
+		}
+
+		if (Track->VisitSubTracks([this, ObjectId](const TSharedPtr<RewindDebugger::FRewindDebuggerTrack>& ChildTrack)
+			{
+				if (ChildTrack->GetAssociatedObjectId() == ObjectId)
+				{
+					TrackSelectionChanged(ChildTrack);
+					return RewindDebugger::FRewindDebuggerTrack::EVisitResult::Break;
+				}
+
+				return RewindDebugger::FRewindDebuggerTrack::EVisitResult::Continue;
+			}) == RewindDebugger::FRewindDebuggerTrack::EVisitResult::Break)
+		{
+			break;
+		}
+	}
+}
+
 // build a tree that's compatible with the public api from 5.0 for GetDebuggedObjects.
 void FRewindDebugger::RefreshDebuggedObjects(TArray<TSharedPtr<RewindDebugger::FRewindDebuggerTrack>>& InTracks, TArray<TSharedPtr<FDebugObjectInfo>>& OutObjects)
 {
diff --git a/Engine/Plugins/Animation/GameplayInsights/Source/RewindDebugger/Private/RewindDebugger.h b/Engine/Plugins/Animation/GameplayInsights/Source/RewindDebugger/Private/RewindDebugger.h
--- a/Engine/Plugins/Animation/GameplayInsights/Source/RewindDebugger/Private/RewindDebugger.h
+++ b/Engine/Plugins/Animation/GameplayInsights/Source/RewindDebugger/Private/RewindDebugger.h
@@ -74,6 +74,8 @@ public:
 
 	virtual TSharedPtr<FDebugObjectInfo> GetSelectedObject() const override;
 	virtual TSharedPtr<RewindDebugger::FRewindDebuggerTrack> GetSelectedTrack() const override;
+	virtual void SelectTrack(RewindDebugger::FObjectId ObjectId) override;
+	virtual void SetObjectToDebug(RewindDebugger::FObjectId ObjectId) override;
 	virtual TArray<TSharedPtr<FDebugObjectInfo>>& GetDebuggedObjects() override;
 	virtual bool IsObjectCurrentlyDebugged(uint64 ObjectId) const override;
 	virtual bool ShouldDisplayWorld(uint64 WorldId) override
@@ -102,8 +104,11 @@ public:
 		return static_cast<FRewindDebugger*>(InternalInstance);
 	}
 
+	virtual FName GetMainTabName() const override;
+	virtual FName GetDetailsTabName() const override;
+
 	/** Start a new Recording:  Start tracing Object + Animation data, increment the current recording index, and reset the recording elapsed time to 0 */
-	void StartRecording() const;
+	virtual void StartRecording() const override;
 	void OnClearRecording();
 	void OnRecordingStarted();
 	void OnRecordingStopped();
diff --git a/Engine/Plugins/Runtime/StateTree/Source/StateTreeEditorModule/Private/Debugger/SStateTreeDebuggerView.cpp b/Engine/Plugins/Runtime/StateTree/Source/StateTreeEditorModule/Private/Debugger/SStateTreeDebuggerView.cpp
--- a/Engine/Plugins/Runtime/StateTree/Source/StateTreeEditorModule/Private/Debugger/SStateTreeDebuggerView.cpp
+++ b/Engine/Plugins/Runtime/StateTree/Source/StateTreeEditorModule/Private/Debugger/SStateTreeDebuggerView.cpp
@@ -5,6 +5,8 @@
 #include "SStateTreeDebuggerView.h"
 #include "Debugger/SStateTreeFrameEventsView.h"
 #include "Editor.h"
+#include "IGameplayProvider.h"
+#include "LevelEditor.h"
 #include "Framework/MultiBox/MultiBoxBuilder.h"
 #include "Kismet2/DebuggerCommands.h"
 #include "Modules/ModuleManager.h"
@@ -18,6 +20,7 @@
 #include "StateTreeEditorStyle.h"
 #include "StateTreeModule.h"
 #include "StateTreeViewModel.h"
+#include "Framework/Docking/TabManager.h"
 #include "Widgets/Input/SComboButton.h"
 #include "Widgets/Layout/SScrollBox.h"
 
@@ -123,6 +126,15 @@ void SStateTreeDebuggerView::ResumeDebuggerAnalysis() const
 	{
 		Debugger->ResumeSessionAnalysis();
 	}
+
+	// Resume PIE session if possible
+	if (FPlayWorldCommands::Get().ResumePlaySession.IsValid())
+	{
+		if (CommandList->CanExecuteAction(FPlayWorldCommands::Get().ResumePlaySession.ToSharedRef()))
+		{
+			CommandList->ExecuteAction(FPlayWorldCommands::Get().ResumePlaySession.ToSharedRef());
+		}
+	}
 }
 
 bool SStateTreeDebuggerView::CanResetTracks() const
@@ -190,18 +202,11 @@ void SStateTreeDebuggerView::Construct(const FArguments& InArgs
 	, const TSharedRef<FStateTreeViewModel>& InStateTreeViewModel
 	, const TSharedRef<FUICommandList>& InCommandList)
 {
-	BeginPIEHandle = FEditorDelegates::BeginPIE.AddSP(this, &SStateTreeDebuggerView::OnPIEStarted);
-	PausePIEHandle = FEditorDelegates::PausePIE.AddSP(this, &SStateTreeDebuggerView::OnPIEPaused);
-	ResumePIEHandle = FEditorDelegates::ResumePIE.AddSP(this, &SStateTreeDebuggerView::OnPIEResumed);
-	SingleStepPIEHandle = FEditorDelegates::SingleStepPIE.AddSP(this, &SStateTreeDebuggerView::OnPIESingleStepped);
-	// Using ShutdownPIE instead of EndPIE to make sure all traces emitted during world EndPlay get processed before disabling channels
-	ShutdownPIEHandle = FEditorDelegates::ShutdownPIE.AddSP(this, &SStateTreeDebuggerView::OnPIEStopped);
+	CommandList = InCommandList;
 
 	StateTreeViewModel = InStateTreeViewModel;
 	StateTree = InStateTree;
 	StateTreeEditorData = Cast<UStateTreeEditorData>(InStateTree->EditorData.Get());
-	CommandList = InCommandList;
-
 	Debugger = InStateTreeViewModel->GetDebugger();
 	check(Debugger);
 
@@ -209,6 +214,45 @@ void SStateTreeDebuggerView::Construct(const FArguments& InArgs
 	bRecording = StateTreeModule.IsTracing();
 	UE::StateTree::Delegates::OnTracingStateChanged.AddSP(this, &SStateTreeDebuggerView::HandleTracesStateChanged);
 
+	Debugger->OnBreakpointHit.BindSP(this, &SStateTreeDebuggerView::OnBreakpointHit, InCommandList);
+
+	if (GetDefault<UStateTreeEditorSettings>()->bEnableLegacyDebuggerWindow)
+	{
+		ConstructLegacyView(InCommandList);
+	}
+	else
+	{
+		ConstructView(InCommandList);
+	}
+
+	if (bRecording && !Debugger->IsAnalysisSessionActive())
+	{
+		// Auto-select session if there is only one available and that we are tracing
+		// Do that after creating all our widgets in case we receive a callback
+		TArray<FStateTreeDebugger::FTraceDescriptor> TraceDescriptors;
+		Debugger->GetLiveTraces(TraceDescriptors);
+		if (TraceDescriptors.Num() == 1)
+		{
+			Debugger->RequestSessionAnalysis(TraceDescriptors.Last());
+		}
+	}
+}
+
+void SStateTreeDebuggerView::ConstructLegacyView(const TSharedRef<FUICommandList>& InCommandList)
+{
+	// Put debugger in proper simulation state when view is constructed after PIE/SIE was started
+	if (FPlayWorldCommandCallbacks::HasPlayWorldAndPaused())
+	{
+		Debugger->PauseSessionAnalysis();
+	}
+
+	BeginPIEHandle = FEditorDelegates::BeginPIE.AddSP(this, &SStateTreeDebuggerView::OnPIEStarted);
+	PausePIEHandle = FEditorDelegates::PausePIE.AddSP(this, &SStateTreeDebuggerView::OnPIEPaused);
+	ResumePIEHandle = FEditorDelegates::ResumePIE.AddSP(this, &SStateTreeDebuggerView::OnPIEResumed);
+	SingleStepPIEHandle = FEditorDelegates::SingleStepPIE.AddSP(this, &SStateTreeDebuggerView::OnPIESingleStepped);
+	// Using ShutdownPIE instead of EndPIE to make sure all traces emitted during world EndPlay get processed before disabling channels
+	ShutdownPIEHandle = FEditorDelegates::ShutdownPIE.AddSP(this, &SStateTreeDebuggerView::OnPIEStopped);
+
 	// Registers a delegate to be notified when the associated StateTree asset get successfully recompiled
 	// to clear any previous recorded data that could mismatch new compiled tree data version.
 	UE::StateTree::Delegates::OnPostCompile.AddSPLambda(this, [this](const UStateTree& RecompiledStateTree)
@@ -223,16 +267,9 @@ void SStateTreeDebuggerView::Construct(const FArguments& InArgs
 	Debugger->OnNewSession.BindSP(this, &SStateTreeDebuggerView::OnNewSession);
 	Debugger->OnNewInstance.BindSP(this, &SStateTreeDebuggerView::OnNewInstance);
 	Debugger->OnScrubStateChanged.BindSP(this, &SStateTreeDebuggerView::OnDebuggerScrubStateChanged);
-	Debugger->OnBreakpointHit.BindSP(this, &SStateTreeDebuggerView::OnBreakpointHit, InCommandList);
 
 	// Bind our scrub time attribute to follow the value computed by the debugger
-	ScrubTimeAttribute = TAttribute<double>(InStateTreeViewModel->GetDebugger(), &FStateTreeDebugger::GetScrubTime);
-
-	// Put debugger in proper simulation state when view is constructed after PIE/SIE was started
-	if (FPlayWorldCommandCallbacks::HasPlayWorldAndPaused())
-	{
-		Debugger->PauseSessionAnalysis();
-	}
+	ScrubTimeAttribute = TAttribute<double>(StateTreeViewModel->GetDebugger(), &FStateTreeDebugger::GetScrubTime);
 
 	// Add & Bind commands
 	BindDebuggerToolbarCommands(InCommandList);
@@ -464,7 +501,7 @@ void SStateTreeDebuggerView::Construct(const FArguments& InArgs
 					+ SVerticalBox::Slot()
 					.FillHeight(1.0f)
 					[
-						SAssignNew(EventsView, UE::StateTreeDebugger::SFrameEventsView, InStateTree)
+						SAssignNew(EventsView, UE::StateTreeDebugger::SFrameEventsView, StateTree.Get())
 					]
 				]
 			]
@@ -485,17 +522,125 @@ void SStateTreeDebuggerView::Construct(const FArguments& InArgs
 
 		Debugger->SetScrubTime(Debugger->GetScrubTime());
 	}
-	else if (bRecording)
+}
+
+void SStateTreeDebuggerView::ConstructView(const TSharedRef<FUICommandList>& InCommandList)
+{
+	// Add & Bind commands
+	BindDebuggingToolbarCommands(InCommandList);
+
+	// Register the play world commands (to pause/resume PIE when hitting breakpoint)
+	InCommandList->Append(FPlayWorldCommands::GlobalPlayWorldActions.ToSharedRef());
+
+	// Toolbars
+	FSlimHorizontalToolBarBuilder LeftToolbar(InCommandList, FMultiBoxCustomization::None, /*InExtender*/ nullptr, /*InForceSmallIcons*/ true);
+	LeftToolbar.BeginSection(TEXT("Debugging"));
 	{
-		// Auto-select session if there is only one available and that we are tracing
-		// Do that after creating all our widgets in case we receive a callback
-		TArray<FStateTreeDebugger::FTraceDescriptor> TraceDescriptors;
-		Debugger->GetLiveTraces(TraceDescriptors);
-		if (TraceDescriptors.Num() == 1)
-		{
-			Debugger->RequestSessionAnalysis(TraceDescriptors.Last());
-		}
+		const FStateTreeDebuggerCommands& DebuggerCommands = FStateTreeDebuggerCommands::Get();
+		LeftToolbar.AddToolBarButton(DebuggerCommands.StartRewindDebugger);
+		LeftToolbar.AddToolBarButton(DebuggerCommands.ResumeDebuggerAnalysis);
 	}
+	LeftToolbar.EndSection();
+
+	// Instance selection combo
+	const TSharedRef<SWidget> InstanceSelectionBox = SNew(SComboButton)
+		.OnGetMenuContent_Lambda([Debugger = Debugger]
+			{
+				TArray<const TSharedRef<const UE::StateTreeDebugger::FInstanceDescriptor>> InstanceDescriptors;
+				Debugger->GetSessionInstanceDescriptors(InstanceDescriptors);
+
+				FMenuBuilder MenuBuilder(/*bInShouldCloseWindowAfterMenuSelection*/true, /*InCommandList*/nullptr);
+				for (const TSharedRef<const UE::StateTreeDebugger::FInstanceDescriptor>& InstanceDescriptor : InstanceDescriptors)
+				{
+					if (InstanceDescriptor.Get().StateTree != Debugger->GetAsset())
+					{
+						continue;
+					}
+					const FText Label = Debugger->DescribeInstance(InstanceDescriptor.Get());
+
+					FUIAction ItemAction(FExecuteAction::CreateSPLambda(Debugger.ToSharedRef(), [Debugger, InstanceId = InstanceDescriptor.Get().Id]
+						{
+							Debugger->SelectInstance(InstanceId);
+
+							IRewindDebugger* RewindDebugger = IRewindDebugger::Instance();
+
+							// By default, we set the debug object associated to the StateTree instance
+							RewindDebugger::FObjectId ObjectToDebug{InstanceId.ToUint64()};
+
+							// Then if possible we traverse its outer chain until the last outer before the level or world.
+							if (const TraceServices::IAnalysisSession* AnalysisSession = RewindDebugger->GetAnalysisSession())
+							{
+								const IGameplayProvider* GameplayProvider = AnalysisSession->ReadProvider<IGameplayProvider>("GameplayProvider");
+								TraceServices::FAnalysisSessionReadScope SessionReadScope(*AnalysisSession);
+								const FClassInfo* LevelClass = GameplayProvider->FindClassInfo(TEXT("/Script/Engine.Level"));
+								const FClassInfo* WorldClass = GameplayProvider->FindClassInfo(TEXT("/Script/Engine.World"));
+								const FObjectInfo* ObjectInfo = GameplayProvider->FindObjectInfo(RewindDebugger::FObjectId{InstanceId.ToUint64()});
+								while (ObjectInfo != nullptr)
+								{
+									if ((LevelClass != nullptr && ObjectInfo->ClassId == LevelClass->Id)
+										|| (WorldClass != nullptr && ObjectInfo->ClassId == WorldClass->Id))
+									{
+										break;
+									}
+
+									ObjectToDebug = ObjectInfo->GetId();
+									ObjectInfo = GameplayProvider->FindObjectInfo(ObjectInfo->GetOuterId());
+								}
+							}
+							RewindDebugger->SetObjectToDebug(ObjectToDebug);
+							RewindDebugger->SelectTrack(RewindDebugger::FObjectId{InstanceId.ToUint64()});
+						}));
+					MenuBuilder.AddMenuEntry(Label, TAttribute<FText>(), FSlateIcon(), ItemAction);
+				}
+
+				// Failsafe when no match
+				if (InstanceDescriptors.Num() == 0)
+				{
+					const FText Desc = LOCTEXT("NoActiveInstances", "Can't find active instances");
+					MenuBuilder.AddMenuEntry(Desc, TAttribute<FText>(), FSlateIcon(), FUIAction());
+				}
+
+				return MenuBuilder.MakeWidget();
+			})
+		.ButtonContent()
+		[
+			SNew(STextBlock)
+				.ToolTipText(LOCTEXT("SelectInstance", "Pick instance to debug"))
+				.Text_Lambda([Debugger = Debugger]
+					{
+						check(Debugger);
+						const TSharedPtr<const UE::StateTreeDebugger::FInstanceDescriptor> DebuggedInstance = Debugger->GetSelectedDescriptor();
+						return Debugger->DescribeInstance(DebuggedInstance ? *DebuggedInstance : UE::StateTreeDebugger::FInstanceDescriptor{});
+					})
+		];
+
+	ChildSlot
+		[
+			SNew(SBorder)
+				.Padding(4.0f)
+				.BorderImage(FAppStyle::GetBrush("ToolPanel.GroupBorder"))
+				[
+					SNew(SVerticalBox)
+						+ SVerticalBox::Slot()
+						.AutoHeight()
+						.MinHeight(48.f)
+						[
+							SNew(SHorizontalBox)
+								+ SHorizontalBox::Slot()
+								.AutoWidth()
+								.VAlign(VAlign_Center)
+								[
+									LeftToolbar.MakeWidget()
+								]
+								+ SHorizontalBox::Slot()
+								.AutoWidth()
+								.VAlign(VAlign_Center)
+								[
+									InstanceSelectionBox
+								]
+						]
+				]
+		];
 }
 
 FReply SStateTreeDebuggerView::OnPreviewKeyDown(const FGeometry& MyGeometry, const FKeyEvent& InKeyEvent)
@@ -544,7 +689,7 @@ void SStateTreeDebuggerView::Tick(const FGeometry& AllottedGeometry, const doubl
 		// When no events are processed this frame we force timeline to progress to make UI feel more responsive,
 		// This is useful when debugging a live recording where events are produced in realtime and can't be read faster from a saved recording.
 		// Note that we update regardless of autoscroll to maintain a valid value if it gets enabled.
-		if (!bHasMoreRecentData)
+		if (!bHasMoreRecentData && bRecording)
 		{
 			ExtrapolatedRecordedWorldTime += InDeltaTime;
 		}
@@ -577,8 +722,16 @@ void SStateTreeDebuggerView::RefreshTracks()
 
 	if (bChanged)
 	{
-		InstancesTreeView->Refresh();
-		InstanceTimelinesTreeView->Refresh();
+		if (InstancesTreeView)
+		{
+			InstancesTreeView->Refresh();
+		}
+
+		if (InstanceTimelinesTreeView)
+		{
+			InstanceTimelinesTreeView->Refresh();
+		}
+
 		TrackCursor();
 	}
 }
@@ -632,6 +785,51 @@ void SStateTreeDebuggerView::BindDebuggerToolbarCommands(const TSharedRef<FUICom
 		FCanExecuteAction::CreateSP(this, &SStateTreeDebuggerView::CanResetTracks));
 }
 
+void SStateTreeDebuggerView::BindDebuggingToolbarCommands(const TSharedRef<FUICommandList> ToolkitCommands)
+{
+	const FStateTreeDebuggerCommands& Commands = FStateTreeDebuggerCommands::Get();
+
+	ToolkitCommands->MapAction(
+		Commands.StartRewindDebugger,
+		FExecuteAction::CreateSP(this, &SStateTreeDebuggerView::StartRewindDebugger),
+		FCanExecuteAction::CreateSP(this, &SStateTreeDebuggerView::CanStartRewindDebugger));
+
+	ToolkitCommands->MapAction(
+		Commands.ResumeDebuggerAnalysis,
+		FExecuteAction::CreateSP(this, &SStateTreeDebuggerView::ResumeDebuggerAnalysis),
+		FCanExecuteAction::CreateSP(this, &SStateTreeDebuggerView::CanResumeDebuggerAnalysis));
+}
+
+bool SStateTreeDebuggerView::CanStartRewindDebugger() const
+{
+	if (!IsRecording())
+	{
+		return true;
+	}
+
+	if (const IRewindDebugger* RewindDebugger = IRewindDebugger::Instance())
+	{
+		const FLevelEditorModule& LevelEditorModule = FModuleManager::GetModuleChecked<FLevelEditorModule>("LevelEditor");
+		const TSharedPtr<FTabManager> LevelEditorTabManager = LevelEditorModule.GetLevelEditorTabManager();
+		return !LevelEditorTabManager->FindExistingLiveTab(RewindDebugger->GetMainTabName()).IsValid();
+	}
+
+	return false;
+}
+
+void SStateTreeDebuggerView::StartRewindDebugger()
+{
+	if (const IRewindDebugger* RewindDebugger = IRewindDebugger::Instance())
+	{
+		const FLevelEditorModule& LevelEditorModule = FModuleManager::GetModuleChecked<FLevelEditorModule>("LevelEditor");
+		const TSharedPtr<FTabManager> LevelEditorTabManager = LevelEditorModule.GetLevelEditorTabManager();
+		LevelEditorTabManager->TryInvokeTab(RewindDebugger->GetMainTabName());
+		LevelEditorTabManager->TryInvokeTab(RewindDebugger->GetDetailsTabName());
+
+		RewindDebugger->StartRecording();
+	}
+}
+
 bool SStateTreeDebuggerView::CanUseScrubButtons() const
 {
 	// Nothing preventing use of scrub buttons on the Editor side at the moment.
diff --git a/Engine/Plugins/Runtime/StateTree/Source/StateTreeEditorModule/Private/Debugger/SStateTreeDebuggerView.h b/Engine/Plugins/Runtime/StateTree/Source/StateTreeEditorModule/Private/Debugger/SStateTreeDebuggerView.h
--- a/Engine/Plugins/Runtime/StateTree/Source/StateTreeEditorModule/Private/Debugger/SStateTreeDebuggerView.h
+++ b/Engine/Plugins/Runtime/StateTree/Source/StateTreeEditorModule/Private/Debugger/SStateTreeDebuggerView.h
@@ -58,6 +58,9 @@ public:
 	}
 
 private:
+	void ConstructLegacyView(const TSharedRef<FUICommandList>& InCommandList);
+	void ConstructView(const TSharedRef<FUICommandList>& InCommandList);
+
 	virtual FReply OnPreviewKeyDown(const FGeometry& MyGeometry, const FKeyEvent& InKeyEvent) override;
 	virtual void Tick(const FGeometry& AllottedGeometry, const double InCurrentTime, const float InDeltaTime) override;
 
@@ -77,8 +80,15 @@ private:
 	void OnNewSession();
 	void OnNewInstance(FStateTreeInstanceDebugId InstanceId);
 
+	/** Legacy debugger commands */
 	void BindDebuggerToolbarCommands(const TSharedRef<FUICommandList>& ToolkitCommands);
 
+	/** RewindDebugger sync commands */
+	void BindDebuggingToolbarCommands(const TSharedRef<FUICommandList> ToolkitCommands);
+
+	bool CanStartRewindDebugger() const;
+	void StartRewindDebugger();
+
 	bool CanUseScrubButtons() const;
 
 	bool CanStartRecording() const
diff --git a/Engine/Plugins/Runtime/StateTree/Source/StateTreeEditorModule/Private/Debugger/StateTreeDebuggerCommands.cpp b/Engine/Plugins/Runtime/StateTree/Source/StateTreeEditorModule/Private/Debugger/StateTreeDebuggerCommands.cpp
--- a/Engine/Plugins/Runtime/StateTree/Source/StateTreeEditorModule/Private/Debugger/StateTreeDebuggerCommands.cpp
+++ b/Engine/Plugins/Runtime/StateTree/Source/StateTreeEditorModule/Private/Debugger/StateTreeDebuggerCommands.cpp
@@ -31,6 +31,7 @@ void FStateTreeDebuggerCommands::RegisterCommands()
 	UI_COMMAND(NextFrameWithStateChange, "Next State Change", "Jump to next state changed frame", EUserInterfaceActionType::Button, FInputChord(EModifierKey::Control, EKeys::Right));
 	UI_COMMAND(ResumeDebuggerAnalysis, "Resume Analysis", "Resume trace analysis after hitting breakpoint", EUserInterfaceActionType::Button, FInputChord());
 	UI_COMMAND(ResetTracks, "Reset Tracks", "Clear current tracks and associated events.", EUserInterfaceActionType::Button, FInputChord());
+	UI_COMMAND(StartRewindDebugger, "Start RewindDebugger", "Open RewindDebugger window and start recording", EUserInterfaceActionType::Button, FInputChord());
 }
 
 #undef LOCTEXT_NAMESPACE
diff --git a/Engine/Plugins/Runtime/StateTree/Source/StateTreeEditorModule/Private/Debugger/StateTreeDebuggerCommands.h b/Engine/Plugins/Runtime/StateTree/Source/StateTreeEditorModule/Private/Debugger/StateTreeDebuggerCommands.h
--- a/Engine/Plugins/Runtime/StateTree/Source/StateTreeEditorModule/Private/Debugger/StateTreeDebuggerCommands.h
+++ b/Engine/Plugins/Runtime/StateTree/Source/StateTreeEditorModule/Private/Debugger/StateTreeDebuggerCommands.h
@@ -14,7 +14,7 @@ class FStateTreeDebuggerCommands : public TCommands<FStateTreeDebuggerCommands>
 public:
 	FStateTreeDebuggerCommands();
 
-	// TCommands<> overrides
+	//~ TCommands<> overrides
 	virtual void RegisterCommands() override;
 
 	TSharedPtr<FUICommandInfo> StartRecording;
@@ -25,6 +25,7 @@ public:
 	TSharedPtr<FUICommandInfo> NextFrameWithStateChange;
 	TSharedPtr<FUICommandInfo> ResumeDebuggerAnalysis;
 	TSharedPtr<FUICommandInfo> ResetTracks;
+	TSharedPtr<FUICommandInfo> StartRewindDebugger;
 };
 
 #endif // WITH_STATETREE_TRACE_DEBUGGER
\ No newline at end of file
diff --git a/Engine/Plugins/Runtime/StateTree/Source/StateTreeEditorModule/Private/StateTreeEditorStyle.cpp b/Engine/Plugins/Runtime/StateTree/Source/StateTreeEditorModule/Private/StateTreeEditorStyle.cpp
--- a/Engine/Plugins/Runtime/StateTree/Source/StateTreeEditorModule/Private/StateTreeEditorStyle.cpp
+++ b/Engine/Plugins/Runtime/StateTree/Source/StateTreeEditorModule/Private/StateTreeEditorStyle.cpp
@@ -239,6 +239,8 @@ FStateTreeEditorStyle::FStateTreeEditorStyle()
 		// From generic Engine Editor
 		FContentRootScope Scope(this, EngineEditorSlateContentDir);
 
+		Set("StateTreeEditor.Debugger.StartRewindDebugger", new IMAGE_BRUSH("Sequencer/Transport_Bar/Record_24x", CoreStyleConstants::Icon16x16));
+
 		Set("StateTreeEditor.Debugger.StartRecording", new IMAGE_BRUSH("Sequencer/Transport_Bar/Record_24x", CoreStyleConstants::Icon16x16));
 		Set("StateTreeEditor.Debugger.StopRecording", new IMAGE_BRUSH("Sequencer/Transport_Bar/Recording_24x", CoreStyleConstants::Icon16x16));
 
diff --git a/Engine/Plugins/Runtime/StateTree/Source/StateTreeEditorModule/Public/StateTreeEditorSettings.h b/Engine/Plugins/Runtime/StateTree/Source/StateTreeEditorModule/Public/StateTreeEditorSettings.h
--- a/Engine/Plugins/Runtime/StateTree/Source/StateTreeEditorModule/Public/StateTreeEditorSettings.h
+++ b/Engine/Plugins/Runtime/StateTree/Source/StateTreeEditorModule/Public/StateTreeEditorSettings.h
@@ -19,8 +19,11 @@ UCLASS(MinimalAPI, config = EditorPerProjectUserSettings)
 class UStateTreeEditorSettings : public UDeveloperSettings
 {
 	GENERATED_UCLASS_BODY()
-public:
-	static UStateTreeEditorSettings& Get() { return *CastChecked<UStateTreeEditorSettings>(UStateTreeEditorSettings::StaticClass()->GetDefaultObject()); }
+
+	static UStateTreeEditorSettings& Get()
+	{
+		return *GetMutableDefault<UStateTreeEditorSettings>();
+	}
 
 	/** Determines when to save StateTrees post-compile */
 	UPROPERTY(EditAnywhere, config, Category = "Compiler")
@@ -34,6 +37,14 @@ public:
 	UPROPERTY(EditAnywhere, config, Category = "Debugger")
 	bool bShouldDebuggerResetDataOnNewPIESession = false;
 
+	/**
+	 * If enabled, debugger window in the StateTree Asset Editor will display all widgets
+	 * related to the legacy debugger (recording controls, timelines, frame details, etc.).
+	 * Otherwise, it will display options to link to open RewindDebugger and select a given instance
+	 * */
+	UPROPERTY(EditAnywhere, config, Category = "Debugger")
+	bool bEnableLegacyDebuggerWindow = false;
+
 	/**
 	 * If enabled, changing the class of a node will try to copy over values of properties with the same name and type.
 	 * i.e. if you change one condition for another, and both have a "Target" BB key selector, it'll be kept.
diff --git a/Engine/Plugins/Runtime/StateTree/Source/StateTreeModule/Private/Debugger/StateTreeDebugger.cpp b/Engine/Plugins/Runtime/StateTree/Source/StateTreeModule/Private/Debugger/StateTreeDebugger.cpp
--- a/Engine/Plugins/Runtime/StateTree/Source/StateTreeModule/Private/Debugger/StateTreeDebugger.cpp
+++ b/Engine/Plugins/Runtime/StateTree/Source/StateTreeModule/Private/Debugger/StateTreeDebugger.cpp
@@ -86,6 +86,15 @@ FStateTreeDebugger::FStateTreeDebugger()
 			}
 		});
 
+	TraceAnalysisStateChangedHandle = UE::StateTree::Delegates::OnTraceAnalysisStateChanged.AddLambda([this](const EStateTreeTraceAnalysisStatus TraceAnalysisStatus)
+		{
+			if (TraceAnalysisStatus == EStateTreeTraceAnalysisStatus::Cleared
+				|| TraceAnalysisStatus == EStateTreeTraceAnalysisStatus::Stopped)
+			{
+				StopSessionAnalysis();
+			}
+		});
+
 	TracingTimelineScrubbedHandle = UE::StateTree::Delegates::OnTracingTimelineScrubbed.AddLambda([this](const double InScrubTime)
 		{
 			SetScrubTime(InScrubTime);
@@ -96,10 +105,8 @@ PRAGMA_DISABLE_DEPRECATION_WARNINGS
 FStateTreeDebugger::~FStateTreeDebugger()
 {
 	UE::StateTree::Delegates::OnTracingStateChanged.Remove(TracingStateChangedHandle);
-	TracingStateChangedHandle.Reset();
-
+	UE::StateTree::Delegates::OnTraceAnalysisStateChanged.Remove(TraceAnalysisStateChangedHandle);
 	UE::StateTree::Delegates::OnTracingTimelineScrubbed.Remove(TracingTimelineScrubbedHandle);
-	TracingTimelineScrubbedHandle.Reset();
 
 	StopSessionAnalysis();
 }
diff --git a/Engine/Plugins/Runtime/StateTree/Source/StateTreeModule/Private/StateTreeModule.cpp b/Engine/Plugins/Runtime/StateTree/Source/StateTreeModule/Private/StateTreeModule.cpp
--- a/Engine/Plugins/Runtime/StateTree/Source/StateTreeModule/Private/StateTreeModule.cpp
+++ b/Engine/Plugins/Runtime/StateTree/Source/StateTreeModule/Private/StateTreeModule.cpp
@@ -19,7 +19,6 @@
 #include "Debugger/StateTreeDebuggerTypes.h"
 #include "Debugger/StateTreeTraceModule.h"
 #include "Features/IModularFeatures.h"
-#include "Misc/CoreDelegates.h"
 #include "Trace/StoreClient.h"
 #include "TraceServices/AnalysisService.h"
 #include "TraceServices/ITraceServicesModule.h"
@@ -118,24 +117,24 @@ void FStateTreeModule::StartupModule()
 #endif // WITH_EDITORONLY_DATA
 
 #if WITH_EDITOR
-	 UE::PropertyBinding::PropertyBindingIndex16ConversionFuncList.Add([](const FPropertyTag& Tag, FStructuredArchive::FSlot Slot, TNotNull<FPropertyBindingIndex16*> Index)
-	 {
-		static FName FStateTreeIndex16TypeName = FStateTreeIndex16::StaticStruct()->GetFName();
-		const FName StructFName = Tag.GetType().GetParameter(0).GetName();
-		if (FStateTreeIndex16TypeName == StructFName)
+	UE::PropertyBinding::PropertyBindingIndex16ConversionFuncList.Add([](const FPropertyTag& Tag, FStructuredArchive::FSlot Slot, TNotNull<FPropertyBindingIndex16*> Index)
 		{
-			FStateTreeIndex16 StateTreeIndex16;
-			FStateTreeIndex16::StaticStruct()->SerializeItem(Slot, &StateTreeIndex16, /*Defaults*/nullptr);
-			*Index = StateTreeIndex16;
-			return true;
-		}
-		return false;
-	 });
+			static FName FStateTreeIndex16TypeName = FStateTreeIndex16::StaticStruct()->GetFName();
+			const FName StructFName = Tag.GetType().GetParameter(0).GetName();
+			if (FStateTreeIndex16TypeName == StructFName)
+			{
+				FStateTreeIndex16 StateTreeIndex16;
+				FStateTreeIndex16::StaticStruct()->SerializeItem(Slot, &StateTreeIndex16, /*Defaults*/nullptr);
+				*Index = StateTreeIndex16;
+				return true;
+			}
+			return false;
+		});
 #endif //WITH_EDITOR
 
 #if WITH_STATETREE_DEBUG
-	 PreGCHandle = FCoreUObjectDelegates::GetPreGarbageCollectDelegate().AddRaw(this, &FStateTreeModule::HandlePreGC);
-	 PostGCHandle = FCoreUObjectDelegates::GetPostPurgeGarbageDelegate().AddRaw(this, &FStateTreeModule::HandlePostGC);
+	PreGCHandle = FCoreUObjectDelegates::GetPreGarbageCollectDelegate().AddRaw(this, &FStateTreeModule::HandlePreGC);
+	PostGCHandle = FCoreUObjectDelegates::GetPostPurgeGarbageDelegate().AddRaw(this, &FStateTreeModule::HandlePostGC);
 #endif
 }
 
@@ -172,7 +171,7 @@ bool FStateTreeModule::StartTraces(int32& OutTraceId)
 	OutTraceId = INDEX_NONE;
 
 #if WITH_STATETREE_TRACE
-	if (IsRunningCommandlet() || bIsTracing)
+	if (IsRunningCommandlet() || IsTracing())
 	{
 		return false;
 	}
@@ -184,21 +183,21 @@ bool FStateTreeModule::StartTraces(int32& OutTraceId)
 	if (const UE::Trace::FStoreClient* Client = GetStoreClient())
 	{
 		const UE::Trace::FStoreClient::FSessionInfo* SessionInfo = Client->GetSessionInfoByGuid(TraceGuid);
-		// Note that 0 is returned instead of INDEX_NONE to match default invalid value for GetTraceId 
-		OutTraceId = SessionInfo != nullptr ? SessionInfo->GetTraceId(): 0;
+		// Note that 0 is returned instead of INDEX_NONE to match default invalid value for GetTraceId
+		OutTraceId = SessionInfo != nullptr ? SessionInfo->GetTraceId() : 0;
 	}
-#endif // WITH_STATETREE_TRACE_DEBUGGER 
+#endif // WITH_STATETREE_TRACE_DEBUGGER
 
 	// If trace is already connected let's keep track of enabled channels to restore them when we stop recording
 	if (bAlreadyConnected)
 	{
 		UE::Trace::EnumerateChannels([](const ANSICHAR* Name, const bool bIsEnabled, void* Channels)
 		{
-			TArray<FString>* EnabledChannels = static_cast<TArray<FString>*>(Channels); 
+			TArray<FString>* EnabledChannels = static_cast<TArray<FString>*>(Channels);
 			if (bIsEnabled)
 			{
 				EnabledChannels->Emplace(ANSI_TO_TCHAR(Name));
-			}		
+			}
 		}, &ChannelsToRestore);
 	}
 	else
@@ -226,7 +225,6 @@ bool FStateTreeModule::StartTraces(int32& OutTraceId)
 		bAreTracesStarted = FTraceAuxiliary::Start(FTraceAuxiliary::EConnectionType::Network, TEXT("localhost"), TEXT(""), &Options, LogStateTree);
 	}
 
-	bIsTracing = true;
 	if (UE::StateTree::Delegates::OnTracingStateChanged.IsBound())
 	{
 		UE_LOG(LogStateTree, Log, TEXT("StateTree traces enabled"));
@@ -241,17 +239,24 @@ bool FStateTreeModule::StartTraces(int32& OutTraceId)
 
 bool FStateTreeModule::IsTracing() const
 {
-#if WITH_STATETREE_TRACE
-	return bIsTracing;
+#if UE_TRACE_ENABLED
+	// We are not relying on a dedicated flag since tracing can be started from many sources (e.g., RewindDebugger)
+	if (!UE::Trace::IsTracing())
+	{
+		return false;
+	}
+
+	const UE::Trace::FChannel* Channel = UE::Trace::FindChannel(TEXT("StateTreeDebugChannel"));
+	return Channel != nullptr && Channel->IsEnabled();
 #else
 	return false;
-#endif // WITH_STATETREE_TRACE
+#endif // UE_TRACE_ENABLED
 }
 
 void FStateTreeModule::StopTraces()
 {
 #if WITH_STATETREE_TRACE
-	if (bIsTracing == false)
+	if (IsTracing() == false)
 	{
 		return;
 	}
@@ -280,8 +285,6 @@ void FStateTreeModule::StopTraces()
 		FTraceAuxiliary::Stop();
 	}
 
-	bIsTracing = false;
-
 	if (UE::StateTree::Delegates::OnTracingStateChanged.IsBound())
 	{
 		UE_LOG(LogStateTree, Log, TEXT("StateTree traces stopped"));
diff --git a/Engine/Plugins/Runtime/StateTree/Source/StateTreeModule/Private/StateTreeModuleImpl.h b/Engine/Plugins/Runtime/StateTree/Source/StateTreeModule/Private/StateTreeModuleImpl.h
--- a/Engine/Plugins/Runtime/StateTree/Source/StateTreeModule/Private/StateTreeModuleImpl.h
+++ b/Engine/Plugins/Runtime/StateTree/Source/StateTreeModule/Private/StateTreeModuleImpl.h
@@ -61,11 +61,7 @@ private:
 #if WITH_STATETREE_TRACE
 	TArray<const FString> ChannelsToRestore;
 
-	/** Keep track if StartTraces was explicitly called. */
-	bool bIsTracing = false;
-
 	FAutoConsoleCommand StartDebuggerTracesCommand;
-
 	FAutoConsoleCommand StopDebuggerTracesCommand;
 #endif // WITH_STATETREE_TRACE
 
diff --git a/Engine/Plugins/Runtime/StateTree/Source/StateTreeModule/Public/Debugger/StateTreeDebugger.h b/Engine/Plugins/Runtime/StateTree/Source/StateTreeModule/Public/Debugger/StateTreeDebugger.h
--- a/Engine/Plugins/Runtime/StateTree/Source/StateTreeModule/Public/Debugger/StateTreeDebugger.h
+++ b/Engine/Plugins/Runtime/StateTree/Source/StateTreeModule/Public/Debugger/StateTreeDebugger.h
@@ -519,6 +519,11 @@ private:
 	 */
 	FDelegateHandle TracingStateChangedHandle;
 
+	/**
+	 * Delegate Handle bound to UE::StateTree::Delegates::OnTraceAnalysisStateChanged
+	 */
+	FDelegateHandle TraceAnalysisStateChangedHandle;
+
 	/**
 	 * Delegate Handle bound to UE::StateTree::Delegates::OnTracingTimelineScrubbed
 	 */
diff --git a/Engine/Plugins/Runtime/StateTree/Source/StateTreeModule/Public/StateTreeExecutionTypes.h b/Engine/Plugins/Runtime/StateTree/Source/StateTreeModule/Public/StateTreeExecutionTypes.h
--- a/Engine/Plugins/Runtime/StateTree/Source/StateTreeModule/Public/StateTreeExecutionTypes.h
+++ b/Engine/Plugins/Runtime/StateTree/Source/StateTreeModule/Public/StateTreeExecutionTypes.h
@@ -694,7 +694,7 @@ struct FStateTreeInstanceDebugId
 
 	friend FString LexToString(const FStateTreeInstanceDebugId InstanceDebugId)
 	{
-		return FString::Printf(TEXT("0x%x | %d"), InstanceDebugId.Id, InstanceDebugId.SerialNumber);
+		return FString::Printf(TEXT("0x%llx"), InstanceDebugId.ToUint64());
 	}
 
 	static UE_API const FStateTreeInstanceDebugId Invalid;
diff --git a/Engine/Source/Editor/RewindDebuggerInterface/Public/IRewindDebugger.h b/Engine/Source/Editor/RewindDebuggerInterface/Public/IRewindDebugger.h
--- a/Engine/Source/Editor/RewindDebuggerInterface/Public/IRewindDebugger.h
+++ b/Engine/Source/Editor/RewindDebuggerInterface/Public/IRewindDebugger.h
@@ -6,6 +6,7 @@
 #endif // UE_ENABLE_INCLUDE_ORDER_DEPRECATED_IN_5_7
 #include "RewindDebuggerTypes.h"
 #include "TraceServices/Model/Frames.h"
+#include "UObject/NameTypes.h"
 #include "IRewindDebugger.generated.h"
 
 REWINDDEBUGGERINTERFACE_API DECLARE_LOG_CATEGORY_EXTERN(LogRewindDebugger, Log, All);
@@ -71,6 +72,15 @@ public:
 
 	REWINDDEBUGGERINTERFACE_API virtual ~IRewindDebugger();
 
+	/** @return the name of the main RewindDebugger window that could be used to invoke the associated tab in the Editor. */
+	virtual FName GetMainTabName() const = 0;
+
+	/** @return the name of the RewindDebugger details window that could be used to invoke the associated tab in the Editor. */
+	virtual FName GetDetailsTabName() const = 0;
+
+	/** Start traces recording if not already started. */
+	virtual void StartRecording() const = 0;
+
 	/** @return the time the debugger is scrubbed to, in seconds since the capture started (or the recording duration while the game is running) */
 	virtual double CurrentTraceTime() const = 0;
 
@@ -98,6 +108,13 @@ public:
 		return GetDebuggedObjects();
 	}
 
+	/**
+	 * Sets the object to display tracks for.
+	 * Method will do nothing uf the specified object is already part of the current object hierarchy.
+	 * @param ObjectId The Id of the object to debug
+	 */
+	virtual void SetObjectToDebug(RewindDebugger::FObjectId ObjectId) = 0;
+
 	/** @return Whether the given object id is one of the object currently debugged, or one of its children. */
 	virtual bool IsObjectCurrentlyDebugged(uint64 InObjectId) const = 0;
 
@@ -113,6 +130,9 @@ public:
 	/** @return the currently selected track */
 	virtual TSharedPtr<RewindDebugger::FRewindDebuggerTrack> GetSelectedTrack() const = 0;
 
+	/** Try selecting the track associated to a given object */
+	virtual void SelectTrack(RewindDebugger::FObjectId ObjectId) = 0;
+
 	/** @return position of the selected target actor (returns true if position is valid) */
 	virtual bool GetTargetActorPosition(FVector& OutPosition) const = 0;
 
