[StateTree] Add unit test for FStateTreeStructRef. It should be deprecated. It can easily point to invalid data if the binding is not executed directly before the reference is used.

diff --git a/Engine/Plugins/Runtime/StateTree/Source/StateTreeTestSuite/Private/StateTreeBindingTest.cpp b/Engine/Plugins/Runtime/StateTree/Source/StateTreeTestSuite/Private/StateTreeBindingTest.cpp
--- a/Engine/Plugins/Runtime/StateTree/Source/StateTreeTestSuite/Private/StateTreeBindingTest.cpp
+++ b/Engine/Plugins/Runtime/StateTree/Source/StateTreeTestSuite/Private/StateTreeBindingTest.cpp
@@ -4,9 +4,11 @@
 #include "StateTreeTestBase.h"
 #include "StateTreeTestTypes.h"
 
+#include "StateTreeCompiler.h"
 #include "StateTreeCompilerLog.h"
+#include "StateTreeExecutionContext.h"
 #include "StateTreeEditorData.h"
-#include "StateTreeCompiler.h"
+#include "StateTreeInstanceData.h"
 #include "Conditions/StateTreeCommonConditions.h"
 
 #define LOCTEXT_NAMESPACE "AITestSuite_StateTreeTest"
@@ -757,6 +759,159 @@ struct FStateTreeTest_TransitionTaskWithBinding : FStateTreeTestBase
 };
 IMPLEMENT_STATE_TREE_INSTANT_TEST(FStateTreeTest_TransitionTaskWithBinding, "System.StateTree.Binding.TransitionTaskWithBinding");
 
+struct FStateTreeTest_BindingStructRef : FStateTreeTestBase
+{
+	virtual bool InstantTest() override
+	{
+		//	Root (global variable, )
+		//		State1 (ref to root variable)
+		//			State2
+		//			State3 (with a lot of instance data to realloc the internal buffer)
+		//				State4 (with a lot of instance data to realloc the internal buffer)
+
+		const FStateTreeTest_PropertyStruct* TaskA_State_PropertyStructPtr = nullptr;
+		const FStateTreeTest_PropertyStruct* TaskB_State_PropertyStructPtr = nullptr;
+		const FStateTreeTest_PropertyStruct* TaskA_Task_PropertyStructPtr = nullptr;
+		const FStateTreeTest_PropertyStruct* TaskB_Task_PropertyStructPtr = nullptr;
+
+		FStateTreeCompilerLog Log;
+		FStateTreePropertyBindings Bindings;
+		FStateTreePropertyBindingCompiler BindingCompiler;
+
+		UStateTree& StateTree = NewStateTree();
+		{
+			UStateTreeEditorData& EditorData = *Cast<UStateTreeEditorData>(StateTree.EditorData);
+
+			UStateTreeState& Root = EditorData.AddSubTree("Tree1StateRoot");
+			UStateTreeState& State1 = Root.AddChildState("Tree1State1", EStateTreeStateType::State);
+			UStateTreeState& State2 = State1.AddChildState("Tree1State2", EStateTreeStateType::State);
+			UStateTreeState& State3 = State1.AddChildState("Tree1State3", EStateTreeStateType::State);
+			UStateTreeState& State4 = State3.AddChildState("Tree1State4", EStateTreeStateType::State);
+			{
+				Root.Parameters.Parameters.AddProperty("RootParam", EPropertyBagPropertyType::Struct, FStateTreeTest_PropertyStruct::StaticStruct());
+				FStateTreeTest_PropertyStruct PropertyStruct;
+				PropertyStruct.A = 111;
+				PropertyStruct.B = 222;
+				PropertyStruct.StructB.B = 333;
+				Root.Parameters.Parameters.SetValueStruct("RootParam", PropertyStruct);
+			}
+			{
+				FPropertyBagPropertyDesc Desc = FPropertyBagPropertyDesc("StateParam", EPropertyBagPropertyType::Struct, FStateTreeStructRef::StaticStruct());
+				Desc.MetaData.Emplace(FName("BaseStruct"), TEXT("/Script/StateTreeTestSuite.StateTreeTest_PropertyStruct"));
+				State1.Parameters.Parameters.AddProperties({ Desc }, /*bOverwrite*/false);
+
+				EditorData.AddPropertyBinding(FPropertyBindingPath(Root.Parameters.ID, TEXT("RootParam")), FPropertyBindingPath(State1.Parameters.ID, TEXT("StateParam")));
+			}
+			{
+				TStateTreeEditorNode<FTestTask_PrintValue_StructRef_NoBindingUpdate>& TaskA = State1.AddTask<FTestTask_PrintValue_StructRef_NoBindingUpdate>("Tree1State1TaskA");
+				TaskA.GetNode().CustomTickFunc = [&TaskA_State_PropertyStructPtr, &TaskA_Task_PropertyStructPtr](FStateTreeExecutionContext& Context, const FTestTask_PrintValue* Task)
+					{
+						TaskA_State_PropertyStructPtr = nullptr;
+						TaskA_Task_PropertyStructPtr = nullptr;
+
+						FConstStructView RootParameters = Context.GetInstanceData()->GetStorage().GetStruct(0);
+						if (const FCompactStateTreeParameters* Parameters = RootParameters.GetPtr<const FCompactStateTreeParameters>())
+						{
+							TValueOrError<FStateTreeTest_PropertyStruct*, EPropertyBagResult> PropertyStruct = Parameters->Parameters.GetValueStruct<FStateTreeTest_PropertyStruct>("RootParam");
+							if (PropertyStruct.HasValue())
+							{
+								TaskA_State_PropertyStructPtr = PropertyStruct.GetValue();
+							}
+						}
+
+						FTestTask_PrintValue_StructRef_NoBindingUpdate::FInstanceDataType& InstanceData = Context.GetInstanceData(*static_cast<const FTestTask_PrintValue_StructRef_NoBindingUpdate*>(Task));
+						TaskA_Task_PropertyStructPtr = &InstanceData.PropertyStruct;
+					};
+				TaskA.GetInstanceData().PropertyStruct.A = 11;
+				TaskA.GetInstanceData().PropertyStruct.B = 22;
+				TaskA.GetInstanceData().PropertyStruct.StructB.B = 33;
+
+				TStateTreeEditorNode<FTestTask_PrintValue_StructRef_NoBindingUpdate>& TaskB = State1.AddTask<FTestTask_PrintValue_StructRef_NoBindingUpdate>("Tree1State1TaskB");
+				TaskB.GetNode().CustomTickFunc = [&TaskB_State_PropertyStructPtr, &TaskB_Task_PropertyStructPtr](FStateTreeExecutionContext& Context, const FTestTask_PrintValue* Task)
+					{
+						TaskB_State_PropertyStructPtr = nullptr;
+
+						FStateTreeTestLog& TestLog = Context.GetExternalData(Task->LogHandle);
+
+						FConstStructView State1Parameters = Context.GetInstanceData()->GetStorage().GetStruct(1);
+						if (const FCompactStateTreeParameters* Parameters = State1Parameters.GetPtr<const FCompactStateTreeParameters>())
+						{
+							TValueOrError<FStateTreeStructRef*, EPropertyBagResult> PropertyStruct = Parameters->Parameters.GetValueStruct<FStateTreeStructRef>("StateParam");
+							if (PropertyStruct.HasValue())
+							{
+								TaskB_State_PropertyStructPtr = PropertyStruct.GetValue()->GetPtr<const FStateTreeTest_PropertyStruct>();
+							}
+						}
+
+						FTestTask_PrintValue_StructRef_NoBindingUpdate::FInstanceDataType& InstanceData = Context.GetInstanceData(*static_cast<const FTestTask_PrintValue_StructRef_NoBindingUpdate*>(Task));
+						TaskB_Task_PropertyStructPtr = InstanceData.StructRef.GetPtr<const FStateTreeTest_PropertyStruct>();
+					};
+				EditorData.AddPropertyBinding(FPropertyBindingPath(TaskA.ID, TEXT("PropertyStruct")), FPropertyBindingPath(TaskB.ID, TEXT("StructRef")));
+			}
+			{
+				State2.SelectionBehavior = EStateTreeStateSelectionBehavior::TryEnterState;
+				State2.AddTransition(EStateTreeTransitionTrigger::OnTick, EStateTreeTransitionType::GotoState, &State3);
+			}
+			{
+				for (int32 Index = 0; Index < 32; ++Index)
+				{
+					State3.AddTask<FTestTask_PrintValue>("Tree1State3TaskX");
+				}
+			}
+			{
+				for (int32 Index = 0; Index < 32; ++Index)
+				{
+					State4.AddTask<FTestTask_PrintValue>("Tree1State4TaskX");
+				}
+			}
+
+		}
+		{
+			FStateTreeCompiler Compiler(Log);
+			const bool bResult = Compiler.Compile(StateTree);
+			AITEST_TRUE(TEXT("StateTree2 should get compiled"), bResult);
+		}
+		{
+			FStateTreeInstanceData InstanceData;
+			{
+				FTestStateTreeExecutionContext Exec(StateTree, StateTree, InstanceData);
+				const bool bInitSucceeded = Exec.IsValid();
+				AITEST_TRUE(TEXT("StateTree should init"), bInitSucceeded);
+			}
+			{
+				FTestStateTreeExecutionContext Exec(StateTree, StateTree, InstanceData);
+				EStateTreeRunStatus Status = Exec.Start();
+				AITEST_EQUAL(TEXT("Start should complete with Running"), Status, EStateTreeRunStatus::Running);
+				AITEST_TRUE(TEXT("In correct states"), Exec.ExpectInActiveStates("Tree1StateRoot", "Tree1State1", "Tree1State2"));
+			}
+			{
+				FTestStateTreeExecutionContext Exec(StateTree, StateTree, InstanceData);
+				EStateTreeRunStatus Status = Exec.Tick(0.1f);
+				AITEST_EQUAL(TEXT("Start should complete with Running"), Status, EStateTreeRunStatus::Running);
+				AITEST_TRUE(TEXT("In correct states"), Exec.ExpectInActiveStates("Tree1StateRoot", "Tree1State1", "Tree1State3", "Tree1State4"));
+				AITEST_TRUE(TEXT("Both task should point to the same StructRef"), TaskA_State_PropertyStructPtr != nullptr && TaskA_State_PropertyStructPtr == TaskB_State_PropertyStructPtr);
+				AITEST_TRUE(TEXT("Both task should point to the same StructRef"), TaskA_Task_PropertyStructPtr != nullptr && TaskA_Task_PropertyStructPtr == TaskB_Task_PropertyStructPtr);
+			}
+			{
+				FTestStateTreeExecutionContext Exec(StateTree, StateTree, InstanceData);
+				EStateTreeRunStatus Status = Exec.Tick(0.1f);
+				AITEST_EQUAL(TEXT("Start should complete with Running"), Status, EStateTreeRunStatus::Running);
+				AITEST_TRUE(TEXT("In correct states"), Exec.ExpectInActiveStates("Tree1StateRoot", "Tree1State1", "Tree1State3", "Tree1State4"));
+				AITEST_TRUE(TEXT("Both task should point to the same StructRef"), TaskA_State_PropertyStructPtr != nullptr && TaskA_State_PropertyStructPtr == TaskB_State_PropertyStructPtr);
+				//@TODO: Deprecate FStateTreeStructRef. If the binding is not updated before it is used it can access invalid data. Here the array of instance data grows.
+				AITEST_FALSE(TEXT("Both task should point to the same StructRef"), TaskA_Task_PropertyStructPtr != nullptr && TaskA_Task_PropertyStructPtr == TaskB_Task_PropertyStructPtr);
+			}
+			{
+				FTestStateTreeExecutionContext Exec(StateTree, StateTree, InstanceData);
+				Exec.Stop();
+			}
+		}
+
+		return true;
+	}
+};
+IMPLEMENT_STATE_TREE_INSTANT_TEST(FStateTreeTest_BindingStructRef, "System.StateTree.Binding.StructRef");
+
 } // namespace UE::StateTree::Tests
 
 UE_ENABLE_OPTIMIZATION_SHIP
diff --git a/Engine/Plugins/Runtime/StateTree/Source/StateTreeTestSuite/Private/StateTreeTestTypes.h b/Engine/Plugins/Runtime/StateTree/Source/StateTreeTestSuite/Private/StateTreeTestTypes.h
--- a/Engine/Plugins/Runtime/StateTree/Source/StateTreeTestSuite/Private/StateTreeTestTypes.h
+++ b/Engine/Plugins/Runtime/StateTree/Source/StateTreeTestSuite/Private/StateTreeTestTypes.h
@@ -415,7 +415,6 @@ struct FTestTask_PrintValue_TransitionNoTick : public FTestTask_PrintValue
 	}
 };
 
-
 USTRUCT()
 struct FTestTask_StopTreeInstanceData
 {
@@ -835,6 +834,41 @@ struct FTestPropertyFunction : public FStateTreePropertyFunctionBase
 	}
 };
 
+USTRUCT()
+struct FTestTask_PrintValue_StructRef_NoBindingUpdateInstanceData : public FTestTask_PrintValueInstanceData
+{
+	GENERATED_BODY()
+
+	UPROPERTY(EditAnywhere, Category = "Parameter")
+	FStateTreeTest_PropertyStruct PropertyStruct;
+
+	UPROPERTY(EditAnywhere, Category = "Parameter", meta = (BaseStruct = "/Script/StateTreeTestSuite.StateTreeTest_PropertyStruct"))
+	FStateTreeStructRef StructRef;
+};
+
+USTRUCT()
+struct FTestTask_PrintValue_StructRef_NoBindingUpdate : public FTestTask_PrintValue
+{
+	GENERATED_BODY()
+
+	using FInstanceDataType = FTestTask_PrintValue_StructRef_NoBindingUpdateInstanceData;
+
+	virtual const UStruct* GetInstanceDataType() const override
+	{
+		return FInstanceDataType::StaticStruct();
+	}
+
+	FTestTask_PrintValue_StructRef_NoBindingUpdate()
+	{
+		bShouldCopyBoundPropertiesOnTick = false;
+	}
+	FTestTask_PrintValue_StructRef_NoBindingUpdate(const FName InName)
+		: FTestTask_PrintValue(InName)
+	{
+		bShouldCopyBoundPropertiesOnTick = false;
+	}
+};
+
 USTRUCT()
 struct FStateTreeTestBooleanConditionInstanceData
 {
