[StateTree] Fix the deprecation code path and bad indentation in ExitState introduced by the transition ID changes.
#rnx

diff --git a/Engine/Plugins/Runtime/StateTree/Source/StateTreeModule/Private/StateTreeExecutionContext.cpp b/Engine/Plugins/Runtime/StateTree/Source/StateTreeModule/Private/StateTreeExecutionContext.cpp
--- a/Engine/Plugins/Runtime/StateTree/Source/StateTreeModule/Private/StateTreeExecutionContext.cpp
+++ b/Engine/Plugins/Runtime/StateTree/Source/StateTreeModule/Private/StateTreeExecutionContext.cpp
@@ -1436,7 +1436,6 @@ EStateTreeRunStatus FStateTreeExecutionContext::Start(FStartParameters Parameter
 
 		// No active states or global tasks anymore, reset frames.
 		Exec.ActiveFrames.Reset();
-
 		RemoveAllDelegateListeners();
 
 		// We are not considered as running yet so we only set the status without requiring a stop.
@@ -1787,7 +1786,6 @@ void FStateTreeExecutionContext::TickTriggerTransitionsInternal()
 
 					// No active states or global tasks anymore, reset frames.
 					Exec.ActiveFrames.Reset();
-
 					RemoveAllDelegateListeners();
 				}
 
@@ -2118,7 +2116,7 @@ void FStateTreeExecutionContext::UpdateInstanceData(FSelectStateResult& Args)
 			{
 				CurrentFrameID = SelectedState.GetFrameID();
 
-				const FStateTreeExecutionFrame* Frame = FindExecutionFrame(CurrentFrameID, MakeArrayView(Exec.ActiveFrames), MakeArrayView(Args.TemporaryFrames));
+				const FStateTreeExecutionFrame* Frame = FindExecutionFrame(CurrentFrameID, MakeConstArrayView(Exec.ActiveFrames), MakeConstArrayView(Args.TemporaryFrames));
 				if (ensure(Frame) && Frame->bIsGlobalFrame)
 				{
 					check(Frame->StateTree);
@@ -3930,7 +3928,6 @@ EStateTreeRunStatus FStateTreeExecutionContext::StartGlobalsForFrameInternal(con
 
 	// Start Global tasks
 	// Even if we call Enter/ExitState() on global tasks, they do not enter any specific state.
-
 	for (int32 GlobalTaskIndex = 0; GlobalTaskIndex < CurrentStateTree->GlobalTasksNum; ++GlobalTaskIndex)
 	{
 		const int32 AssetTaskIndex = CurrentStateTree->GlobalTasksBegin + GlobalTaskIndex;
@@ -5816,10 +5813,11 @@ FRecordedStateTreeTransitionResult FStateTreeExecutionContext::MakeRecordedTrans
 	return Recorded;
 }
 
+PRAGMA_DISABLE_DEPRECATION_WARNINGS
 // Deprecated
 bool FStateTreeExecutionContext::SelectState(const FStateTreeExecutionFrame& CurrentFrame,
 	const FStateTreeStateHandle NextState,
-	FSelectStateResult& OutSelectionResult,
+	FStateSelectionResult& OutSelectionResult,
 	const FStateTreeSharedEvent* TransitionEvent,
 	const EStateTreeSelectionFallback Fallback)
 {
@@ -5836,8 +5834,11 @@ bool FStateTreeExecutionContext::SelectState(const FStateTreeExecutionFrame& Cur
 	SelectStateArgs.TargetState = FStateHandleContext(CurrentFrame.StateTree, NextState);
 	SelectStateArgs.TransitionEvent = TransitionEvent != nullptr ? *TransitionEvent : FStateTreeSharedEvent();
 	SelectStateArgs.Fallback = Fallback;
-	return SelectState(SelectStateArgs, OutSelectionResult);
+
+	FSelectStateResult StateSelection;
+	return SelectState(SelectStateArgs, StateSelection);
 }
+PRAGMA_ENABLE_DEPRECATION_WARNINGS
 
 bool FStateTreeExecutionContext::SelectState(const FSelectStateArguments& Args, FSelectStateResult& OutSelectionResult)
 {
@@ -6129,17 +6130,19 @@ bool FStateTreeExecutionContext::SelectState(const FSelectStateArguments& Args,
 	return false;
 }
 
+PRAGMA_DISABLE_DEPRECATION_WARNINGS
 // Deprecated
 bool FStateTreeExecutionContext::SelectStateInternal(
 	const FStateTreeExecutionFrame* CurrentParentFrame,
 	FStateTreeExecutionFrame& CurrentFrame,
 	const FStateTreeExecutionFrame* CurrentFrameInActiveFrames,
 	TConstArrayView<FStateTreeStateHandle> PathToNextState,
-	FSelectStateResult& OutSelectionResult,
+	FStateSelectionResult& OutSelectionResult,
 	const FStateTreeSharedEvent* TransitionEvent)
 {
 	return false;
 }
+PRAGMA_ENABLE_DEPRECATION_WARNINGS
 
 bool FStateTreeExecutionContext::SelectStateInternal(
 	const FSelectStateArguments& Args,
diff --git a/Engine/Plugins/Runtime/StateTree/Source/StateTreeModule/Public/StateTreeExecutionContext.h b/Engine/Plugins/Runtime/StateTree/Source/StateTreeModule/Public/StateTreeExecutionContext.h
--- a/Engine/Plugins/Runtime/StateTree/Source/StateTreeModule/Public/StateTreeExecutionContext.h
+++ b/Engine/Plugins/Runtime/StateTree/Source/StateTreeModule/Public/StateTreeExecutionContext.h
@@ -763,7 +763,6 @@ PRAGMA_ENABLE_DEPRECATION_WARNINGS
 	}
 
 protected:
-
 	/** Event used during state selection. Identified by the ID of the active state that consumed it. */
 	struct FSelectionEventWithID
 	{
@@ -1173,8 +1172,10 @@ PRAGMA_ENABLE_DEPRECATION_WARNINGS
 	/** Forces transition to a state. It will skip all conditions. */
 	UE_API bool ForceTransitionInternal(const TArrayView<const UE::StateTree::ExecutionContext::FStateHandleContext> States, FSelectStateResult& OutSelectionResult);
 
+PRAGMA_DISABLE_DEPRECATION_WARNINGS
 	UE_DEPRECATED(5.7, "Use the SelectState with the FSelectStateArguments argument.")
-	UE_API bool SelectState(const FStateTreeExecutionFrame& CurrentFrame, const FStateTreeStateHandle NextState, FSelectStateResult& OutSelectionResult, const FStateTreeSharedEvent* TransitionEvent = nullptr, const EStateTreeSelectionFallback Fallback = EStateTreeSelectionFallback::None);
+	UE_API bool SelectState(const FStateTreeExecutionFrame& CurrentFrame, const FStateTreeStateHandle NextState, FStateSelectionResult& OutSelectionResult, const FStateTreeSharedEvent* TransitionEvent = nullptr, const EStateTreeSelectionFallback Fallback = EStateTreeSelectionFallback::None);
+PRAGMA_ENABLE_DEPRECATION_WARNINGS
 
 	/** The different type of selections. */
 	enum class ESelectStateBehavior : uint8
@@ -1212,14 +1213,10 @@ PRAGMA_ENABLE_DEPRECATION_WARNINGS
 	 */
 	UE_API bool SelectState(const FSelectStateArguments& SelectStateArgs, FSelectStateResult& OutSelectionResult);
 
+PRAGMA_DISABLE_DEPRECATION_WARNINGS
 	UE_DEPRECATED(5.7, "Use the SelectStateInternal with FSelectStateArguments argument.")
-	UE_API bool SelectStateInternal(
-		const FStateTreeExecutionFrame* CurrentParentFrame,
-		FStateTreeExecutionFrame& CurrentFrame,
-		const FStateTreeExecutionFrame* CurrentFrameInActiveFrames,
-		TConstArrayView<FStateTreeStateHandle> PathToNextState,
-		FSelectStateResult& OutSelectionResult,
-		const FStateTreeSharedEvent* TransitionEvent = nullptr);
+	UE_API bool SelectStateInternal(const FStateTreeExecutionFrame* CurrentParentFrame, FStateTreeExecutionFrame& CurrentFrame, const FStateTreeExecutionFrame* CurrentFrameInActiveFrames, TConstArrayView<FStateTreeStateHandle> PathToNextState, FStateSelectionResult& OutSelectionResult, const FStateTreeSharedEvent* TransitionEvent = nullptr);
+PRAGMA_ENABLE_DEPRECATION_WARNINGS
 
 	struct FSelectStateInternalArguments
 	{
