[StateTree] Use NonconcurrentLinearArrayAllocator for scoped allocation in execution context transition code.

#rnx

diff --git a/Engine/Plugins/Runtime/StateTree/Source/StateTreeModule/Private/StateTreeExecutionContext.cpp b/Engine/Plugins/Runtime/StateTree/Source/StateTreeModule/Private/StateTreeExecutionContext.cpp
--- a/Engine/Plugins/Runtime/StateTree/Source/StateTreeModule/Private/StateTreeExecutionContext.cpp
+++ b/Engine/Plugins/Runtime/StateTree/Source/StateTreeModule/Private/StateTreeExecutionContext.cpp
@@ -7105,7 +7105,8 @@ bool FStateTreeExecutionContext::SelectStateInternal_TrySelectChildrenAtRandom(
 
 	UE_STATETREE_DEBUG_SCOPED_STATE_PHASE(this, NextStateHandle, EStateTreeUpdatePhase::TrySelectBehavior);
 
-	TArray<uint16, TInlineAllocator<16>> NextLevelChildStates;
+	TArray<uint16, TInlineAllocator<16, FNonconcurrentLinearArrayAllocator>> NextLevelChildStates;
+	NextLevelChildStates.Reserve(NextState.ChildrenEnd - NextState.ChildrenBegin);
 	for (uint16 ChildStateIndex = NextState.ChildrenBegin; ChildStateIndex < NextState.ChildrenEnd; ChildStateIndex = NextStateTree->States[ChildStateIndex].GetNextSibling())
 	{
 		NextLevelChildStates.Push(ChildStateIndex);
@@ -7150,7 +7151,8 @@ bool FStateTreeExecutionContext::SelectStateInternal_TrySelectChildrenWithHighes
 
 	UE_STATETREE_DEBUG_SCOPED_STATE_PHASE(this, NextStateHandle, EStateTreeUpdatePhase::TrySelectBehavior);
 
-	TArray<TPair<uint16, float>, TInlineAllocator<16>> NextLevelChildStates;
+	TArray<TPair<uint16, float>, TInlineAllocator<16, FNonconcurrentLinearArrayAllocator>> NextLevelChildStates;
+	NextLevelChildStates.Reserve(NextState.ChildrenEnd - NextState.ChildrenBegin);
 	{
 		using namespace UE::StateTree::ExecutionContext::Private;
 		const FStateTreeExecutionFrame* NextParentFrame = OutSelectionResult.SelectedFrames.Num() > 1
@@ -7222,7 +7224,9 @@ bool FStateTreeExecutionContext::SelectStateInternal_TrySelectChildrenAtRandomWe
 
 	UE_STATETREE_DEBUG_SCOPED_STATE_PHASE(this, NextStateHandle, EStateTreeUpdatePhase::TrySelectBehavior);
 
-	TArray<TTuple<uint16, float>, TInlineAllocator<16>> NextLevelChildStates;
+	TArray<TTuple<uint16, float>, TInlineAllocator<16, FNonconcurrentLinearArrayAllocator>> NextLevelChildStates;
+	NextLevelChildStates.Reserve(NextState.ChildrenEnd - NextState.ChildrenBegin);
+
 	float TotalScore = 0.0f;
 	{
 		using namespace UE::StateTree::ExecutionContext::Private;
