[StateTree] Fix bug with no-op transition. Introduced with CL 43425491. No-op transitions are there to block other transitions in the hierarchy.

diff --git a/Engine/Plugins/Runtime/StateTree/Source/StateTreeModule/Private/StateTreeExecutionContext.cpp b/Engine/Plugins/Runtime/StateTree/Source/StateTreeModule/Private/StateTreeExecutionContext.cpp
--- a/Engine/Plugins/Runtime/StateTree/Source/StateTreeModule/Private/StateTreeExecutionContext.cpp
+++ b/Engine/Plugins/Runtime/StateTree/Source/StateTreeModule/Private/StateTreeExecutionContext.cpp
@@ -1764,7 +1764,6 @@ void FStateTreeExecutionContext::TickTriggerTransitionsInternal()
 	{
 		ON_SCOPE_EXIT{ InstanceData.ResetTemporaryInstances(); };
 
-		// Trigger conditional transitions or state succeed/failed transitions. First tick transition is handled here too.
 		if (TriggerTransitions())
 		{
 			check(RequestedTransition.IsValid());
@@ -5627,6 +5626,12 @@ PRAGMA_ENABLE_DEPRECATION_WARNINGS
 		}
 	}
 
+	// Request can be no-op, used for blocking other transitions.
+	if (RequestedTransition.IsValid() && !RequestedTransition->Transition.TargetState.IsValid())
+	{
+		RequestedTransition.Reset();
+	}
+
 	return RequestedTransition.IsValid();
 }
 
diff --git a/Engine/Plugins/Runtime/StateTree/Source/StateTreeTestSuite/Private/StateTreeTransitionTest.cpp b/Engine/Plugins/Runtime/StateTree/Source/StateTreeTestSuite/Private/StateTreeTransitionTest.cpp
--- a/Engine/Plugins/Runtime/StateTree/Source/StateTreeTestSuite/Private/StateTreeTransitionTest.cpp
+++ b/Engine/Plugins/Runtime/StateTree/Source/StateTreeTestSuite/Private/StateTreeTransitionTest.cpp
@@ -1693,6 +1693,62 @@ struct FStateTreeTest_RequestTransition : FStateTreeTestBase
 };
 IMPLEMENT_AI_INSTANT_TEST(FStateTreeTest_RequestTransition, "System.StateTree.Transition.RequestTransition");
 
+struct FStateTreeTest_TransitionToNone : FStateTreeTestBase
+{
+	virtual bool InstantTest() override
+	{
+		const FGameplayTag Tag = GetTestTag1();
+
+		UStateTree& StateTree = NewStateTree();
+		UStateTreeEditorData& EditorData = *Cast<UStateTreeEditorData>(StateTree.EditorData);
+
+		UStateTreeState& Root = EditorData.AddSubTree(FName(TEXT("Root")));
+		UStateTreeState& State1 = Root.AddChildState(FName(TEXT("State1")));
+		UStateTreeState& State2 = Root.AddChildState(FName(TEXT("State2")));
+
+		auto& TransitionRoot = Root.AddTransition(EStateTreeTransitionTrigger::OnEvent, EStateTreeTransitionType::GotoState);
+		TransitionRoot.State = State1.GetLinkToState();
+		TransitionRoot.RequiredEvent.Tag = Tag;
+
+		auto& Task1 = State1.AddTask<FTestTask_StandNoTick>(FName(TEXT("Task1")));
+		auto& TransitionState1 = State1.AddTransition(EStateTreeTransitionTrigger::OnEvent, EStateTreeTransitionType::None);
+		TransitionState1.RequiredEvent.Tag = Tag;
+
+		FStateTreeCompilerLog Log;
+		FStateTreeCompiler Compiler(Log);
+		const bool bResult = Compiler.Compile(StateTree);
+		AITEST_TRUE("StateTree should get compiled", bResult);
+
+		EStateTreeRunStatus Status = EStateTreeRunStatus::Unset;
+		FStateTreeInstanceData InstanceData;
+		FTestStateTreeExecutionContext Exec(StateTree, StateTree, InstanceData);
+		const bool bInitSucceeded = Exec.IsValid();
+		AITEST_TRUE("StateTree should init", bInitSucceeded);
+
+		const TCHAR* EnterStateStr(TEXT("EnterState"));
+		const TCHAR* ExitStateStr(TEXT("ExitState"));
+
+		// Start and enter state
+		Status = Exec.Start();
+		AITEST_TRUE(TEXT("StateTree Task1 should enter state"), Exec.Expect(Task1.GetName(), EnterStateStr));
+		Exec.LogClear();
+
+		// Send event with Tag
+		Exec.SendEvent(Tag);
+
+		// Transition from Root to State2 should not be triggered
+		Status = Exec.Tick(0.1f);
+		AITEST_TRUE(TEXT("Active states should stay the same"), Exec.ExpectInActiveStates("Root", "State1"));
+		AITEST_FALSE(TEXT("StateTree Task1 should not exit state"), Exec.Expect(Task1.GetName(), ExitStateStr));
+		Exec.LogClear();
+
+		Exec.Stop();
+
+		return true;
+	}
+};
+IMPLEMENT_STATE_TREE_INSTANT_TEST(FStateTreeTest_TransitionToNone, "System.StateTree.Transition.ToNone");
+
 } // namespace UE::StateTree::Tests
 
 UE_ENABLE_OPTIMIZATION_SHIP
