[StateTree] Add TickReason to Scheduled Tick. Use for debugging purposes.

diff --git a/Engine/Plugins/Runtime/GameplayStateTree/Source/GameplayStateTreeModule/Private/Components/StateTreeComponent.cpp b/Engine/Plugins/Runtime/GameplayStateTree/Source/GameplayStateTreeModule/Private/Components/StateTreeComponent.cpp
--- a/Engine/Plugins/Runtime/GameplayStateTree/Source/GameplayStateTreeModule/Private/Components/StateTreeComponent.cpp
+++ b/Engine/Plugins/Runtime/GameplayStateTree/Source/GameplayStateTreeModule/Private/Components/StateTreeComponent.cpp
@@ -27,7 +27,7 @@ namespace UE::GameplayStateTree::Private
 //////////////////////////////////////////////////////////////////////////
 // UStateTreeComponent
 
-void FStateTreeComponentExecutionExtension::ScheduleNextTick(const FContextParameters& Context)
+void FStateTreeComponentExecutionExtension::ScheduleNextTick(const FContextParameters& Context, const FNextTickArguments& Args)
 {
 	if (ensure(Component))
 	{
diff --git a/Engine/Plugins/Runtime/GameplayStateTree/Source/GameplayStateTreeModule/Public/Components/StateTreeComponent.h b/Engine/Plugins/Runtime/GameplayStateTree/Source/GameplayStateTreeModule/Public/Components/StateTreeComponent.h
--- a/Engine/Plugins/Runtime/GameplayStateTree/Source/GameplayStateTreeModule/Public/Components/StateTreeComponent.h
+++ b/Engine/Plugins/Runtime/GameplayStateTree/Source/GameplayStateTreeModule/Public/Components/StateTreeComponent.h
@@ -27,7 +27,7 @@ struct FStateTreeComponentExecutionExtension : public FStateTreeExecutionExtensi
 	GENERATED_BODY()
 
 public:
-	virtual void ScheduleNextTick(const FContextParameters& Context) override;
+	virtual void ScheduleNextTick(const FContextParameters& Context, const FNextTickArguments& Args) override;
 
 	UPROPERTY()
 	TObjectPtr<UStateTreeComponent> Component;
diff --git a/Engine/Plugins/Runtime/StateTree/Source/StateTreeModule/Private/StateTreeAsyncExecutionContext.cpp b/Engine/Plugins/Runtime/StateTree/Source/StateTreeModule/Private/StateTreeAsyncExecutionContext.cpp
--- a/Engine/Plugins/Runtime/StateTree/Source/StateTreeModule/Private/StateTreeAsyncExecutionContext.cpp
+++ b/Engine/Plugins/Runtime/StateTree/Source/StateTreeModule/Private/StateTreeAsyncExecutionContext.cpp
@@ -88,7 +88,7 @@ bool TStateTreeStrongExecutionContext<bWithWriteAccess>::RequestTransition(FStat
 		Request.Fallback = Fallback;
 
 		Storage->AddTransitionRequest(Owner.Get(), Request);
-		ScheduleNextTick(Owner.Get(), StateTree.Get(), *Storage);
+		ScheduleNextTick(Owner.Get(), StateTree.Get(), *Storage, UE::StateTree::ETickReason::TransitionRequest);
 
 		return true;
 	}
@@ -116,7 +116,7 @@ const
 		Exec.DelegateActiveListeners.BroadcastDelegate(Dispatcher, Exec);
 		if (UE::StateTree::ExecutionContext::MarkDelegateAsBroadcasted(Dispatcher, *ActivePath.Frame, *Storage))
 		{
-			ScheduleNextTick(Owner.Get(), StateTree.Get(), *Storage);
+			ScheduleNextTick(Owner.Get(), StateTree.Get(), *Storage, UE::StateTree::ETickReason::Delegate);
 		}
 
 		return true;
@@ -228,7 +228,7 @@ bool TStateTreeStrongExecutionContext<bWithWriteAccess>::FinishTask(EStateTreeFi
 			if (bCompleted)
 			{
 				Exec.bHasPendingCompletedState = true;
-				ScheduleNextTick(Owner.Get(), StateTree.Get(), *Storage);
+				ScheduleNextTick(Owner.Get(), StateTree.Get(), *Storage, UE::StateTree::ETickReason::CompletedState);
 			}
 		}
 	}
@@ -247,7 +247,7 @@ bool TStateTreeStrongExecutionContext<bWithWriteAccess>::UpdateScheduledTickRequ
 		FStateTreeExecutionState& Exec = Storage->GetMutableExecutionState();
 		if (Exec.UpdateScheduledTickRequest(Handle, ScheduledTick))
 		{
-			ScheduleNextTick(Owner.Get(), StateTree.Get(), *Storage);
+			ScheduleNextTick(Owner.Get(), StateTree.Get(), *Storage, UE::StateTree::ETickReason::ScheduledTickRequest);
 		}
 
 		return true;
@@ -303,14 +303,14 @@ FStateTreeDataView TStateTreeStrongExecutionContext<bWithWriteAccess>::GetInstan
 
 template <bool bWithWriteAccess>
 template<bool bWriteAccess UE_REQUIRES_DEFINITION(bWriteAccess)>
-void TStateTreeStrongExecutionContext<bWithWriteAccess>::ScheduleNextTick(TNotNull<UObject*> Owner, TNotNull<const UStateTree*> RootStateTree, FStateTreeInstanceStorage& Storage)
+void TStateTreeStrongExecutionContext<bWithWriteAccess>::ScheduleNextTick(TNotNull<UObject*> Owner, TNotNull<const UStateTree*> RootStateTree, FStateTreeInstanceStorage& Storage, UE::StateTree::ETickReason Reason)
 {
 	static_assert(bWithWriteAccess);
 
 	TInstancedStruct<FStateTreeExecutionExtension>& ExecutionExtension = Storage.GetMutableExecutionState().ExecutionExtension;
 	if (RootStateTree->IsScheduledTickAllowed() && ExecutionExtension.IsValid())
 	{
-		ExecutionExtension.GetMutable().ScheduleNextTick(FStateTreeExecutionExtension::FContextParameters(*Owner, *RootStateTree, Storage));
+		ExecutionExtension.GetMutable().ScheduleNextTick(FStateTreeExecutionExtension::FContextParameters(*Owner, *RootStateTree, Storage), FStateTreeExecutionExtension::FNextTickArguments(Reason));
 	}
 }
 
@@ -325,7 +325,7 @@ template bool STATETREEMODULE_API TStateTreeStrongExecutionContext<true>::BindDe
 template bool STATETREEMODULE_API TStateTreeStrongExecutionContext<true>::UnbindDelegate(const FStateTreeDelegateListener& Listener) const;
 template bool STATETREEMODULE_API TStateTreeStrongExecutionContext<true>::FinishTask(EStateTreeFinishTaskType FinishType) const;
 template bool STATETREEMODULE_API TStateTreeStrongExecutionContext<true>::UpdateScheduledTickRequest(UE::StateTree::FScheduledTickHandle Handle, FStateTreeScheduledTick ScheduledTick) const;
-template void STATETREEMODULE_API TStateTreeStrongExecutionContext<true>::ScheduleNextTick(TNotNull<UObject*> Owner, TNotNull<const UStateTree*> RootStateTree, FStateTreeInstanceStorage& Storage);
+template void STATETREEMODULE_API TStateTreeStrongExecutionContext<true>::ScheduleNextTick(TNotNull<UObject*> Owner, TNotNull<const UStateTree*> RootStateTree, FStateTreeInstanceStorage& Storage, UE::StateTree::ETickReason Reason);
 
 FStateTreeWeakExecutionContext::FStateTreeWeakExecutionContext(const FStateTreeExecutionContext& Context)
 	: Owner(Context.GetOwner())
diff --git a/Engine/Plugins/Runtime/StateTree/Source/StateTreeModule/Private/StateTreeExecutionContext.cpp b/Engine/Plugins/Runtime/StateTree/Source/StateTreeModule/Private/StateTreeExecutionContext.cpp
--- a/Engine/Plugins/Runtime/StateTree/Source/StateTreeModule/Private/StateTreeExecutionContext.cpp
+++ b/Engine/Plugins/Runtime/StateTree/Source/StateTreeModule/Private/StateTreeExecutionContext.cpp
@@ -366,7 +366,7 @@ FStateTreeScheduledTick FStateTreeReadOnlyExecutionContext::GetNextScheduledTick
 		const FStateTreeExecutionFrame& CurrentFrame = Exec.ActiveFrames[FrameIndex];
 		if (!CurrentFrame.StateTree->IsScheduledTickAllowed())
 		{
-			return FStateTreeScheduledTick::MakeEveryFrames();
+			return FStateTreeScheduledTick::MakeEveryFrames(UE::StateTree::ETickReason::Forced);
 		}
 	}
 
@@ -374,11 +374,30 @@ FStateTreeScheduledTick FStateTreeReadOnlyExecutionContext::GetNextScheduledTick
 	const bool bHasEvents = EventQueue.HasEvents();
 	const bool bHasBroadcastedDelegates = Storage.HasBroadcastedDelegates();
 
+	struct FCustomTickRate
+	{
+		FCustomTickRate(float InDeltaTime, UE::StateTree::ETickReason InReason)
+			: DeltaTime(InDeltaTime)
+			, Reason(InReason)
+		{}
+
+		bool operator<(const FCustomTickRate& Other) const
+		{
+			return DeltaTime < Other.DeltaTime;
+		}
+
+		float DeltaTime = 0.0f;
+		UE::StateTree::ETickReason Reason = UE::StateTree::ETickReason::None;
+	};
+
+
+	TOptional<FCustomTickRate> CustomTickRate;
+
 	// We wish to return in order: EveryFrames, then NextFrame, then CustomTickRate, then Sleep.
 	// Do we have a state that requires a tick or is waiting for an event.
-	TOptional<float> CustomTickRate;
 	{
 		bool bHasTaskWithEveryFramesTick = false;
+		UE::StateTree::ETickReason TaskTickingReason = UE::StateTree::ETickReason::None;
 		for (int32 FrameIndex = 0; FrameIndex < Exec.ActiveFrames.Num(); ++FrameIndex)
 		{
 			const FStateTreeExecutionFrame& CurrentFrame = Exec.ActiveFrames[FrameIndex];
@@ -396,19 +415,25 @@ FStateTreeScheduledTick FStateTreeReadOnlyExecutionContext::GetNextScheduledTick
 				const FStateTreeStateHandle CurrentHandle = CurrentFrame.ActiveStates[StateIndex];
 				const FCompactStateTreeState& State = CurrentStateTree->States[CurrentHandle.Index];
 				checkf(State.bEnabled, TEXT("Only enabled states are in ActiveStates."));
-
 				if (State.bHasCustomTickRate)
 				{
-					CustomTickRate = CustomTickRate.IsSet() ? FMath::Min(CustomTickRate.GetValue(), State.CustomTickRate) : State.CustomTickRate;
+					if (!CustomTickRate.IsSet() || CustomTickRate.GetValue().DeltaTime > State.CustomTickRate)
+					{
+						CustomTickRate = {State.CustomTickRate, UE::StateTree::ETickReason::StateCustomTickRate};
+					}
 				}
 				else if (!CustomTickRate.IsSet())
 				{
-					if (State.DoesRequestTickTasks(bHasEvents)
-						|| State.ShouldTickTransitions(bHasEvents, bHasBroadcastedDelegates)
-						)
+					if (State.DoesRequestTickTasks(bHasEvents))
+					{
+						bHasTaskWithEveryFramesTick = true;
+						TaskTickingReason = UE::StateTree::ETickReason::TaskTicking;
+					}
+					else if (State.ShouldTickTransitions(bHasEvents, bHasBroadcastedDelegates))
 					{
 						// todo: ShouldTickTransitions has onevent or ontick. both can be already triggered and we are waiting for the delay
 						bHasTaskWithEveryFramesTick = true;
+						TaskTickingReason = UE::StateTree::ETickReason::TransitionTicking;
 					}
 				}
 			}
@@ -416,15 +441,18 @@ FStateTreeScheduledTick FStateTreeReadOnlyExecutionContext::GetNextScheduledTick
 
 		if (!CustomTickRate.IsSet() && bHasTaskWithEveryFramesTick)
 		{
-			return FStateTreeScheduledTick::MakeEveryFrames();
+			return FStateTreeScheduledTick::MakeEveryFrames(TaskTickingReason);
 		}
 
 		// If one state has a custom tick rate, then it overrides the tick rate for all states.
 		//Only return the CustomTickRate if it's > than NextFrame, the custom tick rate will be processed at the end.
-		if (CustomTickRate.IsSet() && CustomTickRate.GetValue() <= 0.0f)
+		if (const FCustomTickRate* CustomTickRateValue = CustomTickRate.GetPtrOrNull())
 		{
-			// A state might override the custom tick rate with > 0, then another state overrides it again with 0 to tick back every frame.
-			return FStateTreeScheduledTick::MakeEveryFrames();
+			if (CustomTickRateValue->DeltaTime <= 0.0f)
+			{
+				// A state might override the custom tick rate with > 0, then another state overrides it again with 0 to tick back every frame.
+				return FStateTreeScheduledTick::MakeEveryFrames(CustomTickRateValue->Reason);
+			}
 		}
 	}
 
@@ -437,26 +465,27 @@ FStateTreeScheduledTick FStateTreeReadOnlyExecutionContext::GetNextScheduledTick
 		{
 			return ScheduledTickRequest;
 		}
-		const float CachedTickRate = ScheduledTickRequest.GetTickRate();
+		const FCustomTickRate CachedTickRate(ScheduledTickRequest.GetTickRate(), ScheduledTickRequest.GetReason());
+
 		CustomTickRate = CustomTickRate.IsSet() ? FMath::Min(CustomTickRate.GetValue(), CachedTickRate) : CachedTickRate;
 	}
 
 	// Transitions
 	if (Storage.GetTransitionRequests().Num() > 0)
 	{
-		return FStateTreeScheduledTick::MakeNextFrame();
+		return FStateTreeScheduledTick::MakeNextFrame(UE::StateTree::ETickReason::TransitionRequest);
 	}
 
 	// Events are cleared every tick.
 	if (bHasEvents && Storage.IsOwningEventQueue())
 	{
-		return FStateTreeScheduledTick::MakeNextFrame();
+		return FStateTreeScheduledTick::MakeNextFrame(UE::StateTree::ETickReason::Event);
 	}
 
 	// Completed task. For EnterState or for user that only called TickTasks and not TickTransitions.
 	if (Exec.bHasPendingCompletedState)
 	{
-		return FStateTreeScheduledTick::MakeNextFrame();
+		return FStateTreeScheduledTick::MakeNextFrame(UE::StateTree::ETickReason::CompletedState);
 	}
 
 	// Min of all delayed transitions.
@@ -464,14 +493,15 @@ FStateTreeScheduledTick FStateTreeReadOnlyExecutionContext::GetNextScheduledTick
 	{
 		for (const FStateTreeTransitionDelayedState& Transition : Exec.DelayedTransitions)
 		{
-			CustomTickRate = CustomTickRate.IsSet() ? FMath::Min(CustomTickRate.GetValue(), Transition.TimeLeft) : Transition.TimeLeft;
+			const FCustomTickRate TransitionTickRate(Transition.TimeLeft, UE::StateTree::ETickReason::DelayedTransition);
+			CustomTickRate = CustomTickRate.IsSet() ? FMath::Min(CustomTickRate.GetValue(), TransitionTickRate) : TransitionTickRate;
 		}
 	}
 
 	// Custom tick rate for tasks and transitions.
-	if (CustomTickRate.IsSet())
+	if (const FCustomTickRate* CustomTickRateValue = CustomTickRate.GetPtrOrNull())
 	{
-		return FStateTreeScheduledTick::MakeCustomTickRate(CustomTickRate.GetValue());
+		return FStateTreeScheduledTick::MakeCustomTickRate(CustomTickRateValue->DeltaTime, CustomTickRateValue->Reason);
 	}
 
 	return FStateTreeScheduledTick::MakeSleep();
@@ -783,7 +813,7 @@ UE::StateTree::FScheduledTickHandle FStateTreeMinimalExecutionContext::AddSchedu
 	}
 
 	UE::StateTree::FScheduledTickHandle Result = Storage.GetMutableExecutionState().AddScheduledTickRequest(ScheduledTick);
-	ScheduleNextTick();
+	ScheduleNextTick(UE::StateTree::ETickReason::ScheduledTickRequest);
 	return Result;
 }
 
@@ -798,7 +828,7 @@ void FStateTreeMinimalExecutionContext::UpdateScheduledTickRequest(UE::StateTree
 
 	if (Storage.GetMutableExecutionState().UpdateScheduledTickRequest(Handle, ScheduledTick))
 	{
-		ScheduleNextTick();
+		ScheduleNextTick(UE::StateTree::ETickReason::ScheduledTickRequest);
 	}
 }
 
@@ -813,7 +843,7 @@ void FStateTreeMinimalExecutionContext::RemoveScheduledTickRequest(UE::StateTree
 
 	if (Storage.GetMutableExecutionState().RemoveScheduledTickRequest(Handle))
 	{
-		ScheduleNextTick();
+		ScheduleNextTick(UE::StateTree::ETickReason::ScheduledTickRequest);
 	}
 }
 
@@ -834,17 +864,17 @@ void FStateTreeMinimalExecutionContext::SendEvent(const FGameplayTag Tag, const
 	FStateTreeEventQueue& LocalEventQueue = Storage.GetMutableEventQueue();
 	if (LocalEventQueue.SendEvent(&Owner, Tag, Payload, Origin))
 	{
-		ScheduleNextTick();
+		ScheduleNextTick(UE::StateTree::ETickReason::Event);
 		UE_STATETREE_DEBUG_SEND_EVENT(this, &RootStateTree, Tag, Payload, Origin);
 	}
 }
 
-void FStateTreeMinimalExecutionContext::ScheduleNextTick()
+void FStateTreeMinimalExecutionContext::ScheduleNextTick(UE::StateTree::ETickReason Reason)
 {
 	TInstancedStruct<FStateTreeExecutionExtension>& ExecutionExtension = Storage.GetMutableExecutionState().ExecutionExtension;
 	if (bAllowedToScheduleNextTick && RootStateTree.IsScheduledTickAllowed() && ExecutionExtension.IsValid())
 	{
-		ExecutionExtension.GetMutable().ScheduleNextTick(FStateTreeExecutionExtension::FContextParameters(Owner, RootStateTree, Storage));
+		ExecutionExtension.GetMutable().ScheduleNextTick(FStateTreeExecutionExtension::FContextParameters(Owner, RootStateTree, Storage), FStateTreeExecutionExtension::FNextTickArguments(Reason));
 	}
 }
 
@@ -1833,7 +1863,7 @@ void FStateTreeExecutionContext::BroadcastDelegate(const FStateTreeDelegateDispa
 	GetExecState().DelegateActiveListeners.BroadcastDelegate(Dispatcher, GetExecState());
 	if (UE::StateTree::ExecutionContext::MarkDelegateAsBroadcasted(Dispatcher, *CurrentFrame, GetMutableInstanceData()->GetMutableStorage()))
 	{
-		ScheduleNextTick();
+		ScheduleNextTick(UE::StateTree::ETickReason::Delegate);
 	}
 }
 
@@ -1955,7 +1985,7 @@ void FStateTreeExecutionContext::RequestTransition(const FStateTreeTransitionReq
 		RequestWithSource.SourceStateID = ActiveStateIndex != INDEX_NONE ? RootFrame->ActiveStates.StateIDs[ActiveStateIndex] : RootFrame->ActiveStates.StateIDs[0];
 
 		InstanceData.AddTransitionRequest(&Owner, RequestWithSource);
-		ScheduleNextTick();
+		ScheduleNextTick(UE::StateTree::ETickReason::TransitionRequest);
 	}
 	else
 	{
diff --git a/Engine/Plugins/Runtime/StateTree/Source/StateTreeModule/Private/StateTreeExecutionExtension.cpp b/Engine/Plugins/Runtime/StateTree/Source/StateTreeModule/Private/StateTreeExecutionExtension.cpp
new file mode 100644
--- /dev/null
+++ b/Engine/Plugins/Runtime/StateTree/Source/StateTreeModule/Private/StateTreeExecutionExtension.cpp
@@ -0,0 +1,17 @@
+// Copyright Epic Games, Inc. All Rights Reserved.
+
+#include "StateTreeExecutionExtension.h"
+#include "StateTreeExecutionTypes.h"
+
+#include UE_INLINE_GENERATED_CPP_BY_NAME(StateTreeExecutionExtension)
+
+
+FStateTreeExecutionExtension::FNextTickArguments::FNextTickArguments()
+	: Reason(UE::StateTree::ETickReason::None)
+{
+}
+
+FStateTreeExecutionExtension::FNextTickArguments::FNextTickArguments(UE::StateTree::ETickReason InReason)
+	: Reason(InReason)
+{ }
+
diff --git a/Engine/Plugins/Runtime/StateTree/Source/StateTreeModule/Private/StateTreeExecutionTypes.cpp b/Engine/Plugins/Runtime/StateTree/Source/StateTreeModule/Private/StateTreeExecutionTypes.cpp
--- a/Engine/Plugins/Runtime/StateTree/Source/StateTreeModule/Private/StateTreeExecutionTypes.cpp
+++ b/Engine/Plugins/Runtime/StateTree/Source/StateTreeModule/Private/StateTreeExecutionTypes.cpp
@@ -168,7 +168,7 @@ void FStateTreeExecutionState::CacheScheduledTickRequest()
 				const float CachedTickRate = Request.ScheduledTick.GetTickRate();
 				CustomTickRate = CustomTickRate.IsSet() ? FMath::Min(CustomTickRate.GetValue(), CachedTickRate) : CachedTickRate;
 			}
-			return FStateTreeScheduledTick::MakeCustomTickRate(CustomTickRate.GetValue());
+			return FStateTreeScheduledTick::MakeCustomTickRate(CustomTickRate.GetValue(), UE::StateTree::ETickReason::ScheduledTickRequest);
 		};
 
 	CachedScheduledTickRequest = GetBestRequest(ScheduledTickRequests);
@@ -213,27 +213,27 @@ UE::StateTree::FFinishedTask::FFinishedTask(FActiveFrameID InFrameID, FActiveSta
 //----------------------------------------------------------------------//
 FStateTreeScheduledTick FStateTreeScheduledTick::MakeSleep()
 {
-	return FStateTreeScheduledTick(UE_FLOAT_NON_FRACTIONAL);
+	return FStateTreeScheduledTick(UE_FLOAT_NON_FRACTIONAL, UE::StateTree::ETickReason::None);
 }
 
-FStateTreeScheduledTick FStateTreeScheduledTick::MakeEveryFrames()
+FStateTreeScheduledTick FStateTreeScheduledTick::MakeEveryFrames(UE::StateTree::ETickReason Reason)
 {
-	return FStateTreeScheduledTick(0.0f);
+	return FStateTreeScheduledTick(0.0f, Reason);
 }
 
-FStateTreeScheduledTick FStateTreeScheduledTick::MakeNextFrame()
+FStateTreeScheduledTick FStateTreeScheduledTick::MakeNextFrame(UE::StateTree::ETickReason Reason)
 {
-	return FStateTreeScheduledTick(UE_KINDA_SMALL_NUMBER);
+	return FStateTreeScheduledTick(UE_KINDA_SMALL_NUMBER, Reason);
 }
 
-FStateTreeScheduledTick FStateTreeScheduledTick::MakeCustomTickRate(float DeltaTime)
+FStateTreeScheduledTick FStateTreeScheduledTick::MakeCustomTickRate(float DeltaTime, UE::StateTree::ETickReason Reason)
 {
 	ensureMsgf(DeltaTime >= 0.0f, TEXT("Use a value greater than zero."));
 	if (DeltaTime > 0.0f)
 	{
-		return FStateTreeScheduledTick(DeltaTime);
+		return FStateTreeScheduledTick(DeltaTime, Reason);
 	}
-	return MakeEveryFrames();
+	return MakeEveryFrames(Reason);
 }
 
 bool FStateTreeScheduledTick::ShouldSleep() const
diff --git a/Engine/Plugins/Runtime/StateTree/Source/StateTreeModule/Private/Tasks/StateTreeRunParallelStateTreeTask.cpp b/Engine/Plugins/Runtime/StateTree/Source/StateTreeModule/Private/Tasks/StateTreeRunParallelStateTreeTask.cpp
--- a/Engine/Plugins/Runtime/StateTree/Source/StateTreeModule/Private/Tasks/StateTreeRunParallelStateTreeTask.cpp
+++ b/Engine/Plugins/Runtime/StateTree/Source/StateTreeModule/Private/Tasks/StateTreeRunParallelStateTreeTask.cpp
@@ -9,7 +9,7 @@
 
 #define LOCTEXT_NAMESPACE "StateTree"
 
-void FStateTreeRunParallelStateTreeExecutionExtension::ScheduleNextTick(const FContextParameters& Context)
+void FStateTreeRunParallelStateTreeExecutionExtension::ScheduleNextTick(const FContextParameters& Context, const FNextTickArguments& Args)
 {
 	const FStateTreeMinimalExecutionContext ExecutionContext(&Context.Owner, &Context.StateTree, Context.InstanceData);
 	const FStateTreeScheduledTick ScheduledTick = ExecutionContext.GetNextScheduledTick();
diff --git a/Engine/Plugins/Runtime/StateTree/Source/StateTreeModule/Public/StateTreeAsyncExecutionContext.h b/Engine/Plugins/Runtime/StateTree/Source/StateTreeModule/Public/StateTreeAsyncExecutionContext.h
--- a/Engine/Plugins/Runtime/StateTree/Source/StateTreeModule/Public/StateTreeAsyncExecutionContext.h
+++ b/Engine/Plugins/Runtime/StateTree/Source/StateTreeModule/Public/StateTreeAsyncExecutionContext.h
@@ -190,7 +190,7 @@ private:
 	UE_API FStateTreeDataView GetInstanceDataPtrInternal() const;
 
 	template<bool bWriteAccess = bWithWriteAccess UE_REQUIRES(bWriteAccess)>
-	static void ScheduleNextTick(TNotNull<UObject*> Owner, TNotNull<const UStateTree*> RootStateTree, FStateTreeInstanceStorage& Storage);
+	static void ScheduleNextTick(TNotNull<UObject*> Owner, TNotNull<const UStateTree*> RootStateTree, FStateTreeInstanceStorage& Storage, UE::StateTree::ETickReason Reason);
 
 	TStrongObjectPtr<UObject> Owner;
 	TStrongObjectPtr<const UStateTree> StateTree;
diff --git a/Engine/Plugins/Runtime/StateTree/Source/StateTreeModule/Public/StateTreeExecutionContext.h b/Engine/Plugins/Runtime/StateTree/Source/StateTreeModule/Public/StateTreeExecutionContext.h
--- a/Engine/Plugins/Runtime/StateTree/Source/StateTreeModule/Public/StateTreeExecutionContext.h
+++ b/Engine/Plugins/Runtime/StateTree/Source/StateTreeModule/Public/StateTreeExecutionContext.h
@@ -261,7 +261,7 @@ public:
 protected:
 
 	/** Informs the owner when the instance of the tree must woke up from a scheduled tick sleep. */
-	UE_API void ScheduleNextTick();
+	UE_API void ScheduleNextTick(UE::StateTree::ETickReason Reason = UE::StateTree::ETickReason::None);
 
 protected:
 	/** The context is processing the tree. We do not need to inform the owner that something changed. */
diff --git a/Engine/Plugins/Runtime/StateTree/Source/StateTreeModule/Public/StateTreeExecutionExtension.h b/Engine/Plugins/Runtime/StateTree/Source/StateTreeModule/Public/StateTreeExecutionExtension.h
--- a/Engine/Plugins/Runtime/StateTree/Source/StateTreeModule/Public/StateTreeExecutionExtension.h
+++ b/Engine/Plugins/Runtime/StateTree/Source/StateTreeModule/Public/StateTreeExecutionExtension.h
@@ -4,12 +4,17 @@
 
 #include "StateTreeExecutionExtension.generated.h"
 
+#define UE_API STATETREEMODULE_API
+
 class UObject;
 class UStateTree;
 struct FStateTreeInstanceStorage;
 struct FStateTreeTransitionDelayedState;
 struct FStateTreeReferenceOverrides;
-
+namespace UE::StateTree
+{
+	enum class ETickReason : uint8;
+}
 
 /** Used by the execution context or a weak execution context to extend their functionalities. */
 USTRUCT()
@@ -44,8 +49,21 @@ public:
 		return Context.Owner.GetName();
 	}
 
+	struct FNextTickArguments
+	{
+		UE_API FNextTickArguments();
+		UE_API explicit FNextTickArguments(UE::StateTree::ETickReason Reason);
+
+		UE::StateTree::ETickReason Reason;
+	};
+
 	/** Callback when the execution context request the tree to wakeup from a schedule tick sleep. */
-	virtual void ScheduleNextTick(const FContextParameters& Context)
+	virtual void ScheduleNextTick(const FContextParameters& Context, const FNextTickArguments& Args)
+	{
+	}
+
+	UE_DEPRECATED(5.7, "Use ScheduleNextTick with the FNextTickArguments parameter.")
+	virtual void ScheduleNextTick(const FContextParameters& Context) final
 	{
 	}
 
@@ -54,3 +72,5 @@ public:
 	{
 	}
 };
+
+#undef UE_API
diff --git a/Engine/Plugins/Runtime/StateTree/Source/StateTreeModule/Public/StateTreeExecutionTypes.h b/Engine/Plugins/Runtime/StateTree/Source/StateTreeModule/Public/StateTreeExecutionTypes.h
--- a/Engine/Plugins/Runtime/StateTree/Source/StateTreeModule/Public/StateTreeExecutionTypes.h
+++ b/Engine/Plugins/Runtime/StateTree/Source/StateTreeModule/Public/StateTreeExecutionTypes.h
@@ -745,41 +745,67 @@ PRAGMA_ENABLE_DEPRECATION_WARNINGS
 
 namespace UE::StateTree
 {
-/** Describes a finished task waiting to be processed by an execution context. */
-struct
-UE_DEPRECATED(5.6, "Replaced with FStateTreeTasksCompletionStatus")
-FFinishedTask
-{
-	enum class EReasonType : uint8
+	/** Describes a finished task waiting to be processed by an execution context. */
+	struct
+	UE_DEPRECATED(5.6, "Replaced with FStateTreeTasksCompletionStatus")
+	FFinishedTask
 	{
-		/** A global task finished. The FrameID and TaskIndex are valid. */
-		GlobalTask,
-		/** A task inside a state finished. The FrameID, StateID and TaskIndex are valid. */
-		StateTask,
-		/** An internal transition finish the state. The FrameID and StateID are valid. */
-		InternalTransition,
+		enum class EReasonType : uint8
+		{
+			/** A global task finished. The FrameID and TaskIndex are valid. */
+			GlobalTask,
+			/** A task inside a state finished. The FrameID, StateID and TaskIndex are valid. */
+			StateTask,
+			/** An internal transition finish the state. The FrameID and StateID are valid. */
+			InternalTransition,
+		};
+		FFinishedTask() = default;
+		STATETREEMODULE_API FFinishedTask(FActiveFrameID FrameID, FActiveStateID StateID, FStateTreeIndex16 TaskIndex, EStateTreeRunStatus RunStatus, EReasonType Reason, bool bTickProcessed);
+		/** Frame ID that identifies 1the active frame. */
+		FActiveFrameID FrameID;
+		/** State ID that contains the finished task. */
+		FActiveStateID StateID;
+		/** Task that is finished and needs to be processed. */
+		FStateTreeIndex16 TaskIndex = FStateTreeIndex16::Invalid;
+		/** The result of the finished task. */
+		EStateTreeRunStatus RunStatus = EStateTreeRunStatus::Failed;
+		/** The reason of the finished task. */
+		EReasonType Reason = EReasonType::GlobalTask;
+		/**
+		 * Set to true if the task is completed before or during the TickTasks.
+		 * Used to identify tasks that were completed and had the chance to be processed by the algo.
+		 * If not processed, they won't trigger the transition in this frame.
+		 */
+		bool bTickProcessed = false;
 	};
-	FFinishedTask() = default;
-	STATETREEMODULE_API FFinishedTask(FActiveFrameID FrameID, FActiveStateID StateID, FStateTreeIndex16 TaskIndex, EStateTreeRunStatus RunStatus, EReasonType Reason, bool bTickProcessed);
-	/** Frame ID that identifies 1the active frame. */
-	FActiveFrameID FrameID;
-	/** State ID that contains the finished task. */
-	FActiveStateID StateID;
-	/** Task that is finished and needs to be processed. */
-	FStateTreeIndex16 TaskIndex = FStateTreeIndex16::Invalid;
-	/** The result of the finished task. */
-	EStateTreeRunStatus RunStatus = EStateTreeRunStatus::Failed;
-	/** The reason of the finished task. */
-	EReasonType Reason = EReasonType::GlobalTask;
-	/**
-	 * Set to true if the task is completed before or during the TickTasks.
-	 * Used to identify tasks that were completed and had the chance to be processed by the algo.
-	 * If not processed, they won't trigger the transition in this frame.
-	 */
-	bool bTickProcessed = false;
-};
-} //namespace UE::StateTree
 
+	/** Describes the reason behind StateTree ticking. */
+	UENUM()
+	enum class ETickReason : uint8
+	{
+		None,
+		/** A scheduled tick request is active. */
+		ScheduledTickRequest,
+		/** The tick is forced. The schema doesn't support scheduling. */
+		Forced,
+		/** An active state requested a custom tick rate. */
+		StateCustomTickRate,
+		/** A task in an active state request ticking. */
+		TaskTicking,
+		/** A transition in an active state request ticking. */
+		TransitionTicking,
+		/** A transition request occurs via RequestTransition. */
+		TransitionRequest,
+		/** An event need to be cleared. */
+		Event,
+		/** A state completed async and transition didn't tick yet. */
+		CompletedState,
+		/** A active delayed transitions is pending. */
+		DelayedTransition,
+		/** A broadcast delegate occurs via BroadcastDelegate, and a transition is listening to the delegate. */
+		Delegate,
+	};
+}
 
 /*
  * Information on how a state tree should tick next.
@@ -798,11 +824,11 @@ public:
 	/** Make a scheduled tick that returns Sleep. */
 	static UE_API FStateTreeScheduledTick MakeSleep();
 	/** Make a scheduled tick that returns EveryFrame. */
-	static UE_API FStateTreeScheduledTick MakeEveryFrames();
+	static UE_API FStateTreeScheduledTick MakeEveryFrames(UE::StateTree::ETickReason Reason = UE::StateTree::ETickReason::None);
 	/** Make a scheduled tick that returns NextFrame. */
-	static UE_API FStateTreeScheduledTick MakeNextFrame();
+	static UE_API FStateTreeScheduledTick MakeNextFrame(UE::StateTree::ETickReason Reason = UE::StateTree::ETickReason::None);
 	/** Make a scheduled tick that returns a tick rate. The value needs to be greater than zero. */
-	static UE_API FStateTreeScheduledTick MakeCustomTickRate(float DeltaTime);
+	static UE_API FStateTreeScheduledTick MakeCustomTickRate(float DeltaTime, UE::StateTree::ETickReason Reason = UE::StateTree::ETickReason::None);
 
 public:
 	/** @return true if it doesn't need to tick until an event/delegate/transition/... occurs. */
@@ -816,13 +842,21 @@ public:
 	/** @return the delay in seconds between ticks. */
 	UE_API float GetTickRate() const;
 
+	/** @return the reason why tick is required. */
+	UE::StateTree::ETickReason GetReason() const
+	{
+		return Reason;
+	}
+
 private:
-	FStateTreeScheduledTick(float DeltaTime)
+	FStateTreeScheduledTick(float DeltaTime, UE::StateTree::ETickReason Reason)
 		: NextDeltaTime(DeltaTime)
+		, Reason(Reason)
 	{}
 
 	UPROPERTY()
 	float NextDeltaTime = 0.0f;
+	UE::StateTree::ETickReason Reason = UE::StateTree::ETickReason::None;
 };
 
 
diff --git a/Engine/Plugins/Runtime/StateTree/Source/StateTreeModule/Public/Tasks/StateTreeRunParallelStateTreeTask.h b/Engine/Plugins/Runtime/StateTree/Source/StateTreeModule/Public/Tasks/StateTreeRunParallelStateTreeTask.h
--- a/Engine/Plugins/Runtime/StateTree/Source/StateTreeModule/Public/Tasks/StateTreeRunParallelStateTreeTask.h
+++ b/Engine/Plugins/Runtime/StateTree/Source/StateTreeModule/Public/Tasks/StateTreeRunParallelStateTreeTask.h
@@ -35,7 +35,7 @@ struct FStateTreeRunParallelStateTreeExecutionExtension : public FStateTreeExecu
 	GENERATED_BODY()
 
 public:
-	virtual void ScheduleNextTick(const FContextParameters& Context) override;
+	virtual void ScheduleNextTick(const FContextParameters& Context, const FNextTickArguments& Args) override;
 
 	FStateTreeWeakExecutionContext WeakExecutionContext;
 	UE::StateTree::FScheduledTickHandle ScheduledTickHandle;
diff --git a/Engine/Plugins/Runtime/StateTree/Source/StateTreeModule/StateTreeModule.Build.cs b/Engine/Plugins/Runtime/StateTree/Source/StateTreeModule/StateTreeModule.Build.cs
--- a/Engine/Plugins/Runtime/StateTree/Source/StateTreeModule/StateTreeModule.Build.cs
+++ b/Engine/Plugins/Runtime/StateTree/Source/StateTreeModule/StateTreeModule.Build.cs
@@ -8,6 +8,8 @@ namespace UnrealBuildTool.Rules
 		{
 			CppCompileWarningSettings.UnsafeTypeCastWarningLevel = WarningLevel.Warning;
 
+			bAllowUETypesInNamespaces = true;
+
 			PublicIncludePaths.AddRange(
 				new string[] {
 				}
