[StateTree] Fixing FStateTreeRandomTimeDuration::GetRandomVariance() so it returns the correct value

diff --git a/Engine/Plugins/Runtime/StateTree/Source/StateTreeModule/Public/StateTreeTypes.h b/Engine/Plugins/Runtime/StateTree/Source/StateTreeModule/Public/StateTreeTypes.h
--- a/Engine/Plugins/Runtime/StateTree/Source/StateTreeModule/Public/StateTreeTypes.h
+++ b/Engine/Plugins/Runtime/StateTree/Source/StateTreeModule/Public/StateTreeTypes.h
@@ -501,11 +501,14 @@ struct FStateTreeRandomTimeDuration
 	/** @return the maximum random variance. */
 	float GetRandomVariance() const
 	{
-		return Duration / Scale;
+		return RandomVariance / Scale;
 	}
 
 	/** @return True of the duration is empty (always returns 0). */
-	bool IsEmpty() const { return Duration == 0 && RandomVariance == 0; }
+	bool IsEmpty() const
+	{
+		return Duration == 0 && RandomVariance == 0;
+	}
 	
 	/** @return Returns random duration around Duration, varied by +-RandomVariation. */
 	float GetRandomDuration(const FRandomStream& RandomStream) const
@@ -516,7 +519,6 @@ struct FStateTreeRandomTimeDuration
 	}
 
 protected:
-
 	static constexpr float Scale = 100.0f;
 
 	uint16 Quantize(const float Value) const
