[StateTree] Include the Global Evaluator in the current frame's active node.

diff --git a/Engine/Plugins/Runtime/StateTree/Source/StateTreeModule/Private/StateTreeExecutionContext.cpp b/Engine/Plugins/Runtime/StateTree/Source/StateTreeModule/Private/StateTreeExecutionContext.cpp
--- a/Engine/Plugins/Runtime/StateTree/Source/StateTreeModule/Private/StateTreeExecutionContext.cpp
+++ b/Engine/Plugins/Runtime/StateTree/Source/StateTreeModule/Private/StateTreeExecutionContext.cpp
@@ -2830,7 +2830,7 @@ void FStateTreeExecutionContext::CopyAllBindingsOnActiveInstances(const ECopyBin
 		FStateTreeExecutionFrame* CurrentParentFrame = FrameIndex > 0 ? &Exec.ActiveFrames[FrameIndex - 1] : nullptr;
 		FStateTreeExecutionFrame& CurrentFrame = Exec.ActiveFrames[FrameIndex];
 		const UStateTree* CurrentStateTree = CurrentFrame.StateTree;
-		const int32 CurrentActiveTaskNodeIndex = CurrentFrame.ActiveTaskNodeIndex.AsInt32();
+		const int32 CurrentActiveNodeIndex = CurrentFrame.ActiveNodeIndex.AsInt32();
 
 		FCurrentlyProcessedFrameScope FrameScope(*this, CurrentParentFrame, CurrentFrame);
 
@@ -2840,19 +2840,22 @@ void FStateTreeExecutionContext::CopyAllBindingsOnActiveInstances(const ECopyBin
 			const int32 EvaluatorEnd = CurrentStateTree->EvaluatorsBegin + CurrentStateTree->EvaluatorsNum;
 			for (int32 EvalIndex = CurrentStateTree->EvaluatorsBegin; EvalIndex < EvaluatorEnd; ++EvalIndex)
 			{
-				const FStateTreeEvaluatorBase& Eval = CurrentStateTree->Nodes[EvalIndex].Get<const FStateTreeEvaluatorBase>();
-				if (Eval.BindingsBatch.IsValid())
+				if (EvalIndex <= CurrentActiveNodeIndex)
 				{
-					const FStateTreeDataView EvalInstanceView = GetDataView(CurrentParentFrame, CurrentFrame, Eval.InstanceDataHandle);
-					FNodeInstanceDataScope DataScope(*this, &Eval, EvalIndex, Eval.InstanceDataHandle, EvalInstanceView);
-					CopyBatchOnActiveInstances(CurrentParentFrame, CurrentFrame, EvalInstanceView, Eval.BindingsBatch);
+					const FStateTreeEvaluatorBase& Eval = CurrentStateTree->Nodes[EvalIndex].Get<const FStateTreeEvaluatorBase>();
+					if (Eval.BindingsBatch.IsValid())
+					{
+						const FStateTreeDataView EvalInstanceView = GetDataView(CurrentParentFrame, CurrentFrame, Eval.InstanceDataHandle);
+						FNodeInstanceDataScope DataScope(*this, &Eval, EvalIndex, Eval.InstanceDataHandle, EvalInstanceView);
+						CopyBatchOnActiveInstances(CurrentParentFrame, CurrentFrame, EvalInstanceView, Eval.BindingsBatch);
+					}
 				}
 			}
 
 			const int32 GlobalTasksEnd = CurrentStateTree->GlobalTasksBegin + CurrentStateTree->GlobalTasksNum;
 			for (int32 TaskIndex = CurrentStateTree->GlobalTasksBegin; TaskIndex < GlobalTasksEnd; ++TaskIndex)
 			{
-				if (TaskIndex <= CurrentActiveTaskNodeIndex)
+				if (TaskIndex <= CurrentActiveNodeIndex)
 				{
 					const FStateTreeTaskBase& Task = CurrentStateTree->Nodes[TaskIndex].Get<const FStateTreeTaskBase>();
 					const bool bTaskRequestCopy = (CopyType == ECopyBindings::ExitState && Task.bShouldCopyBoundPropertiesOnExitState)
@@ -2890,7 +2893,7 @@ void FStateTreeExecutionContext::CopyAllBindingsOnActiveInstances(const ECopyBin
 			const int32 TasksEnd = CurrentState.TasksBegin + CurrentState.TasksNum;
 			for (int32 TaskIndex = CurrentState.TasksBegin; TaskIndex < TasksEnd; ++TaskIndex)
 			{
-				if (TaskIndex <= CurrentActiveTaskNodeIndex)
+				if (TaskIndex <= CurrentActiveNodeIndex)
 				{
 					const FStateTreeTaskBase& Task = CurrentStateTree->Nodes[TaskIndex].Get<const FStateTreeTaskBase>();
 
@@ -3368,7 +3371,7 @@ EStateTreeRunStatus FStateTreeExecutionContext::EnterState(FSelectStateResult& A
 
 				FNodeInstanceDataScope DataScope(*this, &Task, AssetTaskIndex, Task.InstanceDataHandle, TaskInstanceView);
 
-				CurrentFrame.ActiveTaskNodeIndex = FStateTreeIndex16(AssetTaskIndex);
+				CurrentFrame.ActiveNodeIndex = FStateTreeIndex16(AssetTaskIndex);
 
 				// Copy bound properties.
 				if (Task.BindingsBatch.IsValid())
@@ -3557,7 +3560,7 @@ void FStateTreeExecutionContext::ExitState(const FSelectStateResult& Args, const
 				{
 					// Call task completed only if EnterState() was called.
 					// The task order in the tree (BF) allows us to use the comparison.
-					if (TaskIndex <= CurrentFrame.ActiveTaskNodeIndex.AsInt32())
+					if (TaskIndex <= CurrentFrame.ActiveNodeIndex.AsInt32())
 					{
 						const FStateTreeDataView TaskInstanceView = GetDataView(CurrentParentFrame, CurrentFrame, Task.InstanceDataHandle);
 						FNodeInstanceDataScope DataScope(*this, &Task, TaskIndex, Task.InstanceDataHandle, TaskInstanceView);
@@ -3582,7 +3585,7 @@ void FStateTreeExecutionContext::ExitState(const FSelectStateResult& Args, const
 				{
 					STATETREE_LOG(VeryVerbose, TEXT("%*sSkipped 'ExitState' for disabled Task: '%s'"), UE::StateTree::Debug::IndentSize, TEXT(""), *Task.Name.ToString());
 				}
-				CurrentFrame.ActiveTaskNodeIndex = FStateTreeIndex16(TaskIndex - 1);
+				CurrentFrame.ActiveNodeIndex = FStateTreeIndex16(TaskIndex - 1);
 			}
 
 			// Call state change events on conditions if needed.
@@ -3676,7 +3679,7 @@ void FStateTreeExecutionContext::StateCompleted()
 		const FStateTreeExecutionFrame* CurrentParentFrame = FrameIndex > 0 ? &Exec.ActiveFrames[FrameIndex - 1] : nullptr;
 		const FStateTreeExecutionFrame& CurrentFrame = Exec.ActiveFrames[FrameIndex];
 		const UStateTree* CurrentStateTree = CurrentFrame.StateTree;
-		const int32 CurrentActiveTaskNodeIndex = CurrentFrame.ActiveTaskNodeIndex.AsInt32();
+		const int32 CurrentActiveNodeIndex = CurrentFrame.ActiveNodeIndex.AsInt32();
 
 		FCurrentlyProcessedFrameScope FrameScope(*this, CurrentParentFrame, CurrentFrame);
 
@@ -3693,7 +3696,7 @@ void FStateTreeExecutionContext::StateCompleted()
 			for (int32 TaskIndex = (State.TasksBegin + State.TasksNum) - 1; TaskIndex >= State.TasksBegin; TaskIndex--)
 			{
 				// Call task completed only if EnterState() was called.
-				if (TaskIndex <= CurrentActiveTaskNodeIndex)
+				if (TaskIndex <= CurrentActiveNodeIndex)
 				{
 					const FStateTreeTaskBase& Task = CurrentStateTree->Nodes[TaskIndex].Get<const FStateTreeTaskBase>();
 
@@ -3916,7 +3919,8 @@ EStateTreeRunStatus FStateTreeExecutionContext::StartGlobalsForFrameInternal(con
 	UE::StateTree::FTasksCompletionStatus CurrentTasksStatus = CurrentFrame.ActiveTasksStatus.GetStatus(CurrentStateTree);
 
 	// Start evaluators
-	for (int32 EvalIndex = CurrentStateTree->EvaluatorsBegin; EvalIndex < (CurrentStateTree->EvaluatorsBegin + CurrentStateTree->EvaluatorsNum); EvalIndex++)
+	const int32 EvaluatorsEnd = CurrentStateTree->EvaluatorsBegin + CurrentStateTree->EvaluatorsNum;
+	for (int32 EvalIndex = CurrentStateTree->EvaluatorsBegin; EvalIndex < EvaluatorsEnd; ++EvalIndex)
 	{
 		const FStateTreeEvaluatorBase& Eval = CurrentStateTree->Nodes[EvalIndex].Get<const FStateTreeEvaluatorBase>();
 		FStateTreeDataView EvalInstanceView = bOnActiveInstances
@@ -3932,7 +3936,9 @@ EStateTreeRunStatus FStateTreeExecutionContext::StartGlobalsForFrameInternal(con
 		}
 
 		FNodeInstanceDataScope DataScope(*this, &Eval, EvalIndex, Eval.InstanceDataHandle, EvalInstanceView);
-		// Copy bound properties with the temporary buffer.
+		CurrentFrame.ActiveNodeIndex = FStateTreeIndex16(EvalIndex);
+
+		// Copy bound properties.
 		if (Eval.BindingsBatch.IsValid())
 		{
 			if constexpr (bOnActiveInstances)
@@ -3982,10 +3988,9 @@ EStateTreeRunStatus FStateTreeExecutionContext::StartGlobalsForFrameInternal(con
 		}
 
 		FNodeInstanceDataScope DataScope(*this, &Task, AssetTaskIndex, Task.InstanceDataHandle, TaskDataView);
+		CurrentFrame.ActiveNodeIndex = FStateTreeIndex16(AssetTaskIndex);
 
-		CurrentFrame.ActiveTaskNodeIndex = FStateTreeIndex16(AssetTaskIndex);
-
-		// Copy bound properties with the temporary buffer.
+		// Copy bound properties.
 		if (Task.BindingsBatch.IsValid())
 		{
 			if constexpr (bOnActiveInstances)
@@ -4141,14 +4146,14 @@ void FStateTreeExecutionContext::StopGlobalsForFrameInternal(const FStateTreeExe
 	const int32 GlobalTasksEnd = CurrentStateTree->GlobalTasksBegin + CurrentStateTree->GlobalTasksNum;
 	for (int32 TaskIndex = GlobalTasksEnd - 1; TaskIndex >= CurrentStateTree->GlobalTasksBegin; --TaskIndex)
 	{
-		const FStateTreeTaskBase& Task = CurrentStateTree->Nodes[TaskIndex].Get<const FStateTreeTaskBase>();
-		const FStateTreeDataView TaskInstanceView = bOnActiveInstances
-			? GetDataView(ParentFrame, Frame, Task.InstanceDataHandle)
-			: GetDataViewOrTemporary(ParentFrame, Frame, Task.InstanceDataHandle);
-		FNodeInstanceDataScope DataScope(*this, &Task, TaskIndex, Task.InstanceDataHandle, TaskInstanceView);
-
-		if (TaskIndex <= Frame.ActiveTaskNodeIndex.AsInt32())
+		if (TaskIndex <= Frame.ActiveNodeIndex.AsInt32())
 		{
+			const FStateTreeTaskBase& Task = CurrentStateTree->Nodes[TaskIndex].Get<const FStateTreeTaskBase>();
+			const FStateTreeDataView TaskInstanceView = bOnActiveInstances
+				? GetDataView(ParentFrame, Frame, Task.InstanceDataHandle)
+				: GetDataViewOrTemporary(ParentFrame, Frame, Task.InstanceDataHandle);
+			FNodeInstanceDataScope DataScope(*this, &Task, TaskIndex, Task.InstanceDataHandle, TaskInstanceView);
+
 			// Ignore disabled task
 			if (Task.bTaskEnabled == false)
 			{
@@ -4165,29 +4170,33 @@ void FStateTreeExecutionContext::StopGlobalsForFrameInternal(const FStateTreeExe
 				UE_STATETREE_DEBUG_TASK_EVENT(this, TaskIndex, TaskInstanceView, EStateTreeTraceEventType::OnExited, Transition.CurrentRunStatus);
 			}
 		}
-		Frame.ActiveTaskNodeIndex = FStateTreeIndex16(TaskIndex - 1);
+		Frame.ActiveNodeIndex = FStateTreeIndex16(TaskIndex - 1);
 	}
-	Frame.ActiveTaskNodeIndex = FStateTreeIndex16::Invalid;
 
 	const int32 EvaluatorsEnd = CurrentStateTree->EvaluatorsBegin + CurrentStateTree->EvaluatorsNum;
 	for (int32 EvalIndex = EvaluatorsEnd - 1; EvalIndex >= CurrentStateTree->EvaluatorsBegin; --EvalIndex)
 	{
-		const FStateTreeEvaluatorBase& Eval = CurrentStateTree->Nodes[EvalIndex].Get<const FStateTreeEvaluatorBase>();
-		const FStateTreeDataView EvalInstanceView = bOnActiveInstances
-			? GetDataView(ParentFrame, Frame, Eval.InstanceDataHandle)
-			: GetDataViewOrTemporary(ParentFrame, Frame, Eval.InstanceDataHandle);
-		FNodeInstanceDataScope DataScope(*this, &Eval, EvalIndex, Eval.InstanceDataHandle, EvalInstanceView);
-
-		STATETREE_LOG(Verbose, TEXT("  Stop: '%s'"), *Eval.Name.ToString());
-		UE_STATETREE_DEBUG_EVALUATOR_EXIT_TREE(this, CurrentStateTree, FStateTreeIndex16(EvalIndex));
+		if (EvalIndex <= Frame.ActiveNodeIndex.AsInt32())
 		{
-			QUICK_SCOPE_CYCLE_COUNTER(StateTree_Eval_TreeStop);
-			Eval.TreeStop(*this);
+			const FStateTreeEvaluatorBase& Eval = CurrentStateTree->Nodes[EvalIndex].Get<const FStateTreeEvaluatorBase>();
+			const FStateTreeDataView EvalInstanceView = bOnActiveInstances
+				? GetDataView(ParentFrame, Frame, Eval.InstanceDataHandle)
+				: GetDataViewOrTemporary(ParentFrame, Frame, Eval.InstanceDataHandle);
+			FNodeInstanceDataScope DataScope(*this, &Eval, EvalIndex, Eval.InstanceDataHandle, EvalInstanceView);
 
-			UE_STATETREE_DEBUG_EVALUATOR_EVENT(this, EvalIndex, EvalInstanceView, EStateTreeTraceEventType::OnTreeStopped);
+			STATETREE_LOG(Verbose, TEXT("  Stop: '%s'"), *Eval.Name.ToString());
+			UE_STATETREE_DEBUG_EVALUATOR_EXIT_TREE(this, CurrentStateTree, FStateTreeIndex16(EvalIndex));
+			{
+				QUICK_SCOPE_CYCLE_COUNTER(StateTree_Eval_TreeStop);
+				Eval.TreeStop(*this);
+
+				UE_STATETREE_DEBUG_EVALUATOR_EVENT(this, EvalIndex, EvalInstanceView, EStateTreeTraceEventType::OnTreeStopped);
+			}
 		}
+		Frame.ActiveNodeIndex = FStateTreeIndex16(EvalIndex - 1);
 	}
 
+	Frame.ActiveNodeIndex = FStateTreeIndex16::Invalid;
 	Frame.bHaveEntered = false;
 
 	if constexpr (bOnActiveInstances)
@@ -4411,7 +4420,7 @@ FStateTreeExecutionContext::FTickTaskResult FStateTreeExecutionContext::TickTask
 	const bool bCopyBoundPropertiesOnNonTickedTask = ExecutionContext::Private::bCopyBoundPropertiesOnNonTickedTask;
 	const UStateTree* CurrentStateTree = Args.Frame->StateTree;
 	const FActiveFrameID CurrentFrameID = Args.Frame->FrameID;
-	const int32 CurrentActiveTaskNodeIndex = Args.Frame->ActiveTaskNodeIndex.AsInt32();
+	const int32 CurrentActiveNodeIndex = Args.Frame->ActiveNodeIndex.AsInt32();
 	check(CurrentStateTree);
 
 	for (int32 OwnerTaskIndex = 0; OwnerTaskIndex < Args.TasksNum; ++OwnerTaskIndex)
@@ -4426,7 +4435,7 @@ FStateTreeExecutionContext::FTickTaskResult FStateTreeExecutionContext::TickTask
 			continue;
 		}
 
-		if (AssetTaskIndex > CurrentActiveTaskNodeIndex)
+		if (AssetTaskIndex > CurrentActiveNodeIndex)
 		{
 			STATETREE_LOG(VeryVerbose, TEXT("%*sSkipped 'Tick' for task that didn't get the EnterState Task: '%s'"), Debug::IndentSize, TEXT(""), *Task.Name.ToString());
 			bShouldTickTasks = false;
@@ -5154,7 +5163,7 @@ bool FStateTreeExecutionContext::TriggerTransitions()
 		{
 			FStateTreeExecutionFrame& CurrentFrame = Exec.ActiveFrames[FrameIndex];
 			const UStateTree* CurrentStateTree = CurrentFrame.StateTree;
-			const int32 CurrentActiveTaskNodeIndex = CurrentFrame.ActiveTaskNodeIndex.AsInt32();
+			const int32 CurrentActiveNodeIndex = CurrentFrame.ActiveNodeIndex.AsInt32();
 
 			for (int32 StateIndex = CurrentFrame.ActiveStates.Num() - 1; StateIndex >= 0; StateIndex--)
 			{
@@ -5171,7 +5180,7 @@ bool FStateTreeExecutionContext::TriggerTransitions()
 					for (int32 TaskIndex = (State.TasksBegin + State.TasksNum) - 1; TaskIndex >= State.TasksBegin; TaskIndex--)
 					{
 						const FStateTreeTaskBase& Task = CurrentStateTree->Nodes[TaskIndex].Get<const FStateTreeTaskBase>();
-						if (Task.bShouldAffectTransitions && Task.bTaskEnabled && TaskIndex <= CurrentActiveTaskNodeIndex)
+						if (Task.bShouldAffectTransitions && Task.bTaskEnabled && TaskIndex <= CurrentActiveNodeIndex)
 						{
 							TransitionHandlers.Emplace(FrameIndex, StateHandle, StateID, FStateTreeIndex16(TaskIndex), Task.TransitionHandlingPriority);
 							bAdded = true;
diff --git a/Engine/Plugins/Runtime/StateTree/Source/StateTreeModule/Private/StateTreeExecutionTypes.cpp b/Engine/Plugins/Runtime/StateTree/Source/StateTreeModule/Private/StateTreeExecutionTypes.cpp
--- a/Engine/Plugins/Runtime/StateTree/Source/StateTreeModule/Private/StateTreeExecutionTypes.cpp
+++ b/Engine/Plugins/Runtime/StateTree/Source/StateTreeModule/Private/StateTreeExecutionTypes.cpp
@@ -365,7 +365,7 @@ void FStateTreeDelegateActiveListeners::BroadcastDelegate(const FStateTreeDelega
 			if (const FStateTreeExecutionFrame* ExectionFrame = Exec.FindActiveFrame(ActiveListener.FrameID))
 			{
 				// Is the node active and is the state active.
-				if (ActiveListener.OwningNodeIndex.Get() <= ExectionFrame->ActiveTaskNodeIndex.Get())
+				if (ActiveListener.OwningNodeIndex.Get() <= ExectionFrame->ActiveNodeIndex.Get())
 				{
 					if (!ActiveListener.StateID.IsValid())
 					{
diff --git a/Engine/Plugins/Runtime/StateTree/Source/StateTreeModule/Public/StateTreeExecutionTypes.h b/Engine/Plugins/Runtime/StateTree/Source/StateTreeModule/Public/StateTreeExecutionTypes.h
--- a/Engine/Plugins/Runtime/StateTree/Source/StateTreeModule/Public/StateTreeExecutionTypes.h
+++ b/Engine/Plugins/Runtime/StateTree/Source/StateTreeModule/Public/StateTreeExecutionTypes.h
@@ -1006,11 +1006,11 @@ struct FStateTreeExecutionFrame
 	UE::StateTree::FActiveFrameID FrameID;
 
 	/**
-	 * The task node index that was "entered".
+	 * The evaluator or task node index that was "entered".
 	 * Used during the Enter and Exit phase. A node can fail EnterState.
 	 * Nodes after ActiveNodeIndex do not receive ExitState, because they didn't receive EnterState.
 	 */
-	FStateTreeIndex16 ActiveTaskNodeIndex = FStateTreeIndex16::Invalid;
+	FStateTreeIndex16 ActiveNodeIndex = FStateTreeIndex16::Invalid;
 
 	/** First index of the external data for this frame. */
 	UPROPERTY()
@@ -1205,15 +1205,15 @@ PRAGMA_DISABLE_DEPRECATION_WARNINGS
 	UE_DEPRECATED(5.6, "CurrentExecutionContext is not needed anymore. Use FrameID and StateID.")
 	FStateTreeExecutionContext* CurrentExecutionContext = nullptr;
 
-	UE_DEPRECATED(5.7, "Use FStateTreeExecutionFrame::ActiveTaskNodeIndex instead.")
+	UE_DEPRECATED(5.7, "Use FStateTreeExecutionFrame::ActiveNodeIndex instead.")
 	UPROPERTY()
 	FStateTreeIndex16 EnterStateFailedFrameIndex = FStateTreeIndex16::Invalid;
 
-	UE_DEPRECATED(5.7, "Use FStateTreeExecutionFrame::ActiveTaskNodeIndex instead.")
+	UE_DEPRECATED(5.7, "Use FStateTreeExecutionFrame::ActiveNodeIndex instead.")
 	UPROPERTY()
 	FStateTreeIndex16 EnterStateFailedTaskIndex = FStateTreeIndex16::Invalid;
 
-	UE_DEPRECATED(5.7, "Use FStateTreeExecutionFrame::ActiveTaskNodeIndex instead.")
+	UE_DEPRECATED(5.7, "Use FStateTreeExecutionFrame::ActiveNodeIndex instead.")
 	UPROPERTY()
 	FStateTreeIndex16 LastExitedNodeIndex = FStateTreeIndex16::Invalid;
 PRAGMA_ENABLE_DEPRECATION_WARNINGS
