[StateTree] Remove error message when the linked asset is not set during selection.

diff --git a/Engine/Plugins/Runtime/StateTree/Source/StateTreeModule/Private/StateTreeExecutionContext.cpp b/Engine/Plugins/Runtime/StateTree/Source/StateTreeModule/Private/StateTreeExecutionContext.cpp
--- a/Engine/Plugins/Runtime/StateTree/Source/StateTreeModule/Private/StateTreeExecutionContext.cpp
+++ b/Engine/Plugins/Runtime/StateTree/Source/StateTreeModule/Private/StateTreeExecutionContext.cpp
@@ -6913,8 +6913,12 @@ bool FStateTreeExecutionContext::SelectStateInternal_LinkedAsset(
 
 	FStateTreeExecutionState& Exec = GetExecState();
 
-	if (NextLinkedStateAsset == nullptr
-		|| NextLinkedStateAsset->States.Num() == 0
+	if (NextLinkedStateAsset == nullptr)
+	{
+		return false;
+	}
+
+	if (NextLinkedStateAsset->States.Num() == 0
 		|| !NextLinkedStateAsset->IsReadyToRun())
 	{
 		STATETREE_LOG(Error, TEXT("%hs: The linked State Tree is invalid. '%s' using StateTree '%s'."),
