[StateTree] Fixed issue when 2 states wait for 2 different events when triggered from a transition. The old behavior was: The state to and the transition need to request the same transition. The new behavior doesn't have that requirement. A state can listen to any event. Added the cvar StateTree.TargetStateRequiresTheSameEventForStateSelectionAsTheRequestedTransition to get the old behavior.

diff --git a/Engine/Plugins/Runtime/StateTree/Source/StateTreeModule/Private/StateTreeExecutionContext.cpp b/Engine/Plugins/Runtime/StateTree/Source/StateTreeModule/Private/StateTreeExecutionContext.cpp
--- a/Engine/Plugins/Runtime/StateTree/Source/StateTreeModule/Private/StateTreeExecutionContext.cpp
+++ b/Engine/Plugins/Runtime/StateTree/Source/StateTreeModule/Private/StateTreeExecutionContext.cpp
@@ -67,6 +67,14 @@ namespace UE::StateTree::ExecutionContext::Private
 		TEXT("The information is not relevant and not valid but can be needed for backward compatibility.")
 	);
 
+	bool bTargetStateRequiresTheSameEventForStateSelectionAsTheRequestedTransition = false;
+	FAutoConsoleVariableRef CVarTargetStateRequiresTheSameEventForStateSelectionAsTheRequestedTransition(
+		TEXT("StateTree.TargetStateRequiresTheSameEventForStateSelectionAsTheRequestedTransition"),
+		bTargetStateRequiresTheSameEventForStateSelectionAsTheRequestedTransition,
+		TEXT("If true, when a transition is triggered by an event that occurs and the target state requires an event, \n")
+		TEXT("then the state required event must match the transition required event.")
+	);
+
 	const FLazyName Name_Tick = "Tick";
 	const FLazyName Name_Start = "Start";
 	const FLazyName Name_Stop = "Stop";
@@ -6412,13 +6420,20 @@ bool FStateTreeExecutionContext::SelectStateInternal(
 		if (ensure(NextState.RequiredEventToEnter.IsValid()))
 		{
 			// Use the same event as performed transition unless it didn't lead to this state as only state selected by the transition should get it's event.
-			if (Args.TransitionEvent.IsValid() && bIsTargetState)
+			const bool bCanUseTransitionEvent = Args.TransitionEvent.IsValid()
+				&& Args.TargetState == NextStateContext
+				&& bIsTargetState;
+			if (bCanUseTransitionEvent && bTargetStateRequiresTheSameEventForStateSelectionAsTheRequestedTransition)
 			{
 				if (NextState.RequiredEventToEnter.DoesEventMatchDesc(*Args.TransitionEvent))
 				{
 					StateSelectionEvents.Emplace(Args.TransitionEvent);
 				}
 			}
+			else if (bCanUseTransitionEvent && NextState.RequiredEventToEnter.DoesEventMatchDesc(*Args.TransitionEvent))
+			{
+				StateSelectionEvents.Emplace(Args.TransitionEvent);
+			}
 			else
 			{
 				TArrayView<FStateTreeSharedEvent> EventsQueue = GetMutableEventsToProcessView();
@@ -6675,9 +6690,9 @@ TOptional<bool> FStateTreeExecutionContext::SelectStateFromSourceInternal(
 		{
 			OutSelectionResult.SelectedStates.Add(InternalArgs.MissingSourceStates[0]);
 
-			bool bSucceededToSelectState = false;
+			bool bSelectStateInternalSucceeded = false;
 			ON_SCOPE_EXIT{
-				if (bSucceededToSelectState == false)
+				if (bSelectStateInternalSucceeded == false)
 				{
 					OutSelectionResult.SelectedStates.Pop();
 				}
@@ -6689,9 +6704,9 @@ TOptional<bool> FStateTreeExecutionContext::SelectStateFromSourceInternal(
 				.MissingSourceStates = InternalArgs.MissingSourceStates.Mid(1),
 				.MissingStatesToReachTarget = InternalArgs.MissingStatesToReachTarget.Mid(1)
 			};
-			bSucceededToSelectState = SelectStateInternal(Args, NewInternalArgs, OutSelectionResult);
+			bSelectStateInternalSucceeded = SelectStateInternal(Args, NewInternalArgs, OutSelectionResult);
 
-			if (bSucceededToSelectState && bCompletedState)
+			if (bSelectStateInternalSucceeded && bCompletedState)
 			{
 				check(EnumHasAllFlags(NextFrame.StateTree->StateSelectionRules, EStateTreeStateSelectionRules::CompletedStateBeforeTransitionSourceFailsTransition));
 				// Before the source, there's a completed state.
@@ -6708,10 +6723,10 @@ TOptional<bool> FStateTreeExecutionContext::SelectStateFromSourceInternal(
 
 					UE_STATETREE_DEBUG_LOG_EVENT(this, Log, TEXT("Selection fails because parent state '%s' is completed."), *NextState.Name.ToString());
 
-					bSucceededToSelectState = false;
+					bSelectStateInternalSucceeded = false;
 				}
 			}
-			return bSucceededToSelectState;
+			return bSelectStateInternalSucceeded;
 		}
 	}
 	return {};
diff --git a/Engine/Plugins/Runtime/StateTree/Source/StateTreeModule/Public/StateTreeExecutionContext.h b/Engine/Plugins/Runtime/StateTree/Source/StateTreeModule/Public/StateTreeExecutionContext.h
--- a/Engine/Plugins/Runtime/StateTree/Source/StateTreeModule/Public/StateTreeExecutionContext.h
+++ b/Engine/Plugins/Runtime/StateTree/Source/StateTreeModule/Public/StateTreeExecutionContext.h
@@ -50,6 +50,8 @@ namespace UE::StateTree::ExecutionContext
 		}
 		const UStateTree* StateTree = nullptr;
 		FStateTreeStateHandle StateHandle;
+
+		bool operator==(const FStateHandleContext&) const = default;
 	};
 }
 
diff --git a/Engine/Plugins/Runtime/StateTree/Source/StateTreeTestSuite/Private/StateTreeTransitionTest.cpp b/Engine/Plugins/Runtime/StateTree/Source/StateTreeTestSuite/Private/StateTreeTransitionTest.cpp
--- a/Engine/Plugins/Runtime/StateTree/Source/StateTreeTestSuite/Private/StateTreeTransitionTest.cpp
+++ b/Engine/Plugins/Runtime/StateTree/Source/StateTreeTestSuite/Private/StateTreeTransitionTest.cpp
@@ -1706,12 +1706,12 @@ struct FStateTreeTest_TransitionToNone : FStateTreeTestBase
 		UStateTreeState& State1 = Root.AddChildState(FName(TEXT("State1")));
 		UStateTreeState& State2 = Root.AddChildState(FName(TEXT("State2")));
 
-		auto& TransitionRoot = Root.AddTransition(EStateTreeTransitionTrigger::OnEvent, EStateTreeTransitionType::GotoState);
+		FStateTreeTransition& TransitionRoot = Root.AddTransition(EStateTreeTransitionTrigger::OnEvent, EStateTreeTransitionType::GotoState);
 		TransitionRoot.State = State1.GetLinkToState();
 		TransitionRoot.RequiredEvent.Tag = Tag;
 
-		auto& Task1 = State1.AddTask<FTestTask_StandNoTick>(FName(TEXT("Task1")));
-		auto& TransitionState1 = State1.AddTransition(EStateTreeTransitionTrigger::OnEvent, EStateTreeTransitionType::None);
+		TStateTreeEditorNode<FTestTask_StandNoTick>& Task1 = State1.AddTask<FTestTask_StandNoTick>(FName(TEXT("Task1")));
+		FStateTreeTransition& TransitionState1 = State1.AddTransition(EStateTreeTransitionTrigger::OnEvent, EStateTreeTransitionType::None);
 		TransitionState1.RequiredEvent.Tag = Tag;
 
 		FStateTreeCompilerLog Log;
@@ -1749,6 +1749,116 @@ struct FStateTreeTest_TransitionToNone : FStateTreeTestBase
 };
 IMPLEMENT_STATE_TREE_INSTANT_TEST(FStateTreeTest_TransitionToNone, "System.StateTree.Transition.ToNone");
 
+struct FStateTreeTest_TransitionTwoEvents : FStateTreeTestBase
+{
+	virtual bool InstantTest() override
+	{
+		UStateTree& StateTree = NewStateTree();
+		{
+			UStateTreeEditorData& EditorData = *Cast<UStateTreeEditorData>(StateTree.EditorData);
+
+			UStateTreeState& Root = EditorData.AddSubTree(FName(TEXT("Root")));
+			UStateTreeState& State1 = Root.AddChildState(FName(TEXT("State1")));
+			UStateTreeState& State2 = Root.AddChildState(FName(TEXT("State2")));
+			UStateTreeState& State3 = State2.AddChildState(FName(TEXT("State3")));
+			UStateTreeState& State4 = State2.AddChildState(FName(TEXT("State4")));
+
+			State2.bHasRequiredEventToEnter = true;
+			State2.RequiredEventToEnter.Tag = GetTestTag1();
+
+			State3.bHasRequiredEventToEnter = true;
+			State3.RequiredEventToEnter.Tag = GetTestTag2();
+
+			FStateTreeTransition& TransitionA = State1.AddTransition(EStateTreeTransitionTrigger::OnEvent, EStateTreeTransitionType::GotoState);
+			TransitionA.State = State2.GetLinkToState();
+			TransitionA.RequiredEvent.Tag = GetTestTag1();
+
+			TStateTreeEditorNode<FTestTask_StandNoTick>& State3Task1 = State3.AddTask<FTestTask_StandNoTick>(FName(TEXT("State3Task1")));
+		}
+		{
+			FStateTreeCompilerLog Log;
+			FStateTreeCompiler Compiler(Log);
+			const bool bResult = Compiler.Compile(StateTree);
+			AITEST_TRUE("StateTree should get compiled", bResult);
+		}
+
+		const TCHAR* EnterStateStr(TEXT("EnterState"));
+		const TCHAR* ExitStateStr(TEXT("ExitState"));
+
+		FStateTreeInstanceData InstanceData;
+		TArray<FRecordedStateTreeTransitionResult> RecordedTransitions;
+		{
+			FTestStateTreeExecutionContext Exec(StateTree, StateTree, InstanceData, {}, EStateTreeRecordTransitions::Yes);
+			const bool bInitSucceeded = Exec.IsValid();
+			AITEST_TRUE("StateTree should init", bInitSucceeded);
+
+			// Start and enter state
+			EStateTreeRunStatus Status = Exec.Start();
+			AITEST_TRUE(TEXT("Active states should stay the same"), Exec.ExpectInActiveStates("Root", "State1"));
+			Exec.LogClear();
+
+			// Transition from State1 to State2. State3 fails (missing event), enter State4.
+			Exec.SendEvent(GetTestTag1());
+			Exec.Tick(0.1f);
+			AITEST_TRUE(TEXT("Active states should stay the same"), Exec.ExpectInActiveStates("Root", "State2", "State4"));
+			Exec.LogClear();
+
+			Exec.Stop();
+			Exec.Start();
+			AITEST_TRUE(TEXT("Active states should stay the same"), Exec.ExpectInActiveStates("Root", "State1"));
+
+			// Transition from State1 to State2 failed.
+			Exec.SendEvent(GetTestTag2());
+			Status = Exec.Tick(0.1f);
+			AITEST_TRUE(TEXT("Active states should stay the same"), Exec.ExpectInActiveStates("Root", "State1"));
+			Exec.LogClear();
+
+			// Transition from State1 to State2. State3 succeed.
+			Exec.SendEvent(GetTestTag2());
+			Exec.SendEvent(GetTestTag1());
+			Status = Exec.Tick(0.1f);
+			AITEST_TRUE(TEXT("Active states should stay the same"), Exec.ExpectInActiveStates("Root", "State2", "State3"));
+			Exec.LogClear();
+
+			RecordedTransitions = Exec.GetRecordedTransitions();
+			Exec.Stop();
+		}
+
+		{
+			FTestStateTreeExecutionContext Exec(StateTree, StateTree, InstanceData);
+			const bool bInitSucceeded = Exec.IsValid();
+			AITEST_TRUE("StateTree should init", bInitSucceeded);
+
+			// Start and enter state
+			EStateTreeRunStatus Status = Exec.Start();
+			AITEST_TRUE(TEXT("Active states should stay the same"), Exec.ExpectInActiveStates("Root", "State1"));
+			Exec.LogClear();
+
+			AITEST_TRUE(TEXT("Wrong number of transitions"), RecordedTransitions.Num() == 5);
+
+			Exec.ForceTransition(RecordedTransitions[0]); // start
+			AITEST_TRUE(TEXT("Start."), Exec.ExpectInActiveStates("Root", "State1"));
+
+			Exec.ForceTransition(RecordedTransitions[1]); // state 3 doesn't have the event
+			AITEST_TRUE(TEXT("State3 failed."), Exec.ExpectInActiveStates("Root", "State2", "State4"));
+			
+			Exec.ForceTransition(RecordedTransitions[2]); // start
+			AITEST_TRUE(TEXT("Start and failed transition."), Exec.ExpectInActiveStates("Root", "State1"));
+
+			Exec.ForceTransition(RecordedTransitions[3]); // failed
+			AITEST_TRUE(TEXT("Start and failed transition."), Exec.ExpectInActiveStates("Root", "State1"));
+
+			Exec.ForceTransition(RecordedTransitions[4]); // succeed
+			AITEST_TRUE(TEXT("State3 succeed."), Exec.ExpectInActiveStates("Root", "State2", "State3"));
+
+			Exec.Stop();
+		}
+
+		return true;
+	}
+};
+IMPLEMENT_STATE_TREE_INSTANT_TEST(FStateTreeTest_TransitionTwoEvents, "System.StateTree.Transition.TwoEvents");
+
 } // namespace UE::StateTree::Tests
 
 UE_ENABLE_OPTIMIZATION_SHIP
