[StateTree] Access instance data from async task. Use a temporary frame created during a transition to access the instance data. Test if the node received the ExitState before accessing the instance data. The instance data are valid if the frame and state are in the current active path or in the future active path (when requested from a transition request).

diff --git a/Engine/Plugins/Runtime/StateTree/Source/StateTreeModule/Private/StateTreeAsyncExecutionContext.cpp b/Engine/Plugins/Runtime/StateTree/Source/StateTreeModule/Private/StateTreeAsyncExecutionContext.cpp
--- a/Engine/Plugins/Runtime/StateTree/Source/StateTreeModule/Private/StateTreeAsyncExecutionContext.cpp
+++ b/Engine/Plugins/Runtime/StateTree/Source/StateTreeModule/Private/StateTreeAsyncExecutionContext.cpp
@@ -18,6 +18,7 @@ TStateTreeStrongExecutionContext<bWithWriteAccess>::TStateTreeStrongExecutionCon
 	: Owner(WeakContext.Owner.Pin()),
 	StateTree(WeakContext.StateTree.Pin()),
 	Storage(WeakContext.Storage.Pin()),
+	TemporaryStorage(WeakContext.TemporaryStorage.Pin()),
 	FrameID(WeakContext.FrameID),
 	StateID(WeakContext.StateID),
 	NodeIndex(WeakContext.NodeIndex)
@@ -187,13 +188,11 @@ bool TStateTreeStrongExecutionContext<bWithWriteAccess>::FinishTask(EStateTreeFi
 	if (ActivePath.IsValid())
 	{
 		FStateTreeExecutionState& Exec = Storage->GetMutableExecutionState();
-		FStateTreeExecutionFrame* NonConstFrame = const_cast<FStateTreeExecutionFrame*>(ActivePath.Frame);
-		if (ensure(NonConstFrame->StateTree))
+		const UStateTree* FrameStateTree = ActivePath.Frame->StateTree;
+		if (ensure(FrameStateTree))
 		{
 			using namespace UE::StateTree;
 
-			const UStateTree* FrameStateTree = NonConstFrame->StateTree;
-
 			const int32 AssetNodeIndex = ActivePath.NodeIndex.AsInt32();
 			const int32 GlobalTaskBeginIndex = FrameStateTree->GlobalTasksBegin;
 			const int32 GlobalTaskEndIndex = FrameStateTree->GlobalTasksBegin + FrameStateTree->GlobalTasksNum;
@@ -201,24 +200,25 @@ bool TStateTreeStrongExecutionContext<bWithWriteAccess>::FinishTask(EStateTreeFi
 			const ETaskCompletionStatus TaskStatus = ExecutionContext::CastToTaskStatus(FinishType);
 
 			bool bCompleted = false;
-			if (NonConstFrame->bIsGlobalFrame && bIsGlobalTask)
+			if (ActivePath.Frame->bIsGlobalFrame && bIsGlobalTask)
 			{
+				check(!ActivePath.StateHandle.IsValid());
 				const int32 FrameTaskIndex = AssetNodeIndex - GlobalTaskBeginIndex;
-				FTasksCompletionStatus GlobalTasksStatus = NonConstFrame->ActiveTasksStatus.GetStatus(NonConstFrame->StateTree);
+				FTasksCompletionStatus GlobalTasksStatus = ActivePath.Frame->ActiveTasksStatus.GetStatus(FrameStateTree);
 				GlobalTasksStatus.SetStatusWithPriority(FrameTaskIndex, TaskStatus);
 				bCompleted = GlobalTasksStatus.IsCompleted();
 				bSucceed = true;
 			}
-			else
+			else if (ensure(ActivePath.StateHandle.IsValid()))
 			{
-				const FCompactStateTreeState& State = NonConstFrame->StateTree->States[ActivePath.StateHandle.Index];
+				const FCompactStateTreeState& State = FrameStateTree->States[ActivePath.StateHandle.Index];
 				const int32 StateTaskBeginIndex = State.TasksBegin;
 				const int32 StateTaskEndIndex = State.TasksBegin + State.TasksNum;
 				const bool bIsStateTask = AssetNodeIndex >= StateTaskBeginIndex && AssetNodeIndex < StateTaskEndIndex;
 				if (bIsStateTask)
 				{
 					const int32 StateTaskIndex = AssetNodeIndex - State.TasksBegin;
-					FTasksCompletionStatus StateTasksStatus = NonConstFrame->ActiveTasksStatus.GetStatus(State);
+					FTasksCompletionStatus StateTasksStatus = ActivePath.Frame->ActiveTasksStatus.GetStatus(State);
 					StateTasksStatus.SetStatusWithPriority(StateTaskIndex, TaskStatus);
 					bCompleted = StateTasksStatus.IsCompleted();
 					bSucceed = true;
@@ -259,24 +259,65 @@ bool TStateTreeStrongExecutionContext<bWithWriteAccess>::UpdateScheduledTickRequ
 template <bool bWithWriteAccess>
 UE::StateTree::Async::FActivePathInfo TStateTreeStrongExecutionContext<bWithWriteAccess>::GetActivePathInfo() const
 {
-	if (IsValidInstanceStorage())
+	if (!IsValidInstanceStorage())
 	{
-		FStateTreeExecutionState& Exec = Storage->GetMutableExecutionState();
+		return {};
+	}
 
-		int32 FrameIndex = Exec.IndexOfActiveFrame(FrameID);
-		if (FrameIndex != INDEX_NONE)
+	UE::StateTree::Async::FActivePathInfo Result;
+	FStateTreeExecutionState& Exec = Storage->GetMutableExecutionState();
+	if (const int32 ActiveFrameIndex = Exec.IndexOfActiveFrame(FrameID); ActiveFrameIndex != INDEX_NONE)
+	{
+		Result.Frame = &Exec.ActiveFrames[ActiveFrameIndex];
+		Result.ParentFrame = ActiveFrameIndex != 0 ? &Exec.ActiveFrames[ActiveFrameIndex - 1] : nullptr;
+	}
+	else if (TemporaryStorage)
+	{
+		using FFrameAndParent = UE::StateTree::ExecutionContext::ITemporaryStorage::FFrameAndParent;
+		const FFrameAndParent FrameInfo = TemporaryStorage->GetExecutionFrame(FrameID);
+		Result.Frame = FrameInfo.Frame;
+		if (FrameInfo.ParentFrameID.IsValid())
 		{
-			FStateTreeExecutionFrame* Frame = &Exec.ActiveFrames[FrameIndex];
-			FStateTreeStateHandle StateHandle = StateID.IsValid() ? Frame->ActiveStates.FindStateHandle(StateID) : FStateTreeStateHandle::Invalid;
-			if (!StateID.IsValid() || StateHandle.IsValid())
+			if (const int32 ActiveParentFrameIndex = Exec.IndexOfActiveFrame(FrameInfo.ParentFrameID); ActiveParentFrameIndex != INDEX_NONE)
 			{
-				FStateTreeExecutionFrame* ParentFrame = FrameIndex != 0 ? &Exec.ActiveFrames[FrameIndex - 1] : nullptr;
-				return {Frame, ParentFrame, StateHandle, NodeIndex};
+				Result.ParentFrame = &Exec.ActiveFrames[ActiveParentFrameIndex];
 			}
+			else
+			{
+				const FFrameAndParent ParentFrameInfo = TemporaryStorage->GetExecutionFrame(FrameInfo.ParentFrameID);
+				Result.ParentFrame = ParentFrameInfo.Frame;
+			}
+
+			ensureMsgf(Result.Frame, TEXT("The frame ID exist in the frame holder. It should exist either in the active of the holder."));
 		}
 	}
 
-	return {};
+	if (Result.Frame == nullptr)
+	{
+		return {};
+	}
+
+	if (NodeIndex.AsInt32() > Result.Frame->ActiveNodeIndex.AsInt32())
+	{
+		return {};
+	}
+	Result.NodeIndex = NodeIndex;
+
+	if (StateID.IsValid())
+	{
+		Result.StateHandle = Result.Frame->ActiveStates.FindStateHandle(StateID);
+		if (!Result.StateHandle.IsValid() && TemporaryStorage)
+		{
+			Result.StateHandle = TemporaryStorage->GetStateHandle(StateID).GetStateHandle();
+		}
+
+		if (!Result.StateHandle.IsValid())
+		{
+			return {};
+		}
+	}
+
+	return Result;
 }
 
 template <bool bWithWriteAccess>
@@ -335,15 +376,10 @@ FStateTreeWeakExecutionContext::FStateTreeWeakExecutionContext(const FStateTreeE
 	if (const FStateTreeExecutionFrame* Frame = Context.GetCurrentlyProcessedFrame())
 	{
 		FrameID = Frame->FrameID;
-		const int32 ActiveStateIndex = Frame->ActiveStates.IndexOfReverse(Context.GetCurrentlyProcessedState());
-		if (ActiveStateIndex != INDEX_NONE)
-		{
-			StateID = Frame->ActiveStates.StateIDs[ActiveStateIndex];
-		}
-
-		// Only writes down the NodeIndex if it is created synchronously from the ExecContext tick frame
+		StateID = Frame->ActiveStates.FindStateID(Context.GetCurrentlyProcessedState());
 		NodeIndex = Context.GetCurrentlyProcessedNodeIndex();
 	}
+	TemporaryStorage = Context.GetCurrentlyProcessedTemporaryStorage();
 }
 
 bool FStateTreeWeakExecutionContext::SendEvent(const FGameplayTag Tag, const FConstStructView Payload, const FName Origin) const
diff --git a/Engine/Plugins/Runtime/StateTree/Source/StateTreeModule/Private/StateTreeExecutionContext.cpp b/Engine/Plugins/Runtime/StateTree/Source/StateTreeModule/Private/StateTreeExecutionContext.cpp
--- a/Engine/Plugins/Runtime/StateTree/Source/StateTreeModule/Private/StateTreeExecutionContext.cpp
+++ b/Engine/Plugins/Runtime/StateTree/Source/StateTreeModule/Private/StateTreeExecutionContext.cpp
@@ -924,6 +924,26 @@ FStateTreeExecutionFrame& FStateTreeExecutionContext::FSelectStateResult::MakeAn
 	return ExecFrame;
 }
 
+FStateTreeExecutionContext::FSelectStateResult::FFrameAndParent FStateTreeExecutionContext::FSelectStateResult::GetExecutionFrame(UE::StateTree::FActiveFrameID ID)
+{
+	const int32 FoundIndex = SelectedFrames.IndexOfByKey(ID);
+	FStateTreeExecutionFrame* Frame = FoundIndex != INDEX_NONE
+		? TemporaryFrames.FindByPredicate([ID](const FStateTreeExecutionFrame& Other){ return Other.FrameID == ID; })
+		: nullptr;
+	const UE::StateTree::FActiveFrameID ParentFrameID = FoundIndex > 0 ? SelectedFrames[FoundIndex - 1] : UE::StateTree::FActiveFrameID();
+	return FFrameAndParent{ .Frame = Frame, .ParentFrameID = ParentFrameID };
+}
+
+UE::StateTree::FActiveState FStateTreeExecutionContext::FSelectStateResult::GetStateHandle(UE::StateTree::FActiveStateID ID) const
+{
+	const int32 FoundIndex = UE::StateTree::FActiveStatePath::IndexOf(SelectedStates, ID);
+	if (FoundIndex != INDEX_NONE)
+	{
+		return SelectedStates[FoundIndex];
+	}
+	return {};
+}
+
 /**
  * FStateTreeExecutionContext::FCurrentlyProcessedFrameScope implementation
  */
@@ -1369,14 +1389,17 @@ EStateTreeRunStatus FStateTreeExecutionContext::Start(FStartParameters Parameter
 	ensure(Exec.ActiveFrames.Num() == 0);
 
 	// Initialize for the init frame.
-	FSelectStateResult SelectStateResult;
+	TSharedRef<FSelectStateResult> SelectStateResult = MakeShared<FSelectStateResult>();
+	check(CurrentlyProcessedTemporaryStorage == nullptr);
+	TGuardValue<TSharedPtr<FSelectStateResult>> TemporaryStorageScope(CurrentlyProcessedTemporaryStorage, SelectStateResult);
+
 	const FActiveFrameID InitFrameID = FActiveFrameID(Storage.GenerateUniqueId());
 	const FExecutionFrameKey InitFrameKey = FExecutionFrameKey(&RootStateTree, FStateTreeStateHandle::Root);
 	constexpr bool bIsGlobalFrame = true;
-	FStateTreeExecutionFrame& InitFrame = SelectStateResult.MakeAndAddTemporaryFrame(InitFrameID, InitFrameKey, bIsGlobalFrame);
+	FStateTreeExecutionFrame& InitFrame = SelectStateResult->MakeAndAddTemporaryFrame(InitFrameID, InitFrameKey, bIsGlobalFrame);
 	InitFrame.GlobalParameterDataHandle = FStateTreeDataHandle(EStateTreeDataSourceType::GlobalParameterData);
 
-	if (!CollectActiveExternalData(SelectStateResult.TemporaryFrames))
+	if (!CollectActiveExternalData(SelectStateResult->TemporaryFrames))
 	{
 		STATETREE_LOG(Warning, TEXT("%hs: Failed to collect external data ('%s' using StateTree '%s')"),
 			__FUNCTION__, *GetNameSafe(&Owner), *GetFullNameSafe(&RootStateTree));
@@ -1394,6 +1417,7 @@ EStateTreeRunStatus FStateTreeExecutionContext::Start(FStartParameters Parameter
 
 	// Start evaluators and global tasks. Fail the execution if any global task fails.
 	constexpr FStateTreeExecutionFrame* ParentInitFrame = nullptr;
+	SelectStateResult->SelectedFrames.Add(InitFrame.FrameID);
 	EStateTreeRunStatus GlobalTasksRunStatus = StartTemporaryEvaluatorsAndGlobalTasks(ParentInitFrame, InitFrame);
 	if (GlobalTasksRunStatus == EStateTreeRunStatus::Running)
 	{
@@ -1411,8 +1435,8 @@ EStateTreeRunStatus FStateTreeExecutionContext::Start(FStartParameters Parameter
 		SelectStateArgs.Behavior = ESelectStateBehavior::StateTransition;
 		if (SelectState(SelectStateArgs, SelectStateResult))
 		{
-			check(SelectStateResult.SelectedStates.Num() > 0);
-			const FStateTreeStateHandle LastSelectedStateHandle = SelectStateResult.SelectedStates.Last().GetStateHandle();
+			check(SelectStateResult->SelectedStates.Num() > 0);
+			const FStateTreeStateHandle LastSelectedStateHandle = SelectStateResult->SelectedStates.Last().GetStateHandle();
 			if (LastSelectedStateHandle.IsCompletionState())
 			{
 				// Transition to a terminal state (succeeded/failed).
@@ -1476,6 +1500,8 @@ EStateTreeRunStatus FStateTreeExecutionContext::Start(FStartParameters Parameter
 		Exec.ActiveFrames.Reset();
 		RemoveAllDelegateListeners();
 
+		SelectStateResult->SelectedFrames.Pop();
+
 		// We are not considered as running yet so we only set the status without requiring a stop.
 		Exec.TreeRunStatus = GlobalTasksRunStatus;
 	}
@@ -1547,7 +1573,7 @@ EStateTreeRunStatus FStateTreeExecutionContext::Stop(EStateTreeRunStatus Complet
 	if (Exec.TreeRunStatus == EStateTreeRunStatus::Running)
 	{
 		// Transition to Succeeded state.
-		FSelectStateResult EmptySelectionResult;
+		const TSharedPtr<FSelectStateResult> EmptySelectionResult;
 		FStateTreeTransitionResult Transition;
 		Transition.TargetState = FStateTreeStateHandle::FromCompletionStatus(CompletionStatus);
 		Transition.CurrentRunStatus = CompletionStatus;
@@ -2124,7 +2150,7 @@ void FStateTreeExecutionContext::UpdateInstanceData(TConstArrayView<FStateTreeEx
 }
 PRAGMA_ENABLE_DEPRECATION_WARNINGS
 
-void FStateTreeExecutionContext::UpdateInstanceData(FSelectStateResult& Args)
+void FStateTreeExecutionContext::UpdateInstanceData(const TSharedPtr<FSelectStateResult>& Args)
 {
 	// Go though all the state and frame (includes unchanged, sustained, changed).
 	//Keep 2 buffers. InstanceStructs for unchanged and sustained. TempInstanceStructs for changed.
@@ -2143,17 +2169,18 @@ void FStateTreeExecutionContext::UpdateInstanceData(FSelectStateResult& Args)
 
 	// Estimate how many new instance data items we might have.
 	int32 EstimatedNumStructs = 0;
+	if (Args)
 	{
 		FActiveFrameID CurrentFrameID;
 		const UStateTree* CurrentStateTree = nullptr;
-		for (const FActiveState& SelectedState : Args.SelectedStates)
+		for (const FActiveState& SelectedState : Args->SelectedStates)
 		{
 			// Global
 			if (SelectedState.GetFrameID() != CurrentFrameID)
 			{
 				CurrentFrameID = SelectedState.GetFrameID();
 
-				const FStateTreeExecutionFrame* Frame = FindExecutionFrame(CurrentFrameID, MakeConstArrayView(Exec.ActiveFrames), MakeConstArrayView(Args.TemporaryFrames));
+				const FStateTreeExecutionFrame* Frame = FindExecutionFrame(CurrentFrameID, MakeConstArrayView(Exec.ActiveFrames), MakeConstArrayView(Args->TemporaryFrames));
 				if (ensure(Frame) && Frame->bIsGlobalFrame)
 				{
 					check(Frame->StateTree);
@@ -2202,14 +2229,15 @@ void FStateTreeExecutionContext::UpdateInstanceData(FSelectStateResult& Args)
 	FStateTreeDataHandle NextStateParameterDataHandle = FStateTreeDataHandle::Invalid;
 
 	FStateTreeDataHandle CurrentGlobalParameterDataHandle = FStateTreeDataHandle(EStateTreeDataSourceType::GlobalParameterData);
+	if (Args)
 	{
 		FActiveFrameID CurrentFrameID;
 		FStateTreeExecutionFrame* CurrentFrame = nullptr;
 		const UStateTree* CurrentStateTree = nullptr;
 		int32 CurrentFrameBaseIndex = 0;
-		for (int32 SelectedStateIndex = 0; SelectedStateIndex < Args.SelectedStates.Num(); ++SelectedStateIndex)
+		for (int32 SelectedStateIndex = 0; SelectedStateIndex < Args->SelectedStates.Num(); ++SelectedStateIndex)
 		{
-			const FActiveState& SelectedState = Args.SelectedStates[SelectedStateIndex];
+			const FActiveState& SelectedState = Args->SelectedStates[SelectedStateIndex];
 
 			// Global
 			if (SelectedState.GetFrameID() != CurrentFrameID)
@@ -2222,7 +2250,7 @@ void FStateTreeExecutionContext::UpdateInstanceData(FSelectStateResult& Args)
 
 				if (CurrentFrame == nullptr)
 				{
-					CurrentFrame = Args.TemporaryFrames.FindByPredicate(
+					CurrentFrame = Args->TemporaryFrames.FindByPredicate(
 						[FrameID = SelectedState.GetFrameID()](const FStateTreeExecutionFrame& Frame)
 						{
 							return Frame.FrameID == FrameID;
@@ -2352,11 +2380,11 @@ void FStateTreeExecutionContext::UpdateInstanceData(FSelectStateResult& Args)
 									// This state is a container for the linked state tree.
 									// Its instance data matches the linked state tree parameters.
 									// The linked state tree asset is the next frame.
-									const bool bIsLastFrame = SelectedStateIndex == Args.SelectedStates.Num() - 1;
+									const bool bIsLastFrame = SelectedStateIndex == Args->SelectedStates.Num() - 1;
 									if (!bIsLastFrame)
 									{
-										const FActiveState& FollowingSelectedState = Args.SelectedStates[SelectedStateIndex + 1];
-										const FStateTreeExecutionFrame* FollowingNextFrame = FindExecutionFrame(FollowingSelectedState.GetFrameID(), MakeArrayView(Exec.ActiveFrames), MakeArrayView(Args.TemporaryFrames));
+										const FActiveState& FollowingSelectedState = Args->SelectedStates[SelectedStateIndex + 1];
+										const FStateTreeExecutionFrame* FollowingNextFrame = FindExecutionFrame(FollowingSelectedState.GetFrameID(), MakeArrayView(Exec.ActiveFrames), MakeArrayView(Args->TemporaryFrames));
 										if (ensure(FollowingNextFrame))
 										{
 											ParamsInstanceData = FConstStructView::Make(TempParams.Emplace_GetRef(FollowingNextFrame->StateTree->GetDefaultParameters()));
@@ -2530,14 +2558,15 @@ EStateTreeRunStatus FStateTreeExecutionContext::ForceTransition(const FRecordedS
 		StateContexts.Add(FStateHandleContext(RecordedActiveState.StateTree, RecordedActiveState.State));
 	}
 
-	FSelectStateResult SelectStateResult;
+	TSharedRef<FSelectStateResult> SelectStateResult = MakeShared<FSelectStateResult>();
+	TGuardValue<TSharedPtr<FSelectStateResult>> TemporaryStorageScope(CurrentlyProcessedTemporaryStorage, SelectStateResult);
 	if (!ForceTransitionInternal(StateContexts, SelectStateResult))
 	{
 		UE_STATETREE_DEBUG_LOG_EVENT(this, Verbose, TEXT("The force transition failed."));
 		return EStateTreeRunStatus::Unset;
 	}
 
-	if (SelectStateResult.SelectedStates.Num() != StateContexts.Num())
+	if (SelectStateResult->SelectedStates.Num() != StateContexts.Num())
 	{
 		UE_STATETREE_DEBUG_LOG_EVENT(this, Verbose, TEXT("The force transition failed to find the target."));
 		return EStateTreeRunStatus::Unset;
@@ -2555,8 +2584,8 @@ EStateTreeRunStatus FStateTreeExecutionContext::ForceTransition(const FRecordedS
 		const FRecordedActiveState& RecordedActiveState = RecordedTransition.States[RecordedStateIndex];
 		if (SharedEvents.IsValidIndex(RecordedActiveState.EventIndex))
 		{
-			const FActiveState& SelectedState = SelectStateResult.SelectedStates[RecordedStateIndex];
-			SelectStateResult.SelectionEvents.Add(FSelectionEventWithID{.State = SelectedState, .Event= SharedEvents[RecordedActiveState.EventIndex]});
+			const FActiveState& SelectedState = SelectStateResult->SelectedStates[RecordedStateIndex];
+			SelectStateResult->SelectionEvents.Add(FSelectionEventWithID{.State = SelectedState, .Event= SharedEvents[RecordedActiveState.EventIndex]});
 		}
 	}
 
@@ -2648,7 +2677,7 @@ namespace UE::StateTree::ExecutionContext::Private
 	}
 }
 
-bool FStateTreeExecutionContext::ForceTransitionInternal(const TArrayView<const UE::StateTree::ExecutionContext::FStateHandleContext> StateContexts, FSelectStateResult& OutSelectionResult)
+bool FStateTreeExecutionContext::ForceTransitionInternal(const TArrayView<const UE::StateTree::ExecutionContext::FStateHandleContext> StateContexts, const TSharedRef<FSelectStateResult>& OutSelectionResult)
 {
 	using namespace UE::StateTree;
 	using namespace UE::StateTree::ExecutionContext;
@@ -3100,7 +3129,7 @@ void FStateTreeExecutionContext::CaptureNewStateEvents(TConstArrayView<FStateTre
 }
 PRAGMA_ENABLE_DEPRECATION_WARNINGS
 
-void FStateTreeExecutionContext::CaptureNewStateEvents(FSelectStateResult& Args)
+void FStateTreeExecutionContext::CaptureNewStateEvents(const TSharedRef<const FSelectStateResult>& Args)
 {
 	using namespace UE::StateTree;
 	using namespace UE::StateTree::ExecutionContext::Private;
@@ -3120,15 +3149,15 @@ void FStateTreeExecutionContext::CaptureNewStateEvents(FSelectStateResult& Args)
 	// For each state that are changed (not sustained)
 	FActiveFrameID CurrentFrameID;
 	const FStateTreeExecutionFrame* CurrentFrame = nullptr;
-	for (int32 SelectedStateIndex = 0; SelectedStateIndex < Args.SelectedStates.Num(); ++SelectedStateIndex)
+	for (int32 SelectedStateIndex = 0; SelectedStateIndex < Args->SelectedStates.Num(); ++SelectedStateIndex)
 	{
-		const FActiveState& SelectedState = Args.SelectedStates[SelectedStateIndex];
+		const FActiveState& SelectedState = Args->SelectedStates[SelectedStateIndex];
 
 		// Global
 		if (SelectedState.GetFrameID() != CurrentFrameID)
 		{
 			CurrentFrameID = SelectedState.GetFrameID();
-			CurrentFrame = FindExecutionFrame(SelectedState.GetFrameID(), MakeConstArrayView(Exec.ActiveFrames), MakeConstArrayView(Args.TemporaryFrames));
+			CurrentFrame = FindExecutionFrame(SelectedState.GetFrameID(), MakeConstArrayView(Exec.ActiveFrames), MakeConstArrayView(Args->TemporaryFrames));
 		}
 
 		check(CurrentFrame);
@@ -3143,7 +3172,7 @@ void FStateTreeExecutionContext::CaptureNewStateEvents(FSelectStateResult& Args)
 			if (State->EventDataIndex.IsValid())
 			{
 				FStateTreeSharedEvent& StateTreeEvent = Storage.GetMutableStruct(CurrentFrame->ActiveInstanceIndexBase.Get() + State->EventDataIndex.Get()).Get<FStateTreeSharedEvent>();
-				const FSelectionEventWithID* EventToCapturePtr = Args.SelectionEvents.FindByPredicate([SelectedState](const FSelectionEventWithID& EventID)
+				const FSelectionEventWithID* EventToCapturePtr = Args->SelectionEvents.FindByPredicate([SelectedState](const FSelectionEventWithID& EventID)
 					{
 						return EventID.State == SelectedState;
 					});
@@ -3173,7 +3202,7 @@ EStateTreeRunStatus FStateTreeExecutionContext::EnterState(FStateTreeTransitionR
 	return EStateTreeRunStatus::Failed;
 }
 
-EStateTreeRunStatus FStateTreeExecutionContext::EnterState(FSelectStateResult& Args, const FStateTreeTransitionResult& Transition)
+EStateTreeRunStatus FStateTreeExecutionContext::EnterState(const TSharedPtr<FSelectStateResult>& ArgsPtr, const FStateTreeTransitionResult& Transition)
 {
 	// 1. Update the data instance data of all frames and selected states.
 	// The data won't be available until the node is reached or the state is added to the FStateTreeExecutionFrame::ActiveState.
@@ -3187,20 +3216,27 @@ EStateTreeRunStatus FStateTreeExecutionContext::EnterState(FSelectStateResult& A
 
 	CSV_SCOPED_TIMING_STAT_EXCLUSIVE(StateTree_EnterState);
 
-	if (Args.SelectedStates.IsEmpty())
+	if (!ArgsPtr.IsValid())
 	{
+		TSharedPtr<FSelectStateResult> EmptySelectStateResult;
+		UpdateInstanceData(EmptySelectStateResult);
 		return EStateTreeRunStatus::Failed;
 	}
 
+	FSelectStateResult& Args = *ArgsPtr;
 	FStateTreeExecutionState& Exec = GetExecState();
 
 	if (bRecordTransitions)
 	{
-		RecordedTransitions.Add(MakeRecordedTransitionResult(Args, Transition));
+		RecordedTransitions.Add(MakeRecordedTransitionResult(ArgsPtr.ToSharedRef(), Transition));
 	}
 
-	UpdateInstanceData(Args);
-	CaptureNewStateEvents(Args);
+	UpdateInstanceData(ArgsPtr);
+	if (Args.SelectedStates.IsEmpty())
+	{
+		return EStateTreeRunStatus::Failed;
+	}
+	CaptureNewStateEvents(ArgsPtr.ToSharedRef());
 
 	++Exec.StateChangeCount;
 
@@ -3447,7 +3483,7 @@ void FStateTreeExecutionContext::ExitState(const FStateTreeTransitionResult& Tra
 {
 }
 
-void FStateTreeExecutionContext::ExitState(const FSelectStateResult& Args, const FStateTreeTransitionResult& Transition)
+void FStateTreeExecutionContext::ExitState(const TSharedPtr<const FSelectStateResult>& Args, const FStateTreeTransitionResult& Transition)
 {
 	// 1. Copy all bindings on all active nodes.
 	// 2. Call ExitState on all active global tasks, global evaluators, state tasks and state condition that are affected by the transition.
@@ -3480,18 +3516,24 @@ void FStateTreeExecutionContext::ExitState(const FSelectStateResult& Args, const
 	FActiveStateInlineArray ActiveStates;
 	GetActiveStatePath(GetExecState().ActiveFrames, ActiveStates);
 
-	const TArrayView<const FActiveState> CommonStates = FActiveStatePath::Intersect(MakeConstArrayView(ActiveStates), MakeConstArrayView(Args.SelectedStates));
+	const FActiveState ArgsTargetState = Args ? Args->TargetState : FActiveState();
+	const TArrayView<const FActiveState> CommonStates = Args != nullptr
+		? FActiveStatePath::Intersect(MakeConstArrayView(ActiveStates), MakeConstArrayView(Args->SelectedStates))
+		: TArrayView<const FActiveState>();
+	const TArrayView<const FActiveFrameID> ArgsSelectedFrames = Args != nullptr
+		? MakeConstArrayView(Args->SelectedFrames)
+		: TArrayView<const FActiveFrameID>();
 	const TArrayView<const FActiveState> ChangedStates = MakeConstArrayView(ActiveStates).Mid(CommonStates.Num());
-	const bool bHasSustained = CommonStates.Contains(Args.TargetState);
+	const bool bHasSustained = CommonStates.Contains(ArgsTargetState);
 
 	bool bContinue = true;
 	int32 PreviousFrameIndex = INDEX_NONE;
 	FStateTreeTransitionResult CurrentTransition = Transition;
 
-	auto ExitPreviousFrame = [this, &Exec, &Args, &CurrentTransition](const int32 FrameIndex, const bool bCallExit)
+	auto ExitPreviousFrame = [this, &Exec, &ArgsSelectedFrames, &CurrentTransition](const int32 FrameIndex, const bool bCallExit)
 		{
 			FStateTreeExecutionFrame& Frame = Exec.ActiveFrames[FrameIndex];
-			const bool bIsFrameCommon = Args.SelectedFrames.Contains(Frame.FrameID);
+			const bool bIsFrameCommon = ArgsSelectedFrames.Contains(Frame.FrameID);
 			if (Frame.bIsGlobalFrame && (bCallExit || !bIsFrameCommon))
 			{
 				const bool bIsFrameSustained = bIsFrameCommon;
@@ -3534,7 +3576,7 @@ void FStateTreeExecutionContext::ExitState(const FSelectStateResult& Args, const
 			}
 
 			// It is the target.
-			if (bIsStateCommon && bHasSustained && CurrentActiveState == Args.TargetState)
+			if (bIsStateCommon && bHasSustained && CurrentActiveState == ArgsTargetState)
 			{
 				// this is the last sustained
 				bContinue = false;
@@ -4905,13 +4947,14 @@ bool FStateTreeExecutionContext::RequestTransitionInternal(
 	SelectStateArgs.Fallback = Fallback;
 	SelectStateArgs.Behavior = ESelectStateBehavior::StateTransition;
 
-	FSelectStateResult StateSelectionResult;
+	TSharedRef<FSelectStateResult> StateSelectionResult = MakeShared<FSelectStateResult>();
+	TGuardValue<TSharedPtr<FSelectStateResult>> ExecutionFrameHolderScope(CurrentlyProcessedTemporaryStorage, StateSelectionResult);
 	if (SelectState(SelectStateArgs, StateSelectionResult))
 	{
 		const FStateTreeExecutionState& Exec = GetExecState();
 
 		SetupNextTransition(SourceFrame, SourceStateID, TargetState, Priority);
-		RequestedTransition->Selection = MoveTemp(StateSelectionResult);
+		RequestedTransition->Selection = StateSelectionResult; // It will keep the temporary storage alive until the end of the EnterState
 
 		// Fill NextActiveFrames & NextActiveFrameEvents for backward compatibility
 		if (bSetDeprecatedTransitionResultProperties)
@@ -4923,12 +4966,12 @@ PRAGMA_DISABLE_DEPRECATION_WARNINGS
 			FStateTreeExecutionFrame* CurrentFrame = nullptr;
 			FStateTreeFrameStateSelectionEvents* CurrentSelectionEvents = nullptr;
 			int32 CurrentSelectionEventsIndex = 0;
-			for (const FActiveState& SelectedState : RequestedTransition->Selection.SelectedStates)
+			for (const FActiveState& SelectedState : RequestedTransition->Selection->SelectedStates)
 			{
 				if (SelectedState.GetFrameID() != PreviousFrameID)
 				{
 					PreviousFrameID = SelectedState.GetFrameID();
-					const FStateTreeExecutionFrame* Frame = FindExecutionFrame(SelectedState.GetFrameID(), MakeConstArrayView(Exec.ActiveFrames), MakeConstArrayView(RequestedTransition->Selection.TemporaryFrames));
+					const FStateTreeExecutionFrame* Frame = FindExecutionFrame(SelectedState.GetFrameID(), MakeConstArrayView(Exec.ActiveFrames), MakeConstArrayView(RequestedTransition->Selection->TemporaryFrames));
 					if (ensure(Frame))
 					{
 						RequestedTransition->Transition.NextActiveFrames.Add(*Frame);
@@ -4945,7 +4988,7 @@ PRAGMA_DISABLE_DEPRECATION_WARNINGS
 				check(CurrentSelectionEvents);
 				CurrentFrame->ActiveStates.Push(SelectedState.GetStateHandle(), SelectedState.GetStateID());
 
-				const FSelectionEventWithID* FoundEvent = RequestedTransition->Selection.SelectionEvents.FindByPredicate(
+				const FSelectionEventWithID* FoundEvent = RequestedTransition->Selection->SelectionEvents.FindByPredicate(
 					[SelectedState](const FSelectionEventWithID& Event)
 					{
 						return Event.State == SelectedState;
@@ -4960,9 +5003,9 @@ PRAGMA_ENABLE_DEPRECATION_WARNINGS
 
 		//@todo In TriggerTransitions, the events are processed bottom up. If an event with a higher priority is process after, the event will already be consumed.
 		// Consume events from states, if required. The transition might also want to consume the event. (Not a bug.)
-		for (FSelectionEventWithID& Event : RequestedTransition->Selection.SelectionEvents)
+		for (FSelectionEventWithID& Event : RequestedTransition->Selection->SelectionEvents)
 		{
-			const FStateTreeExecutionFrame* Frame = FindExecutionFrame(Event.State.GetFrameID(), MakeConstArrayView(Exec.ActiveFrames), MakeConstArrayView(RequestedTransition->Selection.TemporaryFrames));
+			const FStateTreeExecutionFrame* Frame = FindExecutionFrame(Event.State.GetFrameID(), MakeConstArrayView(Exec.ActiveFrames), MakeConstArrayView(RequestedTransition->Selection->TemporaryFrames));
 			if (ensure(Frame))
 			{
 				const FCompactStateTreeState* State = Frame->StateTree->GetStateFromHandle(Event.State.GetStateHandle());
@@ -4975,9 +5018,9 @@ PRAGMA_ENABLE_DEPRECATION_WARNINGS
 
 		UE_SUPPRESS(LogStateTree, Verbose,
 		{
-			check(RequestedTransition->Selection.SelectedStates.Num() > 0);
-			const FActiveState& LastSelectedState = RequestedTransition->Selection.SelectedStates.Last();
-			const FStateTreeExecutionFrame* SelectedFrame = FindExecutionFrame(LastSelectedState.GetFrameID(), MakeConstArrayView(Exec.ActiveFrames), MakeConstArrayView(RequestedTransition->Selection.TemporaryFrames));
+			check(RequestedTransition->Selection->SelectedStates.Num() > 0);
+			const FActiveState& LastSelectedState = RequestedTransition->Selection->SelectedStates.Last();
+			const FStateTreeExecutionFrame* SelectedFrame = FindExecutionFrame(LastSelectedState.GetFrameID(), MakeConstArrayView(Exec.ActiveFrames), MakeConstArrayView(RequestedTransition->Selection->TemporaryFrames));
 			check(SelectedFrame);
 			STATETREE_LOG(Verbose, TEXT("Transition Request. Source:'%s'. Target:'%s'. Selected:'%s'"),
 				*GetSafeStateName(SourceFrame, SourceStateHandle),
@@ -5825,7 +5868,7 @@ PRAGMA_ENABLE_DEPRECATION_WARNINGS
 	return Result;
 }
 
-FRecordedStateTreeTransitionResult FStateTreeExecutionContext::MakeRecordedTransitionResult(const FSelectStateResult& SelectStateResult, const FStateTreeTransitionResult& TransitionResult) const
+FRecordedStateTreeTransitionResult FStateTreeExecutionContext::MakeRecordedTransitionResult(const TSharedRef<FSelectStateResult>& SelectStateResult, const FStateTreeTransitionResult& TransitionResult) const
 {
 	using namespace UE::StateTree;
 	using namespace UE::StateTree::ExecutionContext;
@@ -5834,10 +5877,10 @@ FRecordedStateTreeTransitionResult FStateTreeExecutionContext::MakeRecordedTrans
 	const FStateTreeExecutionState& Exec = GetExecState();
 
 	FRecordedStateTreeTransitionResult Recorded;
-	Recorded.States.Reserve(SelectStateResult.SelectedStates.Num());
-	for (const UE::StateTree::FActiveState& SelectedState : SelectStateResult.SelectedStates)
+	Recorded.States.Reserve(SelectStateResult->SelectedStates.Num());
+	for (const UE::StateTree::FActiveState& SelectedState : SelectStateResult->SelectedStates)
 	{
-		const FStateTreeExecutionFrame* ProcessFrame = FindExecutionFrame(SelectedState.GetFrameID(), MakeConstArrayView(Exec.ActiveFrames), MakeConstArrayView(SelectStateResult.TemporaryFrames));
+		const FStateTreeExecutionFrame* ProcessFrame = FindExecutionFrame(SelectedState.GetFrameID(), MakeConstArrayView(Exec.ActiveFrames), MakeConstArrayView(SelectStateResult->TemporaryFrames));
 		if (ProcessFrame == nullptr)
 		{
 			return FRecordedStateTreeTransitionResult();
@@ -5847,13 +5890,13 @@ FRecordedStateTreeTransitionResult FStateTreeExecutionContext::MakeRecordedTrans
 		RecordedState.State = SelectedState.GetStateHandle();
 	}
 
-	Recorded.Events.Reserve(SelectStateResult.SelectionEvents.Num());
-	for (const FSelectionEventWithID& SelectedEvent : SelectStateResult.SelectionEvents)
+	Recorded.Events.Reserve(SelectStateResult->SelectionEvents.Num());
+	for (const FSelectionEventWithID& SelectedEvent : SelectStateResult->SelectionEvents)
 	{
 		if (SelectedEvent.Event.IsValid())
 		{
 			const int32 RecordedEventIndex = Recorded.Events.Add(*SelectedEvent.Event.Get());
-			if (const int32 SeletedStateIndex = FActiveStatePath::IndexOf(SelectStateResult.SelectedStates, SelectedEvent.State); SeletedStateIndex != INDEX_NONE)
+			if (const int32 SeletedStateIndex = FActiveStatePath::IndexOf(SelectStateResult->SelectedStates, SelectedEvent.State); SeletedStateIndex != INDEX_NONE)
 			{
 				Recorded.States[SeletedStateIndex].EventIndex = RecordedEventIndex;
 			}
@@ -5887,12 +5930,13 @@ bool FStateTreeExecutionContext::SelectState(const FStateTreeExecutionFrame& Cur
 	SelectStateArgs.TransitionEvent = TransitionEvent != nullptr ? *TransitionEvent : FStateTreeSharedEvent();
 	SelectStateArgs.Fallback = Fallback;
 
-	FSelectStateResult StateSelection;
-	return SelectState(SelectStateArgs, StateSelection);
+	TSharedRef<FSelectStateResult> SelectStateResult = MakeShared<FSelectStateResult>();
+	TGuardValue<TSharedPtr<FSelectStateResult>> ExecutionFrameHolderScope(CurrentlyProcessedTemporaryStorage, SelectStateResult);
+	return SelectState(SelectStateArgs, SelectStateResult);
 }
 PRAGMA_ENABLE_DEPRECATION_WARNINGS
 
-bool FStateTreeExecutionContext::SelectState(const FSelectStateArguments& Args, FSelectStateResult& OutSelectionResult)
+bool FStateTreeExecutionContext::SelectState(const FSelectStateArguments& Args, const TSharedRef<FSelectStateResult>& OutSelectionResult)
 {
 	// 1. Find the target root state.
 	// 2. If the the target frame is active, then copy all previous states from the previous frames.
@@ -6030,14 +6074,14 @@ bool FStateTreeExecutionContext::SelectState(const FSelectStateArguments& Args,
 		}
 	}
 	checkf(SourceStates.Num() == 0 || FActiveStatePath::StartsWith(Args.ActiveStates, SourceStates), TEXT("Source is part of the active path."));
-	FStateTreeExecutionFrame* SourceExecFrame = FindExecutionFrame(Args.SourceState.GetFrameID(), MakeArrayView(Exec.ActiveFrames), MakeArrayView(OutSelectionResult.TemporaryFrames));
+	FStateTreeExecutionFrame* SourceExecFrame = FindExecutionFrame(Args.SourceState.GetFrameID(), MakeArrayView(Exec.ActiveFrames), MakeArrayView(OutSelectionResult->TemporaryFrames));
 	if (!ensure(SourceExecFrame))
 	{
 		return false;
 	}
 
 	const int32 SelectedStatesSlack = Args.ActiveStates.Num() + 10; // Number to help with buffer reallocation
-	OutSelectionResult.SelectedStates.Empty(SelectedStatesSlack);
+	OutSelectionResult->SelectedStates.Empty(SelectedStatesSlack);
 
 	// Find the Target inside the ActiveStates.
 	//Look in the SourceStates because you can select in a different frame up the tree but not down the tree.
@@ -6058,9 +6102,9 @@ bool FStateTreeExecutionContext::SelectState(const FSelectStateArguments& Args,
 					{
 						break;
 					}
-					OutSelectionResult.SelectedFrames.Add(CurrentFrameID);
+					OutSelectionResult->SelectedFrames.Add(CurrentFrameID);
 				}
-				OutSelectionResult.SelectedStates.Add(ActiveState);
+				OutSelectionResult->SelectedStates.Add(ActiveState);
 			}
 		}
 		else if (SourceStates.Num() >= 0)
@@ -6077,7 +6121,7 @@ bool FStateTreeExecutionContext::SelectState(const FSelectStateArguments& Args,
 					const FActiveFrameID CurrentActiveFrameID = SourceStates[FoundActiveStateIndex].GetFrameID();
 					if (CurrentActiveFrameID != CurrentFrameID)
 					{
-						FStateTreeExecutionFrame* ProcessFrame = FindExecutionFrame(CurrentActiveFrameID, MakeArrayView(Exec.ActiveFrames), MakeArrayView(OutSelectionResult.TemporaryFrames));
+						FStateTreeExecutionFrame* ProcessFrame = FindExecutionFrame(CurrentActiveFrameID, MakeArrayView(Exec.ActiveFrames), MakeArrayView(OutSelectionResult->TemporaryFrames));
 						if (!ensure(ProcessFrame))
 						{
 							return false;
@@ -6121,9 +6165,9 @@ bool FStateTreeExecutionContext::SelectState(const FSelectStateArguments& Args,
 					if (CurrentFrameID != ActiveState.GetFrameID())
 					{
 						CurrentFrameID = ActiveState.GetFrameID();
-						OutSelectionResult.SelectedFrames.Add(CurrentFrameID);
+						OutSelectionResult->SelectedFrames.Add(CurrentFrameID);
 					}
-					OutSelectionResult.SelectedStates.Add(ActiveState);
+					OutSelectionResult->SelectedStates.Add(ActiveState);
 				}
 			}
 		}
@@ -6142,8 +6186,8 @@ bool FStateTreeExecutionContext::SelectState(const FSelectStateArguments& Args,
 	}
 
 	FSelectStateInternalArguments InternalArgs;
-	InternalArgs.MissingActiveStates = Args.ActiveStates.Mid(OutSelectionResult.SelectedStates.Num());
-	InternalArgs.MissingSourceStates = SourceStates.Mid(OutSelectionResult.SelectedStates.Num());
+	InternalArgs.MissingActiveStates = Args.ActiveStates.Mid(OutSelectionResult->SelectedStates.Num());
+	InternalArgs.MissingSourceStates = SourceStates.Mid(OutSelectionResult->SelectedStates.Num());
 	InternalArgs.MissingStatesToReachTarget = PathToTargetState;
 	InternalArgs.MissingSourceFrameID = TargetFrameID;
 
@@ -6199,7 +6243,7 @@ PRAGMA_ENABLE_DEPRECATION_WARNINGS
 bool FStateTreeExecutionContext::SelectStateInternal(
 	const FSelectStateArguments& Args,
 	const FSelectStateInternalArguments& InternalArgs,
-	FSelectStateResult& OutSelectionResult)
+	const TSharedRef<FSelectStateResult>& OutSelectionResult)
 {
 	CSV_SCOPED_TIMING_STAT_EXCLUSIVE(StateTree_SelectState);
 
@@ -6250,7 +6294,7 @@ bool FStateTreeExecutionContext::SelectStateInternal(
 		if (InternalArgs.MissingSourceStates.Num() == 0 && !InternalArgs.MissingSourceFrameID.IsValid())
 		{
 			// Must be in a valid selection (state are already been selected).
-			if (OutSelectionResult.SelectedStates.Num() == 0)
+			if (OutSelectionResult->SelectedStates.Num() == 0)
 			{
 				STATETREE_LOG(Verbose, TEXT("%hs: No root state to select '%s'.  '%s' using StateTree '%s'."),
 					__FUNCTION__,
@@ -6261,16 +6305,16 @@ bool FStateTreeExecutionContext::SelectStateInternal(
 				return false;
 			}
 
-			const FActiveState PreviousSelectedState = OutSelectionResult.SelectedStates.Last();
-			FStateTreeExecutionFrame* PreviousFramePtr = FindExecutionFrame(PreviousSelectedState.GetFrameID(), MakeArrayView(GetExecState().ActiveFrames), MakeArrayView(OutSelectionResult.TemporaryFrames));
+			const FActiveState PreviousSelectedState = OutSelectionResult->SelectedStates.Last();
+			FStateTreeExecutionFrame* PreviousFramePtr = FindExecutionFrame(PreviousSelectedState.GetFrameID(), MakeArrayView(GetExecState().ActiveFrames), MakeArrayView(OutSelectionResult->TemporaryFrames));
 			check(PreviousFramePtr);
 			check(PreviousFramePtr->StateTree == NextStateTree);
 
 			if (PreviousSelectedState.GetStateHandle() != NextState.Parent)
 			{
 				const FExecutionFrameKey FrameKey = FExecutionFrameKey(NextStateContext.StateTree, NextStateContext.StateHandle);
-				const FStateTreeExecutionFrame& NewFrame = OutSelectionResult.MakeAndAddTemporaryFrameWithNewRoot(FActiveFrameID(Storage.GenerateUniqueId()), FrameKey, *PreviousFramePtr);
-				OutSelectionResult.SelectedFrames.Add(NewFrame.FrameID);
+				const FStateTreeExecutionFrame& NewFrame = OutSelectionResult->MakeAndAddTemporaryFrameWithNewRoot(FActiveFrameID(Storage.GenerateUniqueId()), FrameKey, *PreviousFramePtr);
+				OutSelectionResult->SelectedFrames.Add(NewFrame.FrameID);
 				bNewFrameCreated = true;
 			}
 		}
@@ -6278,7 +6322,7 @@ bool FStateTreeExecutionContext::SelectStateInternal(
 		{
 			// We are building the path toward the desired target.
 			const FActiveFrameID MissingFrameID = InternalArgs.MissingSourceStates.Num() != 0 ? InternalArgs.MissingSourceStates[0].GetFrameID() : InternalArgs.MissingSourceFrameID;
-			FStateTreeExecutionFrame* MissingFramePtr = FindExecutionFrame(MissingFrameID, MakeArrayView(GetExecState().ActiveFrames), MakeArrayView(OutSelectionResult.TemporaryFrames));
+			FStateTreeExecutionFrame* MissingFramePtr = FindExecutionFrame(MissingFrameID, MakeArrayView(GetExecState().ActiveFrames), MakeArrayView(OutSelectionResult->TemporaryFrames));
 			check(MissingFramePtr);
 			if (!ensure(MissingFramePtr->StateTree == NextStateTree))
 			{
@@ -6291,7 +6335,7 @@ bool FStateTreeExecutionContext::SelectStateInternal(
 				if (!NextState.Parent.IsValid() && MissingFramePtr->RootState != NextStateHandle)
 				{
 					const FExecutionFrameKey FrameKey = FExecutionFrameKey(NextStateContext.StateTree, NextStateContext.StateHandle);
-					const FStateTreeExecutionFrame& NewFrame = OutSelectionResult.MakeAndAddTemporaryFrameWithNewRoot(FActiveFrameID(Storage.GenerateUniqueId()), FrameKey, *MissingFramePtr);
+					const FStateTreeExecutionFrame& NewFrame = OutSelectionResult->MakeAndAddTemporaryFrameWithNewRoot(FActiveFrameID(Storage.GenerateUniqueId()), FrameKey, *MissingFramePtr);
 					NextFrameID = NewFrame.FrameID;
 					bNewFrameCreated = true;
 				}
@@ -6299,22 +6343,22 @@ bool FStateTreeExecutionContext::SelectStateInternal(
 			else if (MissingFramePtr->RootState != NextStateHandle)
 			{
 				const FExecutionFrameKey FrameKey = FExecutionFrameKey(NextStateContext.StateTree, NextStateContext.StateHandle);
-				const FStateTreeExecutionFrame& NewFrame = OutSelectionResult.MakeAndAddTemporaryFrameWithNewRoot(FActiveFrameID(Storage.GenerateUniqueId()), FrameKey, *MissingFramePtr);
+				const FStateTreeExecutionFrame& NewFrame = OutSelectionResult->MakeAndAddTemporaryFrameWithNewRoot(FActiveFrameID(Storage.GenerateUniqueId()), FrameKey, *MissingFramePtr);
 				NextFrameID = NewFrame.FrameID;
 				bNewFrameCreated = true;
 			}
 
 			// Add new frames.
-			if (OutSelectionResult.SelectedFrames.Num() > 0)
+			if (OutSelectionResult->SelectedFrames.Num() > 0)
 			{
-				if (OutSelectionResult.SelectedFrames.Last() != NextFrameID)
+				if (OutSelectionResult->SelectedFrames.Last() != NextFrameID)
 				{
-					OutSelectionResult.SelectedFrames.Add(NextFrameID);
+					OutSelectionResult->SelectedFrames.Add(NextFrameID);
 				}
 			}
 			else
 			{
-				OutSelectionResult.SelectedFrames.Add(NextFrameID);
+				OutSelectionResult->SelectedFrames.Add(NextFrameID);
 			}
 
 			TOptional<bool> bSelected = SelectStateFromSourceInternal(Args, InternalArgs, OutSelectionResult, *MissingFramePtr, NextState, NextStateHandle, bNewFrameCreated);
@@ -6326,21 +6370,21 @@ bool FStateTreeExecutionContext::SelectStateInternal(
 	}
 
 	// Use the ID to get the ExecutionFrame because it is saved in an array (and this is a recursive function), the array can grow and the pointer won't be valid.
-	check(OutSelectionResult.SelectedFrames.Num() != 0);
-	UE::StateTree::FActiveFrameID NextFrameID = OutSelectionResult.SelectedFrames.Last();
-	const FStateTreeExecutionFrame* NextFramePtr = FindExecutionFrame(NextFrameID, MakeConstArrayView(Exec.ActiveFrames), MakeConstArrayView(OutSelectionResult.TemporaryFrames));
+	check(OutSelectionResult->SelectedFrames.Num() != 0);
+	UE::StateTree::FActiveFrameID NextFrameID = OutSelectionResult->SelectedFrames.Last();
+	const FStateTreeExecutionFrame* NextFramePtr = FindExecutionFrame(NextFrameID, MakeConstArrayView(Exec.ActiveFrames), MakeConstArrayView(OutSelectionResult->TemporaryFrames));
 	check(NextFramePtr);
 
 	const FStateTreeExecutionFrame& NextFrame = *NextFramePtr;
 	const FStateTreeExecutionFrame* NextParentFrame = nullptr;
-	if (OutSelectionResult.SelectedFrames.Num() > 1)
+	if (OutSelectionResult->SelectedFrames.Num() > 1)
 	{
-		NextParentFrame = FindExecutionFrame(OutSelectionResult.SelectedFrames.Last(1), MakeConstArrayView(Exec.ActiveFrames), MakeConstArrayView(OutSelectionResult.TemporaryFrames));
+		NextParentFrame = FindExecutionFrame(OutSelectionResult->SelectedFrames.Last(1), MakeConstArrayView(Exec.ActiveFrames), MakeConstArrayView(OutSelectionResult->TemporaryFrames));
 		check(NextParentFrame);
 	}
 
 	// Save the current result to use SelectionEvents in GetDataView
-	TGuardValue<FSelectStateResult*> GuardCurrentlyProcessedStateSelectionResult(CurrentlyProcessedStateSelectionResult, &OutSelectionResult);
+	TGuardValue<FSelectStateResult*> GuardCurrentlyProcessedStateSelectionResult(CurrentlyProcessedStateSelectionResult, &OutSelectionResult.Get());
 
 	UE_STATETREE_DEBUG_SCOPED_STATE_PHASE(this, NextStateHandle, EStateTreeUpdatePhase::StateSelection);
 
@@ -6486,19 +6530,19 @@ bool FStateTreeExecutionContext::SelectStateInternal(
 	ON_SCOPE_EXIT{
 		if (bSucceededToSelectState == false)
 		{
-			OutSelectionResult.SelectedStates.Pop();
+			OutSelectionResult->SelectedStates.Pop();
 		}
 	};
 
 	// Add the state to the selected list. It will be removed if the selection fails.
 	if (bShouldCreateNewState)
 	{
-		OutSelectionResult.SelectedStates.Emplace(NextFrame.FrameID, FActiveStateID(Storage.GenerateUniqueId()), NextStateHandle);
+		OutSelectionResult->SelectedStates.Emplace(NextFrame.FrameID, FActiveStateID(Storage.GenerateUniqueId()), NextStateHandle);
 
 		int32 StateInFrameCounter = 0;
-		for (int32 Index = OutSelectionResult.SelectedStates.Num() - 1; Index >= 0; --Index)
+		for (int32 Index = OutSelectionResult->SelectedStates.Num() - 1; Index >= 0; --Index)
 		{
-			if (OutSelectionResult.SelectedStates[Index].GetFrameID() != NextFrame.FrameID)
+			if (OutSelectionResult->SelectedStates[Index].GetFrameID() != NextFrame.FrameID)
 			{
 				break;
 			}
@@ -6520,20 +6564,20 @@ bool FStateTreeExecutionContext::SelectStateInternal(
 	else
 	{
 		check(InternalArgs.MissingActiveStates.Num() > 0); // bIsNextTargetStateInActiveStates is false
-		OutSelectionResult.SelectedStates.Add(InternalArgs.MissingActiveStates[0]);
+		OutSelectionResult->SelectedStates.Add(InternalArgs.MissingActiveStates[0]);
 	}
 
 	// Set the target, since this is recursive (can be called for linked state), it only set once the original target is found.
-	if (bIsTargetState && !OutSelectionResult.TargetState.IsValid())
+	if (bIsTargetState && !OutSelectionResult->TargetState.IsValid())
 	{
-		OutSelectionResult.TargetState = OutSelectionResult.SelectedStates.Last();
+		OutSelectionResult->TargetState = OutSelectionResult->SelectedStates.Last();
 	}
 
 	for (const FStateTreeSharedEvent& StateSelectionEvent : StateSelectionEvents)
 	{
 		// Add the event for GetDataView and for consuming it later.
-		const bool bContainsSelectionEvent = OutSelectionResult.SelectionEvents.ContainsByPredicate(
-			[&LatestState = OutSelectionResult.SelectedStates.Last()](const FSelectionEventWithID& Other)
+		const bool bContainsSelectionEvent = OutSelectionResult->SelectionEvents.ContainsByPredicate(
+			[&LatestState = OutSelectionResult->SelectedStates.Last()](const FSelectionEventWithID& Other)
 			{
 				return Other.State == LatestState;
 			});
@@ -6543,17 +6587,17 @@ bool FStateTreeExecutionContext::SelectStateInternal(
 		if (StateSelectionEvent.IsValid())
 		{
 			
-			OutSelectionResult.SelectionEvents.Add(FSelectionEventWithID{.State = OutSelectionResult.SelectedStates.Last(), .Event = StateSelectionEvent});
+			OutSelectionResult->SelectionEvents.Add(FSelectionEventWithID{.State = OutSelectionResult->SelectedStates.Last(), .Event = StateSelectionEvent});
 			bRemoveSelectionEvents = true;
 		}
 
-		auto RemoveStateSelectionEvent = [bRemoveSelectionEvents](FSelectStateResult& SelectionResult)
+		auto RemoveStateSelectionEvent = [bRemoveSelectionEvents](const TSharedRef<FSelectStateResult>& SelectionResult)
 			{
 				if (bRemoveSelectionEvents)
 				{
-					check(SelectionResult.SelectionEvents.Num() > 0);
-					ensureMsgf(SelectionResult.SelectionEvents.Last().State == SelectionResult.SelectedStates.Last(), TEXT("We should remove from the same element it was added."));
-					SelectionResult.SelectionEvents.Pop();
+					check(SelectionResult->SelectionEvents.Num() > 0);
+					ensureMsgf(SelectionResult->SelectionEvents.Last().State == SelectionResult->SelectedStates.Last(), TEXT("We should remove from the same element it was added."));
+					SelectionResult->SelectionEvents.Pop();
 				}
 			};
 
@@ -6648,7 +6692,7 @@ bool FStateTreeExecutionContext::SelectStateInternal(
 TOptional<bool> FStateTreeExecutionContext::SelectStateFromSourceInternal(
 	const FSelectStateArguments& Args,
 	const FSelectStateInternalArguments& InternalArgs,
-	FSelectStateResult& OutSelectionResult,
+	const TSharedRef<FSelectStateResult>& OutSelectionResult,
 	const FStateTreeExecutionFrame& NextFrame,
 	const FCompactStateTreeState& NextState,
 	const FStateTreeStateHandle NextStateHandle,
@@ -6688,13 +6732,13 @@ TOptional<bool> FStateTreeExecutionContext::SelectStateFromSourceInternal(
 
 		if (!bContinueWithStateSelection)
 		{
-			OutSelectionResult.SelectedStates.Add(InternalArgs.MissingSourceStates[0]);
+			OutSelectionResult->SelectedStates.Add(InternalArgs.MissingSourceStates[0]);
 
 			bool bSelectStateInternalSucceeded = false;
 			ON_SCOPE_EXIT{
 				if (bSelectStateInternalSucceeded == false)
 				{
-					OutSelectionResult.SelectedStates.Pop();
+					OutSelectionResult->SelectedStates.Pop();
 				}
 			};
 
@@ -6712,7 +6756,7 @@ TOptional<bool> FStateTreeExecutionContext::SelectStateFromSourceInternal(
 				// Before the source, there's a completed state.
 				//The completed state should fail the transition(because of EStateTreeStateSelectionRules::CompletedStateBeforeTargetFailsTransition).
 				//The transition can contain a state with "follow transition" that would remove this state from the transition.
-				if (OutSelectionResult.SelectedStates.Contains(MissingActiveState))
+				if (OutSelectionResult->SelectedStates.Contains(MissingActiveState))
 				{
 					STATETREE_LOG(VeryVerbose, TEXT("%hs: Selection fails because state '%s' is completed.  '%s' using StateTree '%s'."),
 						__FUNCTION__,
@@ -6781,7 +6825,7 @@ namespace UE::StateTree::ExecutionContext::Private
 bool FStateTreeExecutionContext::SelectStateInternal_Linked(
 	const FSelectStateArguments& Args,
 	const FSelectStateInternalArguments& InternalArgs,
-	FSelectStateResult& OutSelectionResult,
+	const TSharedRef<FSelectStateResult>& OutSelectionResult,
 	TNotNull<const UStateTree*> NextStateTree,
 	const FCompactStateTreeState& TargetState,
 	bool bShouldCreateNewState)
@@ -6822,7 +6866,7 @@ bool FStateTreeExecutionContext::SelectStateInternal_Linked(
 		}
 	}
 
-	if (PreventRecursionCheck(LinkStateFrameKey, OutSelectionResult.SelectedFrames, MakeConstArrayView(Exec.ActiveFrames), MakeConstArrayView(OutSelectionResult.TemporaryFrames)))
+	if (PreventRecursionCheck(LinkStateFrameKey, OutSelectionResult->SelectedFrames, MakeConstArrayView(Exec.ActiveFrames), MakeConstArrayView(OutSelectionResult->TemporaryFrames)))
 	{
 		STATETREE_LOG(Error, TEXT("%hs: Trying to recursively enter subtree '%s' from '%s'. '%s' using StateTree '%s'."),
 			__FUNCTION__,
@@ -6876,7 +6920,7 @@ bool FStateTreeExecutionContext::SelectStateInternal_Linked(
 		FStateTreeDataHandle GlobalParameterDataHandle;
 		FStateTreeIndex16 GlobalInstanceIndexBase;
 		{
-			const FStateTreeExecutionFrame* NextFrame = FindExecutionFrame(OutSelectionResult.SelectedFrames.Last(), MakeConstArrayView(Exec.ActiveFrames), MakeConstArrayView(OutSelectionResult.TemporaryFrames));
+			const FStateTreeExecutionFrame* NextFrame = FindExecutionFrame(OutSelectionResult->SelectedFrames.Last(), MakeConstArrayView(Exec.ActiveFrames), MakeConstArrayView(OutSelectionResult->TemporaryFrames));
 			check(NextFrame);
 			ExternalDataBaseIndex = NextFrame->ExternalDataBaseIndex;
 			GlobalParameterDataHandle = NextFrame->GlobalParameterDataHandle;
@@ -6884,7 +6928,7 @@ bool FStateTreeExecutionContext::SelectStateInternal_Linked(
 		}
 
 		constexpr bool bIsGlobalFrame = false;
-		FStateTreeExecutionFrame& NewFrame = OutSelectionResult.MakeAndAddTemporaryFrame(FActiveFrameID(Storage.GenerateUniqueId()), LinkStateFrameKey, bIsGlobalFrame);
+		FStateTreeExecutionFrame& NewFrame = OutSelectionResult->MakeAndAddTemporaryFrame(FActiveFrameID(Storage.GenerateUniqueId()), LinkStateFrameKey, bIsGlobalFrame);
 		NewFrame.ExternalDataBaseIndex = ExternalDataBaseIndex;
 		NewFrame.GlobalParameterDataHandle = GlobalParameterDataHandle;
 		NewFrame.GlobalInstanceIndexBase = GlobalInstanceIndexBase;
@@ -6893,7 +6937,7 @@ bool FStateTreeExecutionContext::SelectStateInternal_Linked(
 		SelectedFrame = &NewFrame;
 	}
 
-	OutSelectionResult.SelectedFrames.Add(SelectedFrame->FrameID);
+	OutSelectionResult->SelectedFrames.Add(SelectedFrame->FrameID);
 
 	// Select the root state of the new frame.
 	FStateHandleContext RootState = FStateHandleContext(LinkStateFrameKey.StateTree, LinkStateFrameKey.RootState);
@@ -6907,7 +6951,7 @@ bool FStateTreeExecutionContext::SelectStateInternal_Linked(
 
 	if (!bSucceededToSelectState)
 	{
-		OutSelectionResult.SelectedFrames.Pop();
+		OutSelectionResult->SelectedFrames.Pop();
 	}
 
 	return bSucceededToSelectState;
@@ -6916,7 +6960,7 @@ bool FStateTreeExecutionContext::SelectStateInternal_Linked(
 bool FStateTreeExecutionContext::SelectStateInternal_LinkedAsset(
 	const FSelectStateArguments& Args,
 	const FSelectStateInternalArguments& InternalArgs,
-	FSelectStateResult& OutSelectionResult,
+	const TSharedRef<FSelectStateResult>& OutSelectionResult,
 	TNotNull<const UStateTree*> NextStateTree,
 	const FCompactStateTreeState& NextState,
 	const UStateTree* NextLinkedStateAsset,
@@ -6979,7 +7023,7 @@ bool FStateTreeExecutionContext::SelectStateInternal_LinkedAsset(
 		return false;
 	}
 
-	if (PreventRecursionCheck(LinkStateFrameKey, OutSelectionResult.SelectedFrames, MakeConstArrayView(Exec.ActiveFrames), MakeConstArrayView(OutSelectionResult.TemporaryFrames)))
+	if (PreventRecursionCheck(LinkStateFrameKey, OutSelectionResult->SelectedFrames, MakeConstArrayView(Exec.ActiveFrames), MakeConstArrayView(OutSelectionResult->TemporaryFrames)))
 	{
 		STATETREE_LOG(Error, TEXT("%hs: Trying to recursively enter subtree '%s' from '%s'. '%s' using StateTree '%s'."),
 			__FUNCTION__,
@@ -7043,7 +7087,7 @@ bool FStateTreeExecutionContext::SelectStateInternal_LinkedAsset(
 			return false;
 		}
 
-		FStateTreeExecutionFrame& NewFrame = OutSelectionResult.MakeAndAddTemporaryFrame(FActiveFrameID(Storage.GenerateUniqueId()), LinkStateFrameKey, true);
+		FStateTreeExecutionFrame& NewFrame = OutSelectionResult->MakeAndAddTemporaryFrame(FActiveFrameID(Storage.GenerateUniqueId()), LinkStateFrameKey, true);
 		NewFrame.ExternalDataBaseIndex = ExternalDataBaseIndex;
 		// Pass the linked state's parameters as global parameters to the linked asset.
 		NewFrame.GlobalParameterDataHandle = NextState.ParameterDataHandle;
@@ -7054,13 +7098,13 @@ bool FStateTreeExecutionContext::SelectStateInternal_LinkedAsset(
 		SelectedFrame = &NewFrame;
 	}
 
-	OutSelectionResult.SelectedFrames.Add(SelectedFrame->FrameID);
+	OutSelectionResult->SelectedFrames.Add(SelectedFrame->FrameID);
 
 	if (bIsNewFrame)
 	{
 		// Start global tasks and evaluators temporarily, so that their data is available already during select.
-		const FStateTreeExecutionFrame* NextParentFrame = FindExecutionFrame(OutSelectionResult.SelectedFrames.Last(1), MakeConstArrayView(Exec.ActiveFrames), MakeConstArrayView(OutSelectionResult.TemporaryFrames));
-		FStateTreeExecutionFrame* NextFrame = FindExecutionFrame(OutSelectionResult.SelectedFrames.Last(), MakeArrayView(Exec.ActiveFrames), MakeArrayView(OutSelectionResult.TemporaryFrames));
+		const FStateTreeExecutionFrame* NextParentFrame = FindExecutionFrame(OutSelectionResult->SelectedFrames.Last(1), MakeConstArrayView(Exec.ActiveFrames), MakeConstArrayView(OutSelectionResult->TemporaryFrames));
+		FStateTreeExecutionFrame* NextFrame = FindExecutionFrame(OutSelectionResult->SelectedFrames.Last(), MakeArrayView(Exec.ActiveFrames), MakeArrayView(OutSelectionResult->TemporaryFrames));
 		check(NextParentFrame);
 		check(SelectedFrame == NextFrame);
 
@@ -7078,7 +7122,7 @@ bool FStateTreeExecutionContext::SelectStateInternal_LinkedAsset(
 			StopTemporaryEvaluatorsAndGlobalTasks(NextParentFrame, *NextFrame, StartResult);
 			CleanFrame(Exec, SelectedFrame->FrameID);
 
-			OutSelectionResult.SelectedFrames.Pop();
+			OutSelectionResult->SelectedFrames.Pop();
 			return false;
 		}
 	}
@@ -7097,8 +7141,8 @@ bool FStateTreeExecutionContext::SelectStateInternal_LinkedAsset(
 	{
 		if (bIsNewFrame)
 		{
-			const FStateTreeExecutionFrame* NextParentFrame = FindExecutionFrame(OutSelectionResult.SelectedFrames.Last(1), MakeConstArrayView(Exec.ActiveFrames), MakeConstArrayView(OutSelectionResult.TemporaryFrames));
-			FStateTreeExecutionFrame* NextFrame = FindExecutionFrame(OutSelectionResult.SelectedFrames.Last(), MakeArrayView(Exec.ActiveFrames), MakeArrayView(OutSelectionResult.TemporaryFrames));
+			const FStateTreeExecutionFrame* NextParentFrame = FindExecutionFrame(OutSelectionResult->SelectedFrames.Last(1), MakeConstArrayView(Exec.ActiveFrames), MakeConstArrayView(OutSelectionResult->TemporaryFrames));
+			FStateTreeExecutionFrame* NextFrame = FindExecutionFrame(OutSelectionResult->SelectedFrames.Last(), MakeArrayView(Exec.ActiveFrames), MakeArrayView(OutSelectionResult->TemporaryFrames));
 			check(NextParentFrame);
 			check(NextFrame);
 
@@ -7106,7 +7150,7 @@ bool FStateTreeExecutionContext::SelectStateInternal_LinkedAsset(
 			StopTemporaryEvaluatorsAndGlobalTasks(NextParentFrame, *NextFrame, CompletionStatus);
 			CleanFrame(Exec, SelectedFrame->FrameID);
 		}
-		OutSelectionResult.SelectedFrames.Pop();
+		OutSelectionResult->SelectedFrames.Pop();
 	}
 
 	return bSucceededToSelectState;
@@ -7115,7 +7159,7 @@ bool FStateTreeExecutionContext::SelectStateInternal_LinkedAsset(
 bool FStateTreeExecutionContext::SelectStateInternal_TrySelectChildrenInOrder(
 	const FSelectStateArguments& Args,
 	const FSelectStateInternalArguments& InternalArgs,
-	FSelectStateResult& OutSelectionResult,
+	const TSharedRef<FSelectStateResult>& OutSelectionResult,
 	TNotNull<const UStateTree*> NextStateTree,
 	const FCompactStateTreeState& NextState,
 	const FStateTreeStateHandle NextStateHandle)
@@ -7152,7 +7196,7 @@ bool FStateTreeExecutionContext::SelectStateInternal_TrySelectChildrenInOrder(
 bool FStateTreeExecutionContext::SelectStateInternal_TrySelectChildrenAtRandom(
 	const FSelectStateArguments& Args,
 	const FSelectStateInternalArguments& InternalArgs,
-	FSelectStateResult& OutSelectionResult,
+	const TSharedRef<FSelectStateResult>& OutSelectionResult,
 	TNotNull<const UStateTree*> NextStateTree,
 	const FCompactStateTreeState& NextState,
 	const FStateTreeStateHandle NextStateHandle)
@@ -7198,7 +7242,7 @@ bool FStateTreeExecutionContext::SelectStateInternal_TrySelectChildrenAtRandom(
 bool FStateTreeExecutionContext::SelectStateInternal_TrySelectChildrenWithHighestUtility(
 	const FSelectStateArguments& Args,
 	const FSelectStateInternalArguments& InternalArgs,
-	FSelectStateResult& OutSelectionResult,
+	const TSharedRef<FSelectStateResult>& OutSelectionResult,
 	TNotNull<const UStateTree*> NextStateTree,
 	const FCompactStateTreeState& NextState,
 	const FStateTreeStateHandle NextStateHandle)
@@ -7218,10 +7262,10 @@ bool FStateTreeExecutionContext::SelectStateInternal_TrySelectChildrenWithHighes
 	NextLevelChildStates.Reserve(NextState.ChildrenEnd - NextState.ChildrenBegin);
 	{
 		using namespace UE::StateTree::ExecutionContext::Private;
-		const FStateTreeExecutionFrame* NextParentFrame = OutSelectionResult.SelectedFrames.Num() > 1
-			? FindExecutionFrame(OutSelectionResult.SelectedFrames.Last(1), MakeConstArrayView(GetExecState().ActiveFrames), MakeConstArrayView(OutSelectionResult.TemporaryFrames))
+		const FStateTreeExecutionFrame* NextParentFrame = OutSelectionResult->SelectedFrames.Num() > 1
+			? FindExecutionFrame(OutSelectionResult->SelectedFrames.Last(1), MakeConstArrayView(GetExecState().ActiveFrames), MakeConstArrayView(OutSelectionResult->TemporaryFrames))
 			: nullptr;
-		FStateTreeExecutionFrame* NextFrame = FindExecutionFrame(OutSelectionResult.SelectedFrames.Last(), MakeArrayView(GetExecState().ActiveFrames), MakeArrayView(OutSelectionResult.TemporaryFrames));
+		FStateTreeExecutionFrame* NextFrame = FindExecutionFrame(OutSelectionResult->SelectedFrames.Last(), MakeArrayView(GetExecState().ActiveFrames), MakeArrayView(OutSelectionResult->TemporaryFrames));
 		check(NextFrame);
 
 		for (uint16 ChildState = NextState.ChildrenBegin; ChildState < NextState.ChildrenEnd; ChildState = NextStateTree->States[ChildState].GetNextSibling())
@@ -7271,7 +7315,7 @@ bool FStateTreeExecutionContext::SelectStateInternal_TrySelectChildrenWithHighes
 bool FStateTreeExecutionContext::SelectStateInternal_TrySelectChildrenAtRandomWeightedByUtility(
 	const FSelectStateArguments& Args,
 	const FSelectStateInternalArguments& InternalArgs,
-	FSelectStateResult& OutSelectionResult,
+	const TSharedRef<FSelectStateResult>& OutSelectionResult,
 	TNotNull<const UStateTree*> NextStateTree,
 	const FCompactStateTreeState& NextState,
 	const FStateTreeStateHandle NextStateHandle)
@@ -7293,10 +7337,10 @@ bool FStateTreeExecutionContext::SelectStateInternal_TrySelectChildrenAtRandomWe
 	float TotalScore = 0.0f;
 	{
 		using namespace UE::StateTree::ExecutionContext::Private;
-		const FStateTreeExecutionFrame* NextParentFrame = OutSelectionResult.SelectedFrames.Num() > 1
-			? FindExecutionFrame(OutSelectionResult.SelectedFrames.Last(1), MakeConstArrayView(GetExecState().ActiveFrames), MakeConstArrayView(OutSelectionResult.TemporaryFrames))
+		const FStateTreeExecutionFrame* NextParentFrame = OutSelectionResult->SelectedFrames.Num() > 1
+			? FindExecutionFrame(OutSelectionResult->SelectedFrames.Last(1), MakeConstArrayView(GetExecState().ActiveFrames), MakeConstArrayView(OutSelectionResult->TemporaryFrames))
 			: nullptr;
-		FStateTreeExecutionFrame* NextFrame = FindExecutionFrame(OutSelectionResult.SelectedFrames.Last(), MakeArrayView(GetExecState().ActiveFrames), MakeArrayView(OutSelectionResult.TemporaryFrames));
+		FStateTreeExecutionFrame* NextFrame = FindExecutionFrame(OutSelectionResult->SelectedFrames.Last(), MakeArrayView(GetExecState().ActiveFrames), MakeArrayView(OutSelectionResult->TemporaryFrames));
 		check(NextFrame);
 
 		for (uint16 ChildState = NextState.ChildrenBegin; ChildState < NextState.ChildrenEnd; ChildState = NextStateTree->States[ChildState].GetNextSibling())
@@ -7347,7 +7391,7 @@ bool FStateTreeExecutionContext::SelectStateInternal_TrySelectChildrenAtRandomWe
 bool FStateTreeExecutionContext::SelectStateInternal_TryFollowTransitions(
 	const FSelectStateArguments& Args,
 	const FSelectStateInternalArguments& InternalArgs,
-	FSelectStateResult& OutSelectionResult,
+	const TSharedRef<FSelectStateResult>& OutSelectionResult,
 	TNotNull<const UStateTree*> NextStateTree,
 	const FCompactStateTreeState& NextState,
 	const FStateTreeStateHandle NextStateHandle)
@@ -7397,7 +7441,7 @@ bool FStateTreeExecutionContext::SelectStateInternal_TryFollowTransitions(
 		}
 
 		// Try to prevent (infinite) loops in the selection.
-		const bool bSelectionContainsState = OutSelectionResult.SelectedStates.ContainsByPredicate(
+		const bool bSelectionContainsState = OutSelectionResult->SelectedStates.ContainsByPredicate(
 			[State = Transition.State](const FActiveState& Other)
 			{
 				return State == Other.GetStateHandle();
@@ -7408,7 +7452,7 @@ bool FStateTreeExecutionContext::SelectStateInternal_TryFollowTransitions(
 				, __FUNCTION__
 				, *GetSafeStateName(NextStateTree, Transition.State)
 				, *GetStateStatusString(GetExecState())
-				, *GetStatePathAsString(&RootStateTree, OutSelectionResult.SelectedStates)
+				, *GetStatePathAsString(&RootStateTree, OutSelectionResult->SelectedStates)
 				, *Owner.GetName()
 				, *NextStateTree->GetFullName());
 			continue;
@@ -7422,10 +7466,10 @@ bool FStateTreeExecutionContext::SelectStateInternal_TryFollowTransitions(
 			bool bTransitionConditionsPassed = false;
 			{
 				using namespace UE::StateTree::ExecutionContext::Private;
-				const FStateTreeExecutionFrame* NextParentFrame = OutSelectionResult.SelectedFrames.Num() > 1
-					? FindExecutionFrame(OutSelectionResult.SelectedFrames.Last(1), MakeConstArrayView(GetExecState().ActiveFrames), MakeConstArrayView(OutSelectionResult.TemporaryFrames))
+				const FStateTreeExecutionFrame* NextParentFrame = OutSelectionResult->SelectedFrames.Num() > 1
+					? FindExecutionFrame(OutSelectionResult->SelectedFrames.Last(1), MakeConstArrayView(GetExecState().ActiveFrames), MakeConstArrayView(OutSelectionResult->TemporaryFrames))
 					: nullptr;
-				FStateTreeExecutionFrame* NextFrame = FindExecutionFrame(OutSelectionResult.SelectedFrames.Last(), MakeArrayView(GetExecState().ActiveFrames), MakeArrayView(OutSelectionResult.TemporaryFrames));
+				FStateTreeExecutionFrame* NextFrame = FindExecutionFrame(OutSelectionResult->SelectedFrames.Last(), MakeArrayView(GetExecState().ActiveFrames), MakeArrayView(OutSelectionResult->TemporaryFrames));
 				check(NextFrame);
 
 				FCurrentlyProcessedTransitionEventScope TransitionEventScope(*this, SelectedStateTransitionEvent.IsValid() ? SelectedStateTransitionEvent.Get() : nullptr);
@@ -7439,16 +7483,22 @@ bool FStateTreeExecutionContext::SelectStateInternal_TryFollowTransitions(
 				// Using SelectState() instead of SelectStateInternal to treat the transitions the same way as regular transitions,
 				// e.g. it may jump to a completely different branch.
 
+				FSelectStateResult CopySelectStateResult = OutSelectionResult.Get();
+
 				FSelectStateArguments SelectStateArgs;
-				SelectStateArgs.SourceState = Args.SourceState;
+				SelectStateArgs.ActiveStates = MakeConstArrayView(CopySelectStateResult.SelectedStates);
+				if (CopySelectStateResult.SelectedStates.Num() > 0)
+				{
+					SelectStateArgs.SourceState = CopySelectStateResult.SelectedStates.Last();
+				}
 				SelectStateArgs.TargetState = FStateHandleContext(NextStateTree, Transition.State);
 				SelectStateArgs.TransitionEvent = SelectedStateTransitionEvent;
 				SelectStateArgs.Fallback = Transition.Fallback;
-				FSelectStateResult StateSelectionResult;
-				StateSelectionResult.TemporaryFrames = OutSelectionResult.TemporaryFrames;
-				if (SelectState(SelectStateArgs, StateSelectionResult))
+
+				OutSelectionResult->SelectedStates.Reset();
+				OutSelectionResult->SelectedFrames.Reset();
+				if (SelectState(SelectStateArgs, OutSelectionResult))
 				{
-					OutSelectionResult = MoveTemp(StateSelectionResult);
 					CurrentPriority = Transition.Priority;
 					bSucceededToSelectState = true;
 
@@ -7462,6 +7512,11 @@ bool FStateTreeExecutionContext::SelectStateInternal_TryFollowTransitions(
 					// Cannot return because higher priority transitions may override the selection. 
 					break;
 				}
+				else
+				{
+					FSelectStateResult& SelectionResult = OutSelectionResult.Get();
+					SelectionResult = MoveTemp(CopySelectStateResult);
+				}
 			}
 		}
 	}
diff --git a/Engine/Plugins/Runtime/StateTree/Source/StateTreeModule/Public/StateTreeAsyncExecutionContext.h b/Engine/Plugins/Runtime/StateTree/Source/StateTreeModule/Public/StateTreeAsyncExecutionContext.h
--- a/Engine/Plugins/Runtime/StateTree/Source/StateTreeModule/Public/StateTreeAsyncExecutionContext.h
+++ b/Engine/Plugins/Runtime/StateTree/Source/StateTreeModule/Public/StateTreeAsyncExecutionContext.h
@@ -2,6 +2,7 @@
 
 #pragma once
 #include "StateTree.h"
+#include "StateTreeExecutionContextTypes.h"
 #include "StateTreeInstanceData.h"
 
 #define UE_API STATETREEMODULE_API
@@ -34,10 +35,10 @@ namespace UE::StateTree::Async
 
 		const FStateTreeNodeBase& GetNode() const;
 
-		const FStateTreeExecutionFrame* Frame = nullptr;
-		const FStateTreeExecutionFrame* ParentFrame = nullptr;
-		const FStateTreeStateHandle StateHandle;
-		const FStateTreeIndex16 NodeIndex;
+		FStateTreeExecutionFrame* Frame = nullptr;
+		FStateTreeExecutionFrame* ParentFrame = nullptr;
+		FStateTreeStateHandle StateHandle;
+		FStateTreeIndex16 NodeIndex;
 	};
 }
 
@@ -195,6 +196,7 @@ private:
 	TStrongObjectPtr<UObject> Owner;
 	TStrongObjectPtr<const UStateTree> StateTree;
 	TSharedPtr<FStateTreeInstanceStorage> Storage;
+	TSharedPtr<UE::StateTree::ExecutionContext::ITemporaryStorage> TemporaryStorage;
 
 	UE::StateTree::FActiveFrameID FrameID;
 	UE::StateTree::FActiveStateID StateID;
@@ -325,6 +327,7 @@ private:
 	TWeakObjectPtr<UObject> Owner;
 	TWeakObjectPtr<const UStateTree> StateTree;
 	TWeakPtr<FStateTreeInstanceStorage> Storage;
+	TWeakPtr<UE::StateTree::ExecutionContext::ITemporaryStorage> TemporaryStorage;
 
 	UE::StateTree::FActiveFrameID FrameID;
 	UE::StateTree::FActiveStateID StateID;
diff --git a/Engine/Plugins/Runtime/StateTree/Source/StateTreeModule/Public/StateTreeExecutionContext.h b/Engine/Plugins/Runtime/StateTree/Source/StateTreeModule/Public/StateTreeExecutionContext.h
--- a/Engine/Plugins/Runtime/StateTree/Source/StateTreeModule/Public/StateTreeExecutionContext.h
+++ b/Engine/Plugins/Runtime/StateTree/Source/StateTreeModule/Public/StateTreeExecutionContext.h
@@ -5,6 +5,7 @@
 #include "StateTree.h"
 #include "StateTreeExecutionExtension.h"
 #include "StateTreeExecutionTypes.h"
+#include "StateTreeExecutionContextTypes.h"
 #include "StateTreeNodeBase.h"
 #include "StateTreeNodeRef.h"
 #include "StateTreeReference.h"
@@ -38,21 +39,6 @@ namespace UE::StateTree::ExecutionContext
 	UE_API EStateTreeRunStatus CastToRunStatus(EStateTreeFinishTaskType FinishTask);
 	UE_API UE::StateTree::ETaskCompletionStatus CastToTaskStatus(EStateTreeRunStatus InStatus);
 	UE_API EStateTreeRunStatus CastToRunStatus(UE::StateTree::ETaskCompletionStatus InStatus);
-
-	/** Helper that combines the state handle and its tree asset. */
-	struct FStateHandleContext
-	{
-		FStateHandleContext() = default;
-		FStateHandleContext(TNotNull<const UStateTree*> InStateTree, FStateTreeStateHandle InStateHandle)
-			: StateTree(InStateTree)
-			, StateHandle(InStateHandle)
-		{
-		}
-		const UStateTree* StateTree = nullptr;
-		FStateTreeStateHandle StateHandle;
-
-		bool operator==(const FStateHandleContext&) const = default;
-	};
 }
 
 /**
@@ -600,6 +586,14 @@ PRAGMA_ENABLE_DEPRECATION_WARNINGS
 		return CurrentlyProcessedParentFrame;
 	}
 
+	/** @return in progress transition or the latest requested transition. */
+	TSharedPtr<UE::StateTree::ExecutionContext::ITemporaryStorage> GetCurrentlyProcessedTemporaryStorage() const
+	{
+		// A transition requests can be in progress or already succeeded.
+		//Use the current transition request first (for linked frame).
+		return CurrentlyProcessedTemporaryStorage ? CurrentlyProcessedTemporaryStorage : RequestedTransition ? RequestedTransition->Selection : nullptr;
+	}
+
 	/** @return Array view to named external data descriptors associated with this context. Note: Init() must be called before calling this method. */
 	TConstArrayView<FStateTreeExternalDataDesc> GetContextDataDescs() const
 	{
@@ -773,13 +767,10 @@ protected:
 	};
 
 	/** Describes a result of States Selection. */
-	struct FSelectStateResult
+	struct FSelectStateResult : public UE::StateTree::ExecutionContext::ITemporaryStorage
 	{
 		FSelectStateResult() = default;
-		FSelectStateResult(const FSelectStateResult&) = delete;
-		FSelectStateResult(FSelectStateResult&&) = default;
-		FSelectStateResult& operator=(const FSelectStateResult&) = delete;
-		FSelectStateResult& operator=(FSelectStateResult&&) = default;
+		virtual ~FSelectStateResult() = default;
 
 		/** The active states selected. They are in order from the root to the leaf. */
 		TArray<UE::StateTree::FActiveState> SelectedStates;
@@ -802,6 +793,10 @@ protected:
 		/** @return a new temporary frame */
 		UE_API FStateTreeExecutionFrame& MakeAndAddTemporaryFrame(UE::StateTree::FActiveFrameID FrameID, const UE::StateTree::FExecutionFrameKey& FrameKey, bool bIsGlobalFrame);
 		UE_API FStateTreeExecutionFrame& MakeAndAddTemporaryFrameWithNewRoot(UE::StateTree::FActiveFrameID FrameID, const UE::StateTree::FExecutionFrameKey& FrameKey, FStateTreeExecutionFrame& OtherFrame);
+
+		//~ITemporaryStorage
+		virtual FFrameAndParent GetExecutionFrame(UE::StateTree::FActiveFrameID ID) override;
+		virtual UE::StateTree::FActiveState GetStateHandle(UE::StateTree::FActiveStateID ID) const override;
 	};
 
 PRAGMA_DISABLE_DEPRECATION_WARNINGS
@@ -870,7 +865,7 @@ PRAGMA_ENABLE_DEPRECATION_WARNINGS
 	/** The result of the RequestTransition function. */
 	struct FRequestTransitionResult
 	{
-		FSelectStateResult Selection;
+		TSharedPtr<FSelectStateResult> Selection;
 		FStateTreeTransitionResult Transition;
 		FStateTreeTransitionSource Source;
 	};
@@ -886,7 +881,7 @@ PRAGMA_ENABLE_DEPRECATION_WARNINGS
 	UE_API void UpdateInstanceData(TConstArrayView<FStateTreeExecutionFrame> CurrentActiveFrames, TArrayView<FStateTreeExecutionFrame> NextActiveFrames);
 
 	/** Allocate the instance data. Fixup the TemporarySelectedFrame with the correct data. */
-	UE_API void UpdateInstanceData(FSelectStateResult& SelectStateResult);
+	UE_API void UpdateInstanceData(const TSharedPtr<FSelectStateResult>& SelectStateResult);
 
 	UE_DEPRECATED(5.7, "Use the EnterState with the FSelectStateResult argument.")
 	UE_API EStateTreeRunStatus EnterState(FStateTreeTransitionResult& Transition);
@@ -897,7 +892,7 @@ PRAGMA_ENABLE_DEPRECATION_WARNINGS
 	 * and still active after the transition will remain intact.
 	 * @return run status returned by the tasks.
 	 */
-	UE_API EStateTreeRunStatus EnterState(FSelectStateResult& SelectStateResult, const FStateTreeTransitionResult& Transition);
+	UE_API EStateTreeRunStatus EnterState(const TSharedPtr<FSelectStateResult>& SelectStateResult, const FStateTreeTransitionResult& Transition);
 
 	UE_DEPRECATED(5.7, "Use the ExitState with the FSelectStateResult argument.")
 	UE_API void ExitState(const FStateTreeTransitionResult& Transition);
@@ -907,7 +902,7 @@ PRAGMA_ENABLE_DEPRECATION_WARNINGS
 	 * Re-planned tree is from the transition target up to the leaf state. States that are parent to the transition target state
 	 * and still active after the transition will remain intact.
 	 */
-	UE_API void ExitState(const FSelectStateResult& SelectStateResult, const FStateTreeTransitionResult& Transition);
+	UE_API void ExitState(const TSharedPtr<const FSelectStateResult>& SelectStateResult, const FStateTreeTransitionResult& Transition);
 
 	/**
 	 * Removes all delegate listeners.
@@ -1142,7 +1137,7 @@ protected:
 	UE_API FRecordedStateTreeTransitionResult MakeRecordedTransitionResult(const FStateTreeTransitionResult& Transition) const;
 
 	/** Create a new recorded transition from a select state result. */
-	UE_API FRecordedStateTreeTransitionResult MakeRecordedTransitionResult(const FSelectStateResult& Args, const FStateTreeTransitionResult& Transition) const;
+	UE_API FRecordedStateTreeTransitionResult MakeRecordedTransitionResult(const TSharedRef<FSelectStateResult>& Args, const FStateTreeTransitionResult& Transition) const;
 
 PRAGMA_DISABLE_DEPRECATION_WARNINGS
 	UE_DEPRECATED(5.6, "Replaced with FStateTreeTasksCompletionStatus")
@@ -1172,7 +1167,7 @@ PRAGMA_ENABLE_DEPRECATION_WARNINGS
 	UE_API EStateTreeRunStatus GetGlobalTasksCompletedStatesStatus() const;
 
 	/** Forces transition to a state. It will skip all conditions. */
-	UE_API bool ForceTransitionInternal(const TArrayView<const UE::StateTree::ExecutionContext::FStateHandleContext> States, FSelectStateResult& OutSelectionResult);
+	UE_API bool ForceTransitionInternal(const TArrayView<const UE::StateTree::ExecutionContext::FStateHandleContext> States, const TSharedRef<FSelectStateResult>& OutSelectionResult);
 
 PRAGMA_DISABLE_DEPRECATION_WARNINGS
 	UE_DEPRECATED(5.7, "Use the SelectState with the FSelectStateArguments argument.")
@@ -1213,7 +1208,7 @@ PRAGMA_ENABLE_DEPRECATION_WARNINGS
 	 * @param OutSelectionResult the result of the selection.
 	 * @return True if succeeded to select new active states.
 	 */
-	UE_API bool SelectState(const FSelectStateArguments& SelectStateArgs, FSelectStateResult& OutSelectionResult);
+	UE_API bool SelectState(const FSelectStateArguments& SelectStateArgs, const TSharedRef<FSelectStateResult>& OutSelectionResult);
 
 PRAGMA_DISABLE_DEPRECATION_WARNINGS
 	UE_DEPRECATED(5.7, "Use the SelectStateInternal with FSelectStateArguments argument.")
@@ -1232,13 +1227,13 @@ PRAGMA_ENABLE_DEPRECATION_WARNINGS
 	UE_API bool SelectStateInternal(
 		const FSelectStateArguments& Args,
 		const FSelectStateInternalArguments& InternalArgs,
-		FSelectStateResult& OutSelectionResult);
+		const TSharedRef<FSelectStateResult>& OutSelectionResult);
 
 	/** Used internally to select state until we reach the transition source or the transition target. */
 	UE_API TOptional<bool> SelectStateFromSourceInternal(
 		const FSelectStateArguments& Args,
 		const FSelectStateInternalArguments& InternalArgs,
-		FSelectStateResult& OutSelectionResult,
+		const TSharedRef<FSelectStateResult>& OutSelectionResult,
 		const FStateTreeExecutionFrame& NextFrame,
 		const FCompactStateTreeState& NextState,
 		FStateTreeStateHandle NextStateHandle,
@@ -1247,7 +1242,7 @@ PRAGMA_ENABLE_DEPRECATION_WARNINGS
 	UE_API bool SelectStateInternal_Linked(
 		const FSelectStateArguments& Args,
 		const FSelectStateInternalArguments& InternalArgs,
-		FSelectStateResult& OutSelectionResult,
+		const TSharedRef<FSelectStateResult>& OutSelectionResult,
 		TNotNull<const UStateTree*> NextStateTree,
 		const FCompactStateTreeState& NextState,
 		bool bShouldCreateNewState
@@ -1256,7 +1251,7 @@ PRAGMA_ENABLE_DEPRECATION_WARNINGS
 	UE_API bool SelectStateInternal_LinkedAsset(
 		const FSelectStateArguments& Args,
 		const FSelectStateInternalArguments& InternalArgs,
-		FSelectStateResult& OutSelectionResult,
+		const TSharedRef<FSelectStateResult>& OutSelectionResult,
 		TNotNull<const UStateTree*> NextStateTree,
 		const FCompactStateTreeState& NextState,
 		const UStateTree* NextLinkedStateAsset,
@@ -1265,7 +1260,7 @@ PRAGMA_ENABLE_DEPRECATION_WARNINGS
 	UE_API bool SelectStateInternal_TrySelectChildrenInOrder(
 		const FSelectStateArguments& Args,
 		const FSelectStateInternalArguments& InternalArgs,
-		FSelectStateResult& OutSelectionResult,
+		const TSharedRef<FSelectStateResult>& OutSelectionResult,
 		TNotNull<const UStateTree*> NextStateTree,
 		const FCompactStateTreeState& TargetState,
 		const FStateTreeStateHandle TargetStateHandle);
@@ -1273,7 +1268,7 @@ PRAGMA_ENABLE_DEPRECATION_WARNINGS
 	UE_API bool SelectStateInternal_TrySelectChildrenAtRandom(
 		const FSelectStateArguments& Args,
 		const FSelectStateInternalArguments& InternalArgs,
-		FSelectStateResult& OutSelectionResult,
+		const TSharedRef<FSelectStateResult>& OutSelectionResult,
 		TNotNull<const UStateTree*> NextStateTree,
 		const FCompactStateTreeState& TargetState,
 		const FStateTreeStateHandle TargetStateHandle);
@@ -1281,7 +1276,7 @@ PRAGMA_ENABLE_DEPRECATION_WARNINGS
 	UE_API bool SelectStateInternal_TrySelectChildrenWithHighestUtility(
 		const FSelectStateArguments& Args,
 		const FSelectStateInternalArguments& InternalArgs,
-		FSelectStateResult& OutSelectionResult,
+		const TSharedRef<FSelectStateResult>& OutSelectionResult,
 		TNotNull<const UStateTree*> NextStateTree,
 		const FCompactStateTreeState& TargetState,
 		const FStateTreeStateHandle TargetStateHandle);
@@ -1289,7 +1284,7 @@ PRAGMA_ENABLE_DEPRECATION_WARNINGS
 	UE_API bool SelectStateInternal_TrySelectChildrenAtRandomWeightedByUtility(
 		const FSelectStateArguments& Args,
 		const FSelectStateInternalArguments& InternalArgs,
-		FSelectStateResult& OutSelectionResult,
+		const TSharedRef<FSelectStateResult>& OutSelectionResult,
 		TNotNull<const UStateTree*> NextStateTree,
 		const FCompactStateTreeState& TargetState,
 		const FStateTreeStateHandle TargetStateHandle);
@@ -1297,7 +1292,7 @@ PRAGMA_ENABLE_DEPRECATION_WARNINGS
 	UE_API bool SelectStateInternal_TryFollowTransitions(
 		const FSelectStateArguments& Args,
 		const FSelectStateInternalArguments& InternalArgs,
-		FSelectStateResult& OutSelectionResult,
+		const TSharedRef<FSelectStateResult>& OutSelectionResult,
 		TNotNull<const UStateTree*> NextStateTree,
 		const FCompactStateTreeState& TargetState,
 		const FStateTreeStateHandle TargetStateHandle);
@@ -1405,7 +1400,7 @@ PRAGMA_ENABLE_DEPRECATION_WARNINGS
 	/**
 	 * Captures StateTree events used during the state selection.
 	 */
-	UE_API void CaptureNewStateEvents(FSelectStateResult& Args);
+	UE_API void CaptureNewStateEvents(const TSharedRef<const FSelectStateResult>& Args);
 
 	/** @return a weak reference for a task that can be stored for later uses. */
 PRAGMA_DISABLE_DEPRECATION_WARNINGS
@@ -1542,8 +1537,15 @@ PRAGMA_ENABLE_DEPRECATION_WARNINGS
 		FStateTreeExecutionContext& Context;
 	};
 
-	/** Current select state result we're processing during the state selection, or invalid if not applicable. */
-	FSelectStateResult* CurrentlyProcessedStateSelectionResult = nullptr;
+	/**
+	 * Current select state result we're processing during the state selection, or invalid if not applicable.
+	 * Used by the GetDataView when accessing event data.
+	 * @note used only during selection and will always point to CurrentlyProcessedTemporaryStorage.
+	 */
+	FSelectStateResult* CurrentlyProcessedStateSelectionResult  = nullptr;
+
+	/** Current temporary storage for instance data, frames and state. Valid during the state selection, ExitState and EnterState. */
+	TSharedPtr<FSelectStateResult> CurrentlyProcessedTemporaryStorage;
 
 #if WITH_EDITORONLY_DATA
 PRAGMA_DISABLE_DEPRECATION_WARNINGS
diff --git a/Engine/Plugins/Runtime/StateTree/Source/StateTreeModule/Public/StateTreeExecutionContextTypes.h b/Engine/Plugins/Runtime/StateTree/Source/StateTreeModule/Public/StateTreeExecutionContextTypes.h
new file mode 100644
--- /dev/null
+++ b/Engine/Plugins/Runtime/StateTree/Source/StateTreeModule/Public/StateTreeExecutionContextTypes.h
@@ -0,0 +1,48 @@
+// Copyright Epic Games, Inc. All Rights Reserved.
+
+#pragma once
+
+#include "StateTreeExecutionTypes.h"
+#include "StateTreeStatePath.h"
+#include "Misc/NotNull.h"
+
+#define UE_API STATETREEMODULE_API
+
+class UStateTree;
+
+namespace UE::StateTree::ExecutionContext
+{
+	/** Helper that combines the state handle and its tree asset. */
+	struct FStateHandleContext
+	{
+		FStateHandleContext() = default;
+		FStateHandleContext(TNotNull<const UStateTree*> InStateTree, FStateTreeStateHandle InStateHandle)
+			: StateTree(InStateTree)
+			, StateHandle(InStateHandle)
+		{
+		}
+		
+		bool operator==(const FStateHandleContext&) const = default;
+		
+		const UStateTree* StateTree = nullptr;
+		FStateTreeStateHandle StateHandle;
+	};
+
+	/** Interface of structure that can store temporary frames and states. */
+	struct ITemporaryStorage
+	{
+		virtual ~ITemporaryStorage() = default;
+
+		struct FFrameAndParent
+		{
+			FStateTreeExecutionFrame* Frame = nullptr;
+			UE::StateTree::FActiveFrameID ParentFrameID;
+		};
+
+		virtual FFrameAndParent GetExecutionFrame(UE::StateTree::FActiveFrameID ID) = 0;
+		virtual UE::StateTree::FActiveState GetStateHandle(UE::StateTree::FActiveStateID ID) const = 0;
+	};
+
+} // UE::StateTree::ExecutionContext
+
+#undef UE_API
diff --git a/Engine/Plugins/Runtime/StateTree/Source/StateTreeTestSuite/Private/StateTreeTestTypes.h b/Engine/Plugins/Runtime/StateTree/Source/StateTreeTestSuite/Private/StateTreeTestTypes.h
--- a/Engine/Plugins/Runtime/StateTree/Source/StateTreeTestSuite/Private/StateTreeTestTypes.h
+++ b/Engine/Plugins/Runtime/StateTree/Source/StateTreeTestSuite/Private/StateTreeTestTypes.h
@@ -182,6 +182,58 @@ struct FTestEval_A : public FStateTreeEvaluatorBase
 	virtual const UStruct* GetInstanceDataType() const override { return FInstanceDataType::StaticStruct(); }
 };
 
+USTRUCT()
+struct FTestEval_Custom : public FStateTreeEvaluatorBase
+{
+	GENERATED_BODY()
+
+	using FInstanceDataType = FTestEval_AInstanceData;
+
+	FTestEval_Custom() = default;
+	FTestEval_Custom(const FName InName) { Name = InName; }
+
+	virtual const UStruct* GetInstanceDataType() const override
+	{
+		return FInstanceDataType::StaticStruct();
+	}
+
+	virtual void TreeStart(FStateTreeExecutionContext& Context) const override
+	{
+		if (CustomEnterStateFunc)
+		{
+			CustomEnterStateFunc(Context, this);
+		}
+	}
+
+	virtual void TreeStop(FStateTreeExecutionContext& Context) const override 
+	{
+		if (CustomExitStateFunc)
+		{
+			CustomExitStateFunc(Context, this);
+		}
+	}
+
+	virtual void Tick(FStateTreeExecutionContext& Context, const float DeltaTime) const override
+	{
+		if (CustomTickFunc)
+		{
+			CustomTickFunc(Context, this);
+		}
+	}
+
+	virtual bool Link(FStateTreeLinker& Linker) override
+	{
+		Linker.LinkExternalData(LogHandle);
+		return true;
+	}
+
+	TStateTreeExternalDataHandle<FStateTreeTestLog> LogHandle;
+
+	TFunction<void(FStateTreeExecutionContext&, const FTestEval_Custom*)> CustomEnterStateFunc;
+	TFunction<void(FStateTreeExecutionContext&, const FTestEval_Custom*)> CustomExitStateFunc;
+	TFunction<void(FStateTreeExecutionContext&, const FTestEval_Custom*)> CustomTickFunc;
+};
+
 USTRUCT()
 struct FTestTask_BInstanceData
 {
diff --git a/Engine/Plugins/Runtime/StateTree/Source/StateTreeTestSuite/Private/StateTreeWeakContextTest.cpp b/Engine/Plugins/Runtime/StateTree/Source/StateTreeTestSuite/Private/StateTreeWeakContextTest.cpp
--- a/Engine/Plugins/Runtime/StateTree/Source/StateTreeTestSuite/Private/StateTreeWeakContextTest.cpp
+++ b/Engine/Plugins/Runtime/StateTree/Source/StateTreeTestSuite/Private/StateTreeWeakContextTest.cpp
@@ -220,6 +220,562 @@ struct FStateTreeTest_WeakContext_FinishTask : FStateTreeTestBase
 	}
 };
 IMPLEMENT_STATE_TREE_INSTANT_TEST(FStateTreeTest_WeakContext_FinishTask, "System.StateTree.WeakContext.FinishTask");
+
+struct FStateTreeTest_WeakContext_InstanceData : FStateTreeTestBase
+{
+	//Tree 1 : 2x Global Tasks, 2x global evaluators
+	//	Root
+	//		State1 : 2x tasks -> Next
+	//			State2
+	//				State3Linked -> Tree 2 (fails) -> Next
+	//				State4Linked -> Tree 2 (complete) -> Next
+	//				State5 -> Next (dummy)
+	//Tree 2 : 2x Global tasks, 2x global evaluators
+	//	Root
+	//		State1 : 2x tasks -> fail or wait
+
+
+	//(1)Tree1GlobalEvaluator1 Context.Weak can access data
+	//(2)Tree1GlobalEvaluator2 Context.Weak can access data and access data of Tree1GlobalEvaluator1
+	//(3)Tree1GlobalTask1 Context.Weak can access data and access data of Tree1GlobalEvaluator2
+	//(4)Tree1GlobalTask2 Context.Weak can access data and access data of Tree1GlobalTask1
+	//(1)Tree2GlobalEvaluator1 Context.Weak can access data and access data of Tree1GlobalTask2
+	//(2)Tree2GlobalEvaluator2 Context.Weak can access data and access data of Tree2GlobalEvaluator1
+	//(3)Tree2GlobalTask1 Context.Weak can access data and access data of Tree2GlobalEvaluator
+	//(4)Tree2GlobalTask2 Context.Weak can access data and access data of Tree2GlobalTask1
+	//Tree2GlobalTask2 fails
+	//(1)Tree2GlobalEvaluator1 Context.Weak can access data and access data of Tree1GlobalTask2 but not of previous Tree2GlobalEvaluator1
+	//(2)Tree2GlobalEvaluator2 Context.Weak can access data and access data of Tree2GlobalEvaluator1
+	//(3)Tree2GlobalTask1 Context.Weak can access data and access data of Tree2GlobalEvaluator
+	//(4)Tree2GlobalTask2 Context.Weak can access data and access data of Tree2GlobalTask1
+	//Tree2GlobalTask2 success
+	//can access all data
+	//trigger another state selection of State1 to test the same during state selection
+
+	virtual bool InstantTest() override
+	{
+		struct FTestContext
+		{
+			FStateTreeWeakExecutionContext ContextTree1GlobalEvaluator1;
+			FStateTreeWeakExecutionContext ContextTree1GlobalEvaluator2;
+			FStateTreeWeakExecutionContext ContextTree1GlobalTask1;
+			FStateTreeWeakExecutionContext ContextTree1GlobalTask2;
+			FStateTreeWeakExecutionContext ContextTree1State1Task1;
+			FStateTreeWeakExecutionContext ContextTree1State1Task2;
+			FStateTreeWeakExecutionContext ContextTree2GlobalEvaluator1_1;
+			FStateTreeWeakExecutionContext ContextTree2GlobalEvaluator2_1;
+			FStateTreeWeakExecutionContext ContextTree2GlobalTask1_1;
+			FStateTreeWeakExecutionContext ContextTree2GlobalTask2_1;
+			FStateTreeWeakExecutionContext ContextTree2State1Task1_1;
+			FStateTreeWeakExecutionContext ContextTree2State1Task2_1;
+			FStateTreeWeakExecutionContext ContextTree2GlobalEvaluator1_2;
+			FStateTreeWeakExecutionContext ContextTree2GlobalEvaluator2_2;
+			FStateTreeWeakExecutionContext ContextTree2GlobalTask1_2;
+			FStateTreeWeakExecutionContext ContextTree2GlobalTask2_2;
+			FStateTreeWeakExecutionContext ContextTree2State1Task1_2;
+			FStateTreeWeakExecutionContext ContextTree2State1Task2_2;
+
+			bool bSecondTree2Trigger = false;
+
+			TArray<int32> ReadAllData() const
+			{
+				TArray<int32> Result;
+				Result.Add(ReadEvalData(ContextTree1GlobalEvaluator1));
+				Result.Add(ReadEvalData(ContextTree1GlobalEvaluator2));
+				Result.Add(ReadTaskData(ContextTree1GlobalTask1));
+				Result.Add(ReadTaskData(ContextTree1GlobalTask2));
+				Result.Add(ReadTaskData(ContextTree1State1Task1));
+				Result.Add(ReadTaskData(ContextTree1State1Task2));
+				Result.Add(ReadEvalData(ContextTree2GlobalEvaluator1_1));
+				Result.Add(ReadEvalData(ContextTree2GlobalEvaluator2_1));
+				Result.Add(ReadTaskData(ContextTree2GlobalTask1_1));
+				Result.Add(ReadTaskData(ContextTree2GlobalTask2_1));
+				Result.Add(ReadTaskData(ContextTree2State1Task1_1));
+				Result.Add(ReadTaskData(ContextTree2State1Task2_1));
+				Result.Add(ReadEvalData(ContextTree2GlobalEvaluator1_2));
+				Result.Add(ReadEvalData(ContextTree2GlobalEvaluator2_2));
+				Result.Add(ReadTaskData(ContextTree2GlobalTask1_2));
+				Result.Add(ReadTaskData(ContextTree2GlobalTask2_2));
+				Result.Add(ReadTaskData(ContextTree2State1Task1_2));
+				Result.Add(ReadTaskData(ContextTree2State1Task2_2));
+				return Result;
+			}
+
+			FString ReadAllData_ToString() const
+			{
+				TStringBuilder<256> StringBuilder;
+				StringBuilder << TEXT('{');
+				bool bFirst = true;
+				for (int32 It : ReadAllData())
+				{
+					if (!bFirst)
+					{
+						StringBuilder << TEXT(", ");
+					}
+					bFirst = false;
+					StringBuilder << It;
+				}
+				StringBuilder << TEXT('}');
+				return StringBuilder.ToString();
+			}
+
+			FString ExepectData(TArrayView<const int32> AllElements, TArrayView<const int32> ValidElements)
+			{
+				TStringBuilder<256> StringBuilder;
+				StringBuilder << TEXT("AccessData:{");
+				bool bFirst = true;
+				for (int32 Index = 0; Index < AllElements.Num(); ++Index)
+				{
+					if (!bFirst)
+					{
+						StringBuilder << TEXT(", ");
+					}
+					bFirst = false;
+					if (ValidElements.Contains(Index))
+					{
+						StringBuilder << AllElements[Index];
+					}
+					else
+					{
+						constexpr int32 InvalidNumber = - 1;
+						StringBuilder << InvalidNumber;
+					}
+				}
+				StringBuilder << TEXT('}');
+				return StringBuilder.ToString();
+			}
+
+			int32 ReadTaskData(const FStateTreeWeakExecutionContext& WeakContext) const
+			{
+				FStateTreeStrongReadOnlyExecutionContext StrongContext = WeakContext.MakeStrongReadOnlyExecutionContext();
+				const FTestTask_PrintValue::FInstanceDataType* Ptr = StrongContext.GetInstanceDataPtr<const FTestTask_PrintValue::FInstanceDataType>();
+				return Ptr ? Ptr->Value : INDEX_NONE;
+			}
+			bool WriteTaskData(FStateTreeWeakExecutionContext& WeakContext, int32 NewValue) const
+			{
+				FStateTreeStrongExecutionContext StrongContext = WeakContext.MakeStrongExecutionContext();
+				FTestTask_PrintValue::FInstanceDataType* Ptr = StrongContext.GetInstanceDataPtr<FTestTask_PrintValue::FInstanceDataType>();
+				if (Ptr)
+				{
+					Ptr->Value = NewValue;
+					return true;
+				}
+				return false;
+			}
+			int32 ReadEvalData(const FStateTreeWeakExecutionContext& WeakContext) const
+			{
+				FStateTreeStrongReadOnlyExecutionContext StrongContext = WeakContext.MakeStrongReadOnlyExecutionContext();
+				const FTestEval_Custom::FInstanceDataType* Ptr = StrongContext.GetInstanceDataPtr<const FTestEval_Custom::FInstanceDataType>();
+				return Ptr ? Ptr->IntA : INDEX_NONE;
+			}
+			bool WriteEvalData(FStateTreeWeakExecutionContext& WeakContext, int32 NewValue) const
+			{
+				FStateTreeStrongExecutionContext StrongContext = WeakContext.MakeStrongExecutionContext();
+				FTestEval_Custom::FInstanceDataType* Ptr = StrongContext.GetInstanceDataPtr<FTestEval_Custom::FInstanceDataType>();
+				if (Ptr)
+				{
+					Ptr->IntA = NewValue;
+					return true;
+				}
+				return false;
+			}
+		} TestContext;
+
+		// Building up the State Tree
+		UStateTree& StateTree2 = NewStateTree();
+		UStateTree& StateTree1 = NewStateTree();
+
+		#define UE_STATETREE_TEST_CUSTOMENTERSTATEFUNC(T1) \
+			TestContext.T1 = Context.MakeWeakExecutionContext();\
+			FStateTreeTestLog& TestLog = Context.GetExternalData(Task->LogHandle);\
+			TArray<int32> AllDatas = TestContext.ReadAllData();\
+			TestLog.Log(Task->Name, FString::Printf(TEXT("AccessData:%s"), *TestContext.ReadAllData_ToString()));
+
+		#define UE_STATETREE_TEST_CUSTOMENTERSTATEFUNC_2(T1, T2) \
+			if (!TestContext.bSecondTree2Trigger)\
+			{\
+				TestContext.T1 = Context.MakeWeakExecutionContext();\
+			}\
+			else\
+			{\
+				TestContext.T2 = Context.MakeWeakExecutionContext();\
+			}\
+			FStateTreeTestLog& TestLog = Context.GetExternalData(Task->LogHandle);\
+			TArray<int32> AllDatas = TestContext.ReadAllData();\
+			TestLog.Log(Task->Name, FString::Printf(TEXT("AccessData:%s"), *TestContext.ReadAllData_ToString()));
+
+		#define UE_STATETREE_TEST_CUSTOMTICKFUNC() \
+			FStateTreeTestLog& TestLog = Context.GetExternalData(Task->LogHandle);\
+			TArray<int32> AllDatas = TestContext.ReadAllData();\
+			TestLog.Log(Task->Name, FString::Printf(TEXT("AccessData:%s"), *TestContext.ReadAllData_ToString()));
+		
+		// Tree2
+		{
+			UStateTreeEditorData& EditorData = *Cast<UStateTreeEditorData>(StateTree2.EditorData);
+			{
+				TStateTreeEditorNode<FTestEval_Custom>& Tree2GlobalEval1 = EditorData.AddEvaluator<FTestEval_Custom>(FName("Tree2GlobalEval1"));
+				Tree2GlobalEval1.GetInstanceData().IntA = 201;
+				Tree2GlobalEval1.GetNode().CustomEnterStateFunc = [&TestContext](FStateTreeExecutionContext& Context, const FTestEval_Custom* Task)
+					{
+						UE_STATETREE_TEST_CUSTOMENTERSTATEFUNC_2(ContextTree2GlobalEvaluator1_1, ContextTree2GlobalEvaluator1_2);
+					};
+				Tree2GlobalEval1.GetNode().CustomTickFunc = [&TestContext](FStateTreeExecutionContext& Context, const FTestEval_Custom* Task)
+					{
+						UE_STATETREE_TEST_CUSTOMTICKFUNC();
+					};
+				Tree2GlobalEval1.GetNode().CustomExitStateFunc = [&TestContext](FStateTreeExecutionContext& Context, const FTestEval_Custom* Task)
+					{
+						UE_STATETREE_TEST_CUSTOMTICKFUNC();
+					};
+				TStateTreeEditorNode<FTestEval_Custom>& Tree2GlobalEval2 = EditorData.AddEvaluator<FTestEval_Custom>(FName("Tree2GlobalEval2"));
+				Tree2GlobalEval2.GetInstanceData().IntA = 202;
+				Tree2GlobalEval2.GetNode().CustomEnterStateFunc = [&TestContext](FStateTreeExecutionContext& Context, const FTestEval_Custom* Task)
+					{
+						UE_STATETREE_TEST_CUSTOMENTERSTATEFUNC_2(ContextTree2GlobalEvaluator2_1, ContextTree2GlobalEvaluator2_2);
+					};
+				Tree2GlobalEval2.GetNode().CustomTickFunc = [&TestContext](FStateTreeExecutionContext& Context, const FTestEval_Custom* Task)
+					{
+						UE_STATETREE_TEST_CUSTOMTICKFUNC();
+					};
+				Tree2GlobalEval2.GetNode().CustomExitStateFunc = [&TestContext](FStateTreeExecutionContext& Context, const FTestEval_Custom* Task)
+					{
+						UE_STATETREE_TEST_CUSTOMTICKFUNC();
+					};
+				// Global Task
+				TStateTreeEditorNode<FTestTask_PrintValue>& Tree2GlobalTask1 = EditorData.AddGlobalTask<FTestTask_PrintValue>(FName("Tree2GlobalTask1"));
+				Tree2GlobalTask1.GetInstanceData().Value = 203;
+				Tree2GlobalTask1.GetNode().CustomEnterStateFunc = [&TestContext](FStateTreeExecutionContext& Context, const FTestTask_PrintValue* Task)
+					{
+						UE_STATETREE_TEST_CUSTOMENTERSTATEFUNC_2(ContextTree2GlobalTask1_1, ContextTree2GlobalTask1_2);
+					};
+				Tree2GlobalTask1.GetNode().CustomTickFunc = [&TestContext](FStateTreeExecutionContext& Context, const FTestTask_PrintValue* Task)
+					{
+						UE_STATETREE_TEST_CUSTOMTICKFUNC();
+					};
+				Tree2GlobalTask1.GetNode().CustomExitStateFunc = [&TestContext](FStateTreeExecutionContext& Context, const FTestTask_PrintValue* Task)
+					{
+						UE_STATETREE_TEST_CUSTOMTICKFUNC();
+					};
+
+				TStateTreeEditorNode<FTestTask_PrintValue>& Tree2GlobalTask2 = EditorData.AddGlobalTask<FTestTask_PrintValue>(FName("Tree2GlobalTask2"));
+				Tree2GlobalTask2.GetInstanceData().Value = 204;
+				Tree2GlobalTask2.GetNode().CustomEnterStateFunc = [&TestContext](FStateTreeExecutionContext& Context, const FTestTask_PrintValue* Task)
+					{
+						UE_STATETREE_TEST_CUSTOMENTERSTATEFUNC_2(ContextTree2GlobalTask2_1, ContextTree2GlobalTask2_2);
+					};
+				Tree2GlobalTask2.GetNode().CustomTickFunc = [&TestContext](FStateTreeExecutionContext& Context, const FTestTask_PrintValue* Task)
+					{
+						UE_STATETREE_TEST_CUSTOMTICKFUNC();
+					};
+				Tree2GlobalTask2.GetNode().CustomExitStateFunc = [&TestContext](FStateTreeExecutionContext& Context, const FTestTask_PrintValue* Task)
+					{
+						UE_STATETREE_TEST_CUSTOMTICKFUNC();
+					};
+			}
+			// Root State
+			UStateTreeState& Root = EditorData.AddSubTree("Tree2StateRoot");
+			{
+				TStateTreeEditorNode<FTestTask_PrintValue>& Tree2RootTask1 = Root.AddTask<FTestTask_PrintValue>(FName("Tree2RootTask1"));
+				Tree2RootTask1.GetInstanceData().Value = 211;
+				TStateTreeEditorNode<FTestTask_PrintValue>& Tree2RootTask2 = Root.AddTask<FTestTask_PrintValue>(FName("Tree2RootTask2"));
+				Tree2RootTask2.GetInstanceData().Value = 212;
+			}
+			// State 1
+			UStateTreeState& State1 = Root.AddChildState("Tree2State1");
+			{
+				TStateTreeEditorNode<FTestTask_PrintValue>& Tree2State1Task1 = State1.AddTask<FTestTask_PrintValue>(FName("Tree2State1Task1"));
+				Tree2State1Task1.GetInstanceData().Value = 221;
+				Tree2State1Task1.GetNode().CustomEnterStateFunc = [&TestContext](FStateTreeExecutionContext& Context, const FTestTask_PrintValue* Task)
+					{
+						UE_STATETREE_TEST_CUSTOMENTERSTATEFUNC_2(ContextTree2State1Task1_1, ContextTree2State1Task1_2);
+
+						if (!TestContext.bSecondTree2Trigger)
+						{
+							FTestTask_PrintValue::FInstanceDataType& InstanceData = Context.GetInstanceData(*Task);
+							InstanceData.EnterStateRunStatus = EStateTreeRunStatus::Failed;
+							TestContext.bSecondTree2Trigger = true;
+						}
+					};
+				Tree2State1Task1.GetNode().CustomTickFunc = [&TestContext](FStateTreeExecutionContext& Context, const FTestTask_PrintValue* Task)
+					{
+						UE_STATETREE_TEST_CUSTOMTICKFUNC();
+					};
+				Tree2State1Task1.GetNode().CustomExitStateFunc = [&TestContext](FStateTreeExecutionContext& Context, const FTestTask_PrintValue* Task)
+					{
+						UE_STATETREE_TEST_CUSTOMTICKFUNC();
+					};
+
+				TStateTreeEditorNode<FTestTask_PrintValue>& Tree2State1Task2 = State1.AddTask<FTestTask_PrintValue>(FName("Tree2State1Task2"));
+				Tree2State1Task2.GetInstanceData().Value = 222;
+				Tree2State1Task2.GetNode().CustomEnterStateFunc = [&TestContext](FStateTreeExecutionContext& Context, const FTestTask_PrintValue* Task)
+					{
+						UE_STATETREE_TEST_CUSTOMENTERSTATEFUNC_2(ContextTree2State1Task2_1, ContextTree2State1Task2_2);
+					};
+				Tree2State1Task2.GetNode().CustomTickFunc = [&TestContext](FStateTreeExecutionContext& Context, const FTestTask_PrintValue* Task)
+					{
+						UE_STATETREE_TEST_CUSTOMTICKFUNC();
+					};
+				Tree2State1Task2.GetNode().CustomExitStateFunc = [&TestContext](FStateTreeExecutionContext& Context, const FTestTask_PrintValue* Task)
+					{
+						UE_STATETREE_TEST_CUSTOMTICKFUNC();
+					};
+			}
+
+			// Compile tree
+			{
+				FStateTreeCompilerLog Log;
+				FStateTreeCompiler Compiler(Log);
+				const bool bResult = Compiler.Compile(StateTree2);
+				AITEST_TRUE(TEXT("StateTree2 should get compiled"), bResult);
+			}
+		}
+		// Tree1
+		{
+			UStateTreeEditorData& EditorData = *Cast<UStateTreeEditorData>(StateTree1.EditorData);
+			{
+				// Global Evaluator
+				TStateTreeEditorNode<FTestEval_Custom>& Tree1GlobalEval1 = EditorData.AddEvaluator<FTestEval_Custom>(FName("Tree1GlobalEval1"));
+				Tree1GlobalEval1.GetInstanceData().IntA = 101;
+				Tree1GlobalEval1.GetNode().CustomEnterStateFunc = [&TestContext](FStateTreeExecutionContext& Context, const FTestEval_Custom* Task)
+					{
+						UE_STATETREE_TEST_CUSTOMENTERSTATEFUNC(ContextTree1GlobalEvaluator1);
+					};
+				Tree1GlobalEval1.GetNode().CustomTickFunc = [&TestContext](FStateTreeExecutionContext& Context, const FTestEval_Custom* Task)
+					{
+						UE_STATETREE_TEST_CUSTOMTICKFUNC();
+					};
+				Tree1GlobalEval1.GetNode().CustomExitStateFunc = [&TestContext](FStateTreeExecutionContext& Context, const FTestEval_Custom* Task)
+					{
+						UE_STATETREE_TEST_CUSTOMTICKFUNC();
+					};
+				TStateTreeEditorNode<FTestEval_Custom>& Tree1GlobalEval2 = EditorData.AddEvaluator<FTestEval_Custom>(FName("Tree1GlobalEval2"));
+				Tree1GlobalEval2.GetInstanceData().IntA = 102;
+				Tree1GlobalEval2.GetNode().CustomEnterStateFunc = [&TestContext](FStateTreeExecutionContext& Context, const FTestEval_Custom* Task)
+					{
+						UE_STATETREE_TEST_CUSTOMENTERSTATEFUNC(ContextTree1GlobalEvaluator2);
+					};
+				Tree1GlobalEval2.GetNode().CustomTickFunc = [&TestContext](FStateTreeExecutionContext& Context, const FTestEval_Custom* Task)
+					{
+						UE_STATETREE_TEST_CUSTOMTICKFUNC();
+					};
+				Tree1GlobalEval2.GetNode().CustomExitStateFunc = [&TestContext](FStateTreeExecutionContext& Context, const FTestEval_Custom* Task)
+					{
+						UE_STATETREE_TEST_CUSTOMTICKFUNC();
+					};
+
+				// Global Task
+				TStateTreeEditorNode<FTestTask_PrintValue>& Tree1GlobalTask1 = EditorData.AddGlobalTask<FTestTask_PrintValue>(FName("Tree1GlobalTask1"));
+				Tree1GlobalTask1.GetInstanceData().Value = 103;
+				Tree1GlobalTask1.GetNode().CustomEnterStateFunc = [&TestContext](FStateTreeExecutionContext& Context, const FTestTask_PrintValue* Task)
+					{
+						UE_STATETREE_TEST_CUSTOMENTERSTATEFUNC(ContextTree1GlobalTask1);
+					};
+				Tree1GlobalTask1.GetNode().CustomTickFunc = [&TestContext](FStateTreeExecutionContext& Context, const FTestTask_PrintValue* Task)
+					{
+						UE_STATETREE_TEST_CUSTOMTICKFUNC();
+					};
+				Tree1GlobalTask1.GetNode().CustomExitStateFunc = [&TestContext](FStateTreeExecutionContext& Context, const FTestTask_PrintValue* Task)
+					{
+						UE_STATETREE_TEST_CUSTOMTICKFUNC();
+					};
+
+				TStateTreeEditorNode<FTestTask_PrintValue>& Tree1GlobalTask2 = EditorData.AddGlobalTask<FTestTask_PrintValue>(FName("Tree1GlobalTask2"));
+				Tree1GlobalTask2.GetInstanceData().Value = 104;
+				Tree1GlobalTask2.GetNode().CustomEnterStateFunc = [&TestContext](FStateTreeExecutionContext& Context, const FTestTask_PrintValue* Task)
+					{
+						UE_STATETREE_TEST_CUSTOMENTERSTATEFUNC(ContextTree1GlobalTask2);
+					};
+				Tree1GlobalTask2.GetNode().CustomTickFunc = [&TestContext](FStateTreeExecutionContext& Context, const FTestTask_PrintValue* Task)
+					{
+						UE_STATETREE_TEST_CUSTOMTICKFUNC();
+					};
+				Tree1GlobalTask2.GetNode().CustomExitStateFunc = [&TestContext](FStateTreeExecutionContext& Context, const FTestTask_PrintValue* Task)
+					{
+						UE_STATETREE_TEST_CUSTOMTICKFUNC();
+					};
+			}
+			// Root State
+			UStateTreeState& Root = EditorData.AddSubTree("Tree1StateRoot");
+			{
+				TStateTreeEditorNode<FTestTask_PrintValue>& Tree1RootTask1 = Root.AddTask<FTestTask_PrintValue>(FName("Tree1RootTask1"));
+				Tree1RootTask1.GetInstanceData().Value = 111;
+				TStateTreeEditorNode<FTestTask_PrintValue>& Tree1RootTask2 = Root.AddTask<FTestTask_PrintValue>(FName("Tree1RootTask2"));
+				Tree1RootTask2.GetInstanceData().Value = 112;
+			}
+			// State 1
+			UStateTreeState& State1 = Root.AddChildState("Tree1State1");
+			State1.SelectionBehavior = EStateTreeStateSelectionBehavior::TrySelectChildrenInOrder;
+			{
+				TStateTreeEditorNode<FTestTask_PrintValue>& Tree1State1Task1 = State1.AddTask<FTestTask_PrintValue>(FName("Tree1State1Task1"));
+				Tree1State1Task1.GetInstanceData().Value = 121;
+				Tree1State1Task1.GetNode().CustomEnterStateFunc = [&TestContext](FStateTreeExecutionContext& Context, const FTestTask_PrintValue* Task)
+					{
+						UE_STATETREE_TEST_CUSTOMENTERSTATEFUNC(ContextTree1State1Task1);
+					};
+				Tree1State1Task1.GetNode().CustomTickFunc = [&TestContext](FStateTreeExecutionContext& Context, const FTestTask_PrintValue* Task)
+					{
+						UE_STATETREE_TEST_CUSTOMTICKFUNC();
+					};
+				Tree1State1Task1.GetNode().CustomExitStateFunc = [&TestContext](FStateTreeExecutionContext& Context, const FTestTask_PrintValue* Task)
+					{
+						UE_STATETREE_TEST_CUSTOMTICKFUNC();
+					};
+
+				TStateTreeEditorNode<FTestTask_PrintValue>& Tree1State1Task2 = State1.AddTask<FTestTask_PrintValue>(FName("Tree1State1Task2"));
+				Tree1State1Task2.GetInstanceData().Value = 122;
+				Tree1State1Task2.GetNode().CustomEnterStateFunc = [&TestContext](FStateTreeExecutionContext& Context, const FTestTask_PrintValue* Task)
+					{
+						UE_STATETREE_TEST_CUSTOMENTERSTATEFUNC(ContextTree1State1Task2);
+					};
+				Tree1State1Task2.GetNode().CustomTickFunc = [&TestContext](FStateTreeExecutionContext& Context, const FTestTask_PrintValue* Task)
+					{
+						UE_STATETREE_TEST_CUSTOMTICKFUNC();
+					};
+				Tree1State1Task2.GetNode().CustomExitStateFunc = [&TestContext](FStateTreeExecutionContext& Context, const FTestTask_PrintValue* Task)
+					{
+						UE_STATETREE_TEST_CUSTOMTICKFUNC();
+					};
+			}
+			UStateTreeState& State2 = State1.AddChildState("Tree1State2");
+			State1.SelectionBehavior = EStateTreeStateSelectionBehavior::TrySelectChildrenInOrder;
+			UStateTreeState& State3 = State2.AddChildState("Tree1State3", EStateTreeStateType::LinkedAsset);
+			State3.SetLinkedStateAsset(&StateTree2);
+			State3.AddTransition(EStateTreeTransitionTrigger::OnStateCompleted, EStateTreeTransitionType::NextState);
+			UStateTreeState& State4 = State2.AddChildState("Tree1State4", EStateTreeStateType::LinkedAsset);
+			State4.SetLinkedStateAsset(&StateTree2);
+			UStateTreeState& State5 = State2.AddChildState("Tree1State5");
+			{
+				TStateTreeEditorNode<FTestTask_PrintValue>& Tree1State5Task1 = State4.AddTask<FTestTask_PrintValue>(FName("Tree1State5Task1"));
+				Tree1State5Task1.GetInstanceData().Value = 155;
+			}
+
+			// Compile tree
+			{
+				FStateTreeCompilerLog Log;
+				FStateTreeCompiler Compiler(Log);
+				const bool bResult = Compiler.Compile(StateTree1);
+				AITEST_TRUE(TEXT("StateTree2 should get compiled"), bResult);
+			}
+		}
+
+		#undef UE_STATETREE_TEST_CUSTOMENTERSTATEFUNC
+		#undef UE_STATETREE_TEST_CUSTOMENTERSTATEFUNC_2
+		#undef UE_STATETREE_TEST_CUSTOMTICKFUNC
+
+		//                                 0    1    2    3    4    5    6    7    8    9    10   11   12   13   14   15   16   17
+		const int32 AllValidReadData[] = { 101, 102, 103, 104, 121, 122, 201, 202, 203, 204, 221, 222, 201, 202, 203, 204, 221, 222 };
+
+		// Create context
+		FStateTreeInstanceData InstanceData;
+		{
+			{
+				FTestStateTreeExecutionContext Exec(StateTree1, StateTree1, InstanceData);
+				const bool bInitSucceeded = Exec.IsValid();
+				AITEST_TRUE(TEXT("StateTree should init"), bInitSucceeded);
+			}
+			{
+				FTestStateTreeExecutionContext Exec(StateTree1, StateTree1, InstanceData);
+				EStateTreeRunStatus Status = Exec.Start();
+				AITEST_EQUAL(TEXT("Start should complete with Running"), Status, EStateTreeRunStatus::Running);
+				AITEST_TRUE(TEXT("In correct states"), Exec.ExpectInActiveStates("Tree1StateRoot", "Tree1State1", "Tree1State2", "Tree1State3", "Tree2StateRoot", "Tree2State1"));
+
+				FTestStateTreeExecutionContext::FLogOrder ExpectedLog = Exec.Expect("Tree1GlobalEval1", TestContext.ExepectData(AllValidReadData, { 0 }));
+				AITEST_TRUE(TEXT("Start should EnterState"), ExpectedLog);
+				ExpectedLog = ExpectedLog.Then("Tree1GlobalEval2", TestContext.ExepectData(AllValidReadData, { 0, 1 }));
+				AITEST_TRUE(TEXT("Start should EnterState"), ExpectedLog);
+				ExpectedLog = ExpectedLog.Then("Tree1GlobalTask1", TestContext.ExepectData(AllValidReadData, { 0, 1, 2 }));
+				AITEST_TRUE(TEXT("Start should EnterState"), ExpectedLog);
+				ExpectedLog = ExpectedLog.Then("Tree1GlobalTask2", TestContext.ExepectData(AllValidReadData, { 0, 1, 2, 3 }));
+				AITEST_TRUE(TEXT("Start should EnterState"), ExpectedLog);
+				ExpectedLog = ExpectedLog.Then("Tree1State1Task1", TestContext.ExepectData(AllValidReadData, { 0, 1, 2, 3, 4, 6, 7, 8, 9 }));
+				AITEST_TRUE(TEXT("Start should EnterState"), ExpectedLog);
+				ExpectedLog = ExpectedLog.Then("Tree1State1Task2", TestContext.ExepectData(AllValidReadData, { 0, 1, 2, 3, 4, 5, 6, 7, 8, 9 }));
+				AITEST_TRUE(TEXT("Start should EnterState"), ExpectedLog);
+
+				ExpectedLog = Exec.Expect("Tree2GlobalEval1", TestContext.ExepectData(AllValidReadData, { 0, 1, 2, 3, 6 }));
+				AITEST_TRUE(TEXT("Start should EnterState"), ExpectedLog);
+				ExpectedLog = ExpectedLog.Then("Tree2GlobalEval2", TestContext.ExepectData(AllValidReadData, { 0, 1, 2, 3, 6, 7 }));
+				AITEST_TRUE(TEXT("Start should EnterState"), ExpectedLog);
+				ExpectedLog = ExpectedLog.Then("Tree2GlobalTask1", TestContext.ExepectData(AllValidReadData, { 0, 1, 2, 3, 6, 7, 8 }));
+				AITEST_TRUE(TEXT("Start should EnterState"), ExpectedLog);
+				ExpectedLog = ExpectedLog.Then("Tree2GlobalTask2", TestContext.ExepectData(AllValidReadData, { 0, 1, 2, 3, 6, 7, 8, 9 }));
+				AITEST_TRUE(TEXT("Start should EnterState"), ExpectedLog);
+				ExpectedLog = ExpectedLog.Then("Tree2State1Task1", TestContext.ExepectData(AllValidReadData, { 0, 1, 2, 3, 4, 5, 6, 7, 8, 9, 10 }));
+				AITEST_TRUE(TEXT("Start should EnterState"), ExpectedLog);
+				ExpectedLog = ExpectedLog.Then("Tree2State1Task2", TestContext.ExepectData(AllValidReadData, { 0, 1, 2, 3, 4, 5, 6, 7, 8, 9, 10, 11 })); // 11 is not valid
+				AITEST_FALSE(TEXT("Start should EnterState"), ExpectedLog);
+			}
+			{
+				TArray<int32> AllDatas = TestContext.ReadAllData();
+				AITEST_TRUE(TEXT("After enter. Same num"), AllDatas.Num() == UE_ARRAY_COUNT(AllValidReadData));
+				AITEST_TRUE(TEXT("After enter. Same Value"), AllDatas == MakeArrayView({ 101, 102, 103, 104, 121, 122, 201, 202, 203, 204, 221, -1, -1, -1, -1, -1, -1, -1 }));
+			}
+			{
+				FTestStateTreeExecutionContext Exec(StateTree1, StateTree1, InstanceData);
+				EStateTreeRunStatus Status = Exec.Tick(0.1f);
+				AITEST_EQUAL(TEXT("Start should complete with Running"), Status, EStateTreeRunStatus::Running);
+				AITEST_TRUE(TEXT("In correct states"), Exec.ExpectInActiveStates("Tree1StateRoot", "Tree1State1", "Tree1State2", "Tree1State4", "Tree2StateRoot", "Tree2State1"));
+
+				// Tick
+				FTestStateTreeExecutionContext::FLogOrder ExpectedLog = Exec.Expect("Tree2GlobalTask2", TEXT("Tick204"));
+				AITEST_TRUE(TEXT("Tick"), ExpectedLog);
+				ExpectedLog = ExpectedLog.Then("Tree2GlobalTask2", TestContext.ExepectData(AllValidReadData, { 0, 1, 2, 3, 4, 5, 6, 7, 8, 9, 10 }));
+				AITEST_TRUE(TEXT("Tick"), ExpectedLog);
+
+				// Enter second Tree 2
+				{
+					ExpectedLog = ExpectedLog.Then("Tree2GlobalEval1", TestContext.ExepectData(AllValidReadData, { 0, 1, 2, 3, 4, 5, 6, 7, 8, 9, 10, 12 }));
+					AITEST_TRUE(TEXT("Tick Enter new tree"), ExpectedLog);
+					ExpectedLog = ExpectedLog.Then("Tree2GlobalEval2", TestContext.ExepectData(AllValidReadData, { 0, 1, 2, 3, 4, 5, 6, 7, 8, 9, 10, 12, 13 }));
+					AITEST_TRUE(TEXT("Start should EnterState"), ExpectedLog);
+					ExpectedLog = ExpectedLog.Then("Tree2GlobalTask1", TestContext.ExepectData(AllValidReadData, { 0, 1, 2, 3, 4, 5, 6, 7, 8, 9, 10, 12, 13, 14 }));
+					AITEST_TRUE(TEXT("Start should EnterState"), ExpectedLog);
+					ExpectedLog = ExpectedLog.Then("Tree2GlobalTask2", TestContext.ExepectData(AllValidReadData, { 0, 1, 2, 3, 4, 5, 6, 7, 8, 9, 10, 12, 13, 14, 15 }));
+					AITEST_TRUE(TEXT("Start should EnterState"), ExpectedLog);
+				}
+				// Exit first Tree 2
+				{
+					ExpectedLog = ExpectedLog.Then("Tree2State1Task1", TEXT("ExitState221"));
+					AITEST_TRUE(TEXT("Exit first tree2"), ExpectedLog);
+					ExpectedLog = ExpectedLog.Then("Tree2State1Task1", TestContext.ExepectData(AllValidReadData, { 0, 1, 2, 3, 4, 5, 6, 7, 8, 9, 10, 12, 13, 14, 15 }));
+					AITEST_TRUE(TEXT("Exit first tree2"), ExpectedLog);
+					ExpectedLog = ExpectedLog.Then("Tree2RootTask1", TEXT("ExitState211"));
+					AITEST_TRUE(TEXT("Exit first tree2"), ExpectedLog);
+					ExpectedLog = ExpectedLog.Then("Tree2GlobalTask2", TestContext.ExepectData(AllValidReadData, { 0, 1, 2, 3, 4, 5, 6, 7, 8, 9, 12, 13, 14, 15 })); // -10
+					AITEST_TRUE(TEXT("Exit first tree2"), ExpectedLog);
+					ExpectedLog = ExpectedLog.Then("Tree2GlobalTask1", TestContext.ExepectData(AllValidReadData, { 0, 1, 2, 3, 4, 5, 6, 7, 8, 12, 13, 14, 15 })); //-9
+					AITEST_TRUE(TEXT("Exit first tree2"), ExpectedLog);
+					ExpectedLog = ExpectedLog.Then("Tree2GlobalEval2", TestContext.ExepectData(AllValidReadData, { 0, 1, 2, 3, 4, 5, 6, 7, 12, 13, 14, 15 })); //-8
+					AITEST_TRUE(TEXT("Exit first tree2"), ExpectedLog);
+					ExpectedLog = ExpectedLog.Then("Tree2GlobalEval1", TestContext.ExepectData(AllValidReadData, { 0, 1, 2, 3, 4, 5, 6, 12, 13, 14, 15 })); //-7
+					AITEST_TRUE(TEXT("Exit first tree2"), ExpectedLog);
+				}
+				// Enter second Tree 2
+				{
+					ExpectedLog = ExpectedLog.Then("Tree2RootTask1", TEXT("EnterState211"));
+					AITEST_TRUE(TEXT("Exit first tree2"), ExpectedLog);
+					ExpectedLog = ExpectedLog.Then("Tree2State1Task1", TestContext.ExepectData(AllValidReadData, { 0, 1, 2, 3, 4, 5, 12, 13, 14, 15, 16 })); //-6+16
+					AITEST_TRUE(TEXT("Exit first tree2"), ExpectedLog);
+					ExpectedLog = ExpectedLog.Then("Tree2State1Task2", TestContext.ExepectData(AllValidReadData, { 0, 1, 2, 3, 4, 5, 12, 13, 14, 15, 16, 17 })); //+17
+					AITEST_TRUE(TEXT("Exit first tree2"), ExpectedLog);
+				}
+			}
+			{
+				TArray<int32> AllDatas = TestContext.ReadAllData();
+				AITEST_TRUE(TEXT("After enter. Same num"), AllDatas.Num() == UE_ARRAY_COUNT(AllValidReadData));
+				AITEST_TRUE(TEXT("After enter. Same Value"), AllDatas == MakeArrayView({ 101, 102, 103, 104, 121, 122, -1, -1, -1, -1, -1, -1, 201, 202, 203, 204, 221, 222 }));
+			}
+			// Stop the Exec
+			{
+				FTestStateTreeExecutionContext Exec(StateTree1, StateTree1, InstanceData);
+				Exec.Stop();
+			}
+		}
+
+		return true;
+	}
+};
+IMPLEMENT_STATE_TREE_INSTANT_TEST(FStateTreeTest_WeakContext_InstanceData, "System.StateTree.WeakContext.InstanceData");
 } // namespace UE::StateTree::Tests
 
 UE_ENABLE_OPTIMIZATION_SHIP
