[StateTree] Added transition debug events for fallback transition to Root.

diff --git a/Engine/Plugins/Runtime/StateTree/Source/StateTreeModule/Private/StateTreeExecutionContext.cpp b/Engine/Plugins/Runtime/StateTree/Source/StateTreeModule/Private/StateTreeExecutionContext.cpp
--- a/Engine/Plugins/Runtime/StateTree/Source/StateTreeModule/Private/StateTreeExecutionContext.cpp
+++ b/Engine/Plugins/Runtime/StateTree/Source/StateTreeModule/Private/StateTreeExecutionContext.cpp
@@ -5652,6 +5652,7 @@ PRAGMA_ENABLE_DEPRECATION_WARNINGS
 				FCurrentlyProcessedStateScope RootStateScope(*this, FStateTreeStateHandle::Root);
 
 				UE::StateTree::ExecutionContext::FStateHandleContext TargetState(RootFrame.StateTree, FStateTreeStateHandle::Root);
+				UE_STATETREE_DEBUG_TRANSITION_EVENT(this, FStateTreeTransitionSource(RootFrame.StateTree, EStateTreeTransitionSourceType::Internal, FStateTreeStateHandle::Root, EStateTreeTransitionPriority::Normal), EStateTreeTraceEventType::OnRequesting);
 				if (RequestTransitionInternal(RootFrame, RootFrame.ActiveStates.StateIDs[0], TargetState))
 				{
 					RequestedTransition->Source = FStateTreeTransitionSource(RootFrame.StateTree, EStateTreeTransitionSourceType::Internal, FStateTreeStateHandle::Root, EStateTreeTransitionPriority::Normal);
