[StateTree] Added frame and state scopes on checking enter conditions in state selection

diff --git a/Engine/Plugins/Runtime/StateTree/Source/StateTreeModule/Private/StateTreeExecutionContext.cpp b/Engine/Plugins/Runtime/StateTree/Source/StateTreeModule/Private/StateTreeExecutionContext.cpp
--- a/Engine/Plugins/Runtime/StateTree/Source/StateTreeModule/Private/StateTreeExecutionContext.cpp
+++ b/Engine/Plugins/Runtime/StateTree/Source/StateTreeModule/Private/StateTreeExecutionContext.cpp
@@ -4563,21 +4563,21 @@ FStateTreeExecutionContext::FTickTaskResult FStateTreeExecutionContext::TickTask
 // Deprecated
 bool FStateTreeExecutionContext::TestAllConditions(const FStateTreeExecutionFrame* CurrentParentFrame, const FStateTreeExecutionFrame& CurrentFrame, const int32 ConditionsOffset, const int32 ConditionsNum)
 {
-	return TestAllConditionsInternal<false>(CurrentParentFrame, CurrentFrame, FStateTreeStateHandle(), ConditionsOffset, ConditionsNum);
+	return TestAllConditionsInternal<false>(CurrentParentFrame, CurrentFrame, FStateTreeStateHandle(), ConditionsOffset, ConditionsNum, EStateTreeUpdatePhase::EnterConditions);
 }
 
-bool FStateTreeExecutionContext::TestAllConditionsOnActiveInstances(const FStateTreeExecutionFrame* CurrentParentFrame, const FStateTreeExecutionFrame& CurrentFrame, FStateTreeStateHandle CurrentStateHandle, const int32 ConditionsOffset, const int32 ConditionsNum)
+bool FStateTreeExecutionContext::TestAllConditionsOnActiveInstances(const FStateTreeExecutionFrame* CurrentParentFrame, const FStateTreeExecutionFrame& CurrentFrame, FStateTreeStateHandle CurrentStateHandle, const int32 ConditionsOffset, const int32 ConditionsNum, EStateTreeUpdatePhase Phase)
 {
-	return TestAllConditionsInternal<true>(CurrentParentFrame, CurrentFrame, CurrentStateHandle, ConditionsOffset, ConditionsNum);
+	return TestAllConditionsInternal<true>(CurrentParentFrame, CurrentFrame, CurrentStateHandle, ConditionsOffset, ConditionsNum, Phase);
 }
 
-bool FStateTreeExecutionContext::TestAllConditionsWithValidation(const FStateTreeExecutionFrame* CurrentParentFrame, const FStateTreeExecutionFrame& CurrentFrame, FStateTreeStateHandle CurrentStateHandle, const int32 ConditionsOffset, const int32 ConditionsNum)
+bool FStateTreeExecutionContext::TestAllConditionsWithValidation(const FStateTreeExecutionFrame* CurrentParentFrame, const FStateTreeExecutionFrame& CurrentFrame, FStateTreeStateHandle CurrentStateHandle, const int32 ConditionsOffset, const int32 ConditionsNum, EStateTreeUpdatePhase Phase)
 {
-	return TestAllConditionsInternal<false>(CurrentParentFrame, CurrentFrame, CurrentStateHandle, ConditionsOffset, ConditionsNum);
+	return TestAllConditionsInternal<false>(CurrentParentFrame, CurrentFrame, CurrentStateHandle, ConditionsOffset, ConditionsNum, Phase);
 }
 
 template<bool bOnActiveInstances>
-bool FStateTreeExecutionContext::TestAllConditionsInternal(const FStateTreeExecutionFrame* CurrentParentFrame, const FStateTreeExecutionFrame& CurrentFrame, FStateTreeStateHandle CurrentStateHandle, const int32 ConditionsOffset, const int32 ConditionsNum)
+bool FStateTreeExecutionContext::TestAllConditionsInternal(const FStateTreeExecutionFrame* CurrentParentFrame, const FStateTreeExecutionFrame& CurrentFrame, FStateTreeStateHandle CurrentStateHandle, const int32 ConditionsOffset, const int32 ConditionsNum, EStateTreeUpdatePhase Phase)
 {
 	CSV_SCOPED_TIMING_STAT_EXCLUSIVE(StateTree_TestConditions);
 
@@ -4589,6 +4589,8 @@ bool FStateTreeExecutionContext::TestAllConditionsInternal(const FStateTreeExecu
 	FCurrentlyProcessedFrameScope FrameScope(*this, CurrentParentFrame, CurrentFrame);
 	FCurrentlyProcessedStateScope StateScope(*this, CurrentStateHandle);
 
+	UE_STATETREE_DEBUG_SCOPED_PHASE(this, Phase);
+
 	TStaticArray<EStateTreeExpressionOperand, UE::StateTree::MaxExpressionIndent + 1> Operands(InPlace, EStateTreeExpressionOperand::Copy);
 	TStaticArray<bool, UE::StateTree::MaxExpressionIndent + 1> Values(InPlace, false);
 
@@ -5393,8 +5395,7 @@ bool FStateTreeExecutionContext::TriggerTransitions()
 					{
 						FCurrentlyProcessedTransitionEventScope TransitionEventScope(*this, TransitionEvent.IsValid() ? TransitionEvent.Get() : nullptr);
 						UE_STATETREE_DEBUG_TRANSITION_EVENT(this, FStateTreeTransitionSource(CurrentFrame.StateTree, FStateTreeIndex16(TransitionIndex), Transition.State, Transition.Priority), EStateTreeTraceEventType::OnEvaluating);
-						UE_STATETREE_DEBUG_SCOPED_PHASE(this, EStateTreeUpdatePhase::TransitionConditions);
-						bPassed = TestAllConditionsOnActiveInstances(CurrentParentFrame, CurrentFrame, Handler.StateHandle, Transition.ConditionsBegin, Transition.ConditionsNum);
+						bPassed = TestAllConditionsOnActiveInstances(CurrentParentFrame, CurrentFrame, Handler.StateHandle, Transition.ConditionsBegin, Transition.ConditionsNum, EStateTreeUpdatePhase::TransitionConditions);
 					}
 
 					if (bPassed)
@@ -5605,8 +5606,7 @@ PRAGMA_ENABLE_DEPRECATION_WARNINGS
 								bool bPassed = false;
 								{
 									UE_STATETREE_DEBUG_TRANSITION_EVENT(this, FStateTreeTransitionSource(CurrentFrame.StateTree, FStateTreeIndex16(TransitionIndex), Transition.State, Transition.Priority), EStateTreeTraceEventType::OnEvaluating);
-									UE_STATETREE_DEBUG_SCOPED_PHASE(this, EStateTreeUpdatePhase::TransitionConditions);
-									bPassed = TestAllConditionsOnActiveInstances(CurrentParentFrame, CurrentFrame, CurrentStateHandle, Transition.ConditionsBegin, Transition.ConditionsNum);
+									bPassed = TestAllConditionsOnActiveInstances(CurrentParentFrame, CurrentFrame, CurrentStateHandle, Transition.ConditionsBegin, Transition.ConditionsNum, EStateTreeUpdatePhase::TransitionConditions);
 								}
 
 								if (bPassed)
@@ -6605,9 +6605,7 @@ bool FStateTreeExecutionContext::SelectStateInternal(
 		if (bShouldPrerequisitesBeChecked)
 		{
 			// Check that the state can be entered
-			UE_STATETREE_DEBUG_ENTER_PHASE(this, EStateTreeUpdatePhase::EnterConditions);
-			const bool bEnterConditionsPassed = TestAllConditionsWithValidation(NextParentFrame, NextFrame, NextStateHandle, NextState.EnterConditionsBegin, NextState.EnterConditionsNum);
-			UE_STATETREE_DEBUG_EXIT_PHASE(this, EStateTreeUpdatePhase::EnterConditions);
+			const bool bEnterConditionsPassed = TestAllConditionsWithValidation(NextParentFrame, NextFrame, NextStateHandle, NextState.EnterConditionsBegin, NextState.EnterConditionsNum, EStateTreeUpdatePhase::EnterConditions);
 
 			if (!bEnterConditionsPassed)
 			{
@@ -7475,8 +7473,7 @@ bool FStateTreeExecutionContext::SelectStateInternal_TryFollowTransitions(
 
 				FCurrentlyProcessedTransitionEventScope TransitionEventScope(*this, SelectedStateTransitionEvent.IsValid() ? SelectedStateTransitionEvent.Get() : nullptr);
 				UE_STATETREE_DEBUG_TRANSITION_EVENT(this, FStateTreeTransitionSource(NextStateTree, FStateTreeIndex16(TransitionIndex), Transition.State, Transition.Priority), EStateTreeTraceEventType::OnEvaluating);
-				UE_STATETREE_DEBUG_SCOPED_PHASE(this, EStateTreeUpdatePhase::TransitionConditions);
-				bTransitionConditionsPassed = TestAllConditionsWithValidation(NextParentFrame, *NextFrame, NextStateHandle, Transition.ConditionsBegin, Transition.ConditionsNum);
+				bTransitionConditionsPassed = TestAllConditionsWithValidation(NextParentFrame, *NextFrame, NextStateHandle, Transition.ConditionsBegin, Transition.ConditionsNum, EStateTreeUpdatePhase::TransitionConditions);
 			}
 
 			if (bTransitionConditionsPassed)
diff --git a/Engine/Plugins/Runtime/StateTree/Source/StateTreeModule/Public/StateTreeExecutionContext.h b/Engine/Plugins/Runtime/StateTree/Source/StateTreeModule/Public/StateTreeExecutionContext.h
--- a/Engine/Plugins/Runtime/StateTree/Source/StateTreeModule/Public/StateTreeExecutionContext.h
+++ b/Engine/Plugins/Runtime/StateTree/Source/StateTreeModule/Public/StateTreeExecutionContext.h
@@ -1049,7 +1049,7 @@ protected:
 
 private:
 	template<bool bOnActiveInstances>
-	bool TestAllConditionsInternal(const FStateTreeExecutionFrame* CurrentParentFrame, const FStateTreeExecutionFrame& CurrentFrame, FStateTreeStateHandle CurrentStateHandle, const int32 ConditionsOffset, const int32 ConditionsNum);
+	bool TestAllConditionsInternal(const FStateTreeExecutionFrame* CurrentParentFrame, const FStateTreeExecutionFrame& CurrentFrame, FStateTreeStateHandle CurrentStateHandle, const int32 ConditionsOffset, const int32 ConditionsNum, EStateTreeUpdatePhase Phase);
 
 protected:
 	/**
@@ -1057,14 +1057,14 @@ protected:
 	 * Should be used only on active node instances, assumes valid data handles and does not consider temporary node instances.
 	 * @return True if all conditions pass.
 	 */
-	UE_API bool TestAllConditionsOnActiveInstances(const FStateTreeExecutionFrame* CurrentParentFrame, const FStateTreeExecutionFrame& CurrentFrame, FStateTreeStateHandle CurrentStateHandle, const int32 ConditionsOffset, const int32 ConditionsNum);
+	UE_API bool TestAllConditionsOnActiveInstances(const FStateTreeExecutionFrame* CurrentParentFrame, const FStateTreeExecutionFrame& CurrentFrame, FStateTreeStateHandle CurrentStateHandle, const int32 ConditionsOffset, const int32 ConditionsNum, EStateTreeUpdatePhase Phase);
 
 	/**
 	 * Checks all conditions at given range.
 	 * This version validates the data handles and looks up temporary instances.
 	 * @return True if all conditions pass.
 	 */
-	UE_API bool TestAllConditionsWithValidation(const FStateTreeExecutionFrame* CurrentParentFrame, const FStateTreeExecutionFrame& CurrentFrame, FStateTreeStateHandle CurrentStateHandle, const int32 ConditionsOffset, const int32 ConditionsNum);
+	UE_API bool TestAllConditionsWithValidation(const FStateTreeExecutionFrame* CurrentParentFrame, const FStateTreeExecutionFrame& CurrentFrame, FStateTreeStateHandle CurrentStateHandle, const int32 ConditionsOffset, const int32 ConditionsNum, EStateTreeUpdatePhase Phase);
 
 	UE_DEPRECATED(5.7, "Use TestAllConditionsWithValidation or TestAllConditionsOnActiveInstances. TestAllConditionsWithValidation can get the data from a temporary frame in state selection.")
 	UE_API bool TestAllConditions(const FStateTreeExecutionFrame* CurrentParentFrame, const FStateTreeExecutionFrame& CurrentFrame, const int32 ConditionsOffset, const int32 ConditionsNum);
