[StateTree] Fix for bug with parallel task that needs a second tick for transition that completes the execution.

diff --git a/Engine/Plugins/Runtime/StateTree/Source/StateTreeModule/Private/Tasks/StateTreeRunParallelStateTreeTask.cpp b/Engine/Plugins/Runtime/StateTree/Source/StateTreeModule/Private/Tasks/StateTreeRunParallelStateTreeTask.cpp
--- a/Engine/Plugins/Runtime/StateTree/Source/StateTreeModule/Private/Tasks/StateTreeRunParallelStateTreeTask.cpp
+++ b/Engine/Plugins/Runtime/StateTree/Source/StateTreeModule/Private/Tasks/StateTreeRunParallelStateTreeTask.cpp
@@ -110,7 +110,15 @@ void FStateTreeRunParallelStateTreeTask::TriggerTransitions(FStateTreeExecutionC
 		return;
 	}
 
+	const EStateTreeRunStatus LastTreeRunStatus = InstanceData.TreeInstanceData.GetStorage().GetExecutionState().TreeRunStatus;
 	ParallelTreeContext.TickTriggerTransitions();
+
+	const EStateTreeRunStatus NewTreeRunStatus = InstanceData.TreeInstanceData.GetStorage().GetExecutionState().TreeRunStatus;
+	if (LastTreeRunStatus != NewTreeRunStatus)
+	{
+		ensure(NewTreeRunStatus != EStateTreeRunStatus::Running);
+		Context.FinishTask(*this, NewTreeRunStatus == EStateTreeRunStatus::Succeeded ? EStateTreeFinishTaskType::Succeeded : EStateTreeFinishTaskType::Failed);
+	}
 	Context.UpdateScheduledTickRequest(InstanceData.ScheduledTickHandle, ParallelTreeContext.GetNextScheduledTick());
 }
 
