[StateTree] Add Extension. A way for a plugin to add extra information to the StateTree. They are not part of the binding system. Add a PostInternalCompilation step; extensions can be added in that step.

diff --git a/Engine/Plugins/Runtime/StateTree/Source/StateTreeEditorModule/Private/StateTreeCompiler.cpp b/Engine/Plugins/Runtime/StateTree/Source/StateTreeEditorModule/Private/StateTreeCompiler.cpp
--- a/Engine/Plugins/Runtime/StateTree/Source/StateTreeEditorModule/Private/StateTreeCompiler.cpp
+++ b/Engine/Plugins/Runtime/StateTree/Source/StateTreeEditorModule/Private/StateTreeCompiler.cpp
@@ -5,6 +5,8 @@
 #include "StateTreeAnyEnum.h"
 #include "StateTreeCompilerLog.h"
 #include "StateTreeEditorData.h"
+#include "StateTreeEditorDataExtension.h"
+#include "StateTreeExtension.h"
 #include "StateTreeEditorModule.h"
 #include "StateTreeEvaluatorBase.h"
 #include "StateTreeTaskBase.h"
@@ -474,6 +476,12 @@ bool FStateTreeCompiler::Compile(UStateTree& InStateTree)
 		StateTree->IDToTransitionMappings.Emplace(ToTransition.Key, FStateTreeIndex16(ToTransition.Value));
 	}
 
+	if (!NotifyInternalPost())
+	{
+		StateTree->ResetCompiled();
+		return false;
+	}
+
 	UE::StateTree::Compiler::FCheckOutersArchive CheckOuters(*StateTree, *EditorData, Log);
 	StateTree->Serialize(CheckOuters);
 
@@ -2650,3 +2658,60 @@ void FStateTreeCompiler::InstantiateStructSubobjects(FStructView Struct)
 		}
 	}
 }
+
+bool FStateTreeCompiler::NotifyInternalPost()
+{
+	struct FPostInternalContextImpl : UE::StateTree::Compiler::FPostInternalContext
+	{
+		virtual TNotNull<const UStateTree*> GetStateTree() const override
+		{
+			return StateTree;
+		}
+		virtual TNotNull<const UStateTreeEditorData*> GetEditorData() const override
+		{
+			return EditorData;
+		}
+
+		virtual FStateTreeCompilerLog& GetLog() const override
+		{
+			return Log;
+		}
+
+		virtual UStateTreeExtension* AddExtension(TNotNull<TSubclassOf<UStateTreeExtension>> ExtensionType) const override
+		{
+			UClass* ExtensionClass = static_cast<TSubclassOf<UStateTreeExtension>>(ExtensionType).Get();
+			UStateTreeExtension* Result = NewObject<UStateTreeExtension>(StateTree.Get(), ExtensionClass);
+			StateTree->Extensions.Add(Result);
+			return Result;
+		}
+
+		FPostInternalContextImpl(TNotNull<UStateTree*> InStateTree, TNotNull<UStateTreeEditorData*> InEditorData, FStateTreeCompilerLog& InLog)
+			: StateTree(InStateTree)
+			, EditorData(InEditorData)
+			, Log(InLog)
+		{ }
+
+		TObjectPtr<UStateTree> StateTree = nullptr;
+		TObjectPtr<UStateTreeEditorData> EditorData = nullptr;
+		FStateTreeCompilerLog& Log;
+	};
+
+	if (EditorData->Extensions.Num())
+	{
+		FPostInternalContextImpl Context = FPostInternalContextImpl(StateTree, EditorData, Log);
+
+		for (UStateTreeEditorDataExtension* Extension : EditorData->Extensions)
+		{
+			if (Extension)
+			{
+				if (!Extension->HandlePostInternalCompile(Context))
+				{
+					Log.Reportf(EMessageSeverity::Error, TEXT("An extension failed compilation. See log for more info."));
+					return false;
+				}
+			}
+		}
+	}
+
+	return true;
+}
diff --git a/Engine/Plugins/Runtime/StateTree/Source/StateTreeEditorModule/Private/StateTreeEditorData.cpp b/Engine/Plugins/Runtime/StateTree/Source/StateTreeEditorModule/Private/StateTreeEditorData.cpp
--- a/Engine/Plugins/Runtime/StateTree/Source/StateTreeEditorModule/Private/StateTreeEditorData.cpp
+++ b/Engine/Plugins/Runtime/StateTree/Source/StateTreeEditorModule/Private/StateTreeEditorData.cpp
@@ -5,6 +5,7 @@
 #include "StateTreeConditionBase.h"
 #include "StateTreeConsiderationBase.h"
 #include "StateTreeDelegates.h"
+#include "StateTreeEditorDataExtension.h"
 #include "StateTreeEditorModule.h"
 #include "StateTreeEvaluatorBase.h"
 #include "StateTreeNodeClassCache.h"
@@ -163,6 +164,14 @@ void UStateTreeEditorData::PostLoad()
 		Schema->ConditionalPostLoad();
 	}
 
+	for (UStateTreeEditorDataExtension* Extension : Extensions)
+	{
+		if (Extension)
+		{
+			Extension->ConditionalPostLoad();
+		}
+	}
+
 	VisitHierarchy([](UStateTreeState& State, UStateTreeState* ParentState) mutable 
 	{
 		State.ConditionalPostLoad();
@@ -422,6 +431,18 @@ void UStateTreeEditorData::GetAccessibleStructsInExecutionPath(const TConstArray
 	}
 }
 
+UStateTreeEditorDataExtension* UStateTreeEditorData::K2_GetExtension(TSubclassOf<UStateTreeEditorDataExtension> InExtensionType)
+{
+	for (UStateTreeEditorDataExtension* Extension : Extensions)
+	{
+		if (Extension && Extension->IsA(InExtensionType))
+		{
+			return Extension;
+		}
+	}
+	return nullptr;
+}
+
 FStateTreeBindableStructDesc UStateTreeEditorData::FindContextData(const UStruct* ObjectType, const FString ObjectNameHint) const
 {
 	if (Schema == nullptr)
diff --git a/Engine/Plugins/Runtime/StateTree/Source/StateTreeEditorModule/Private/StateTreeEditorDataExtension.cpp b/Engine/Plugins/Runtime/StateTree/Source/StateTreeEditorModule/Private/StateTreeEditorDataExtension.cpp
new file mode 100644
--- /dev/null
+++ b/Engine/Plugins/Runtime/StateTree/Source/StateTreeEditorModule/Private/StateTreeEditorDataExtension.cpp
@@ -0,0 +1,5 @@
+// Copyright Epic Games, Inc. All Rights Reserved.
+
+#include "StateTreeEditorDataExtension.h"
+
+#include UE_INLINE_GENERATED_CPP_BY_NAME(StateTreeEditorDataExtension)
diff --git a/Engine/Plugins/Runtime/StateTree/Source/StateTreeEditorModule/Private/StateTreeEditorSchema.cpp b/Engine/Plugins/Runtime/StateTree/Source/StateTreeEditorModule/Private/StateTreeEditorSchema.cpp
new file mode 100644
--- /dev/null
+++ b/Engine/Plugins/Runtime/StateTree/Source/StateTreeEditorModule/Private/StateTreeEditorSchema.cpp
@@ -0,0 +1,80 @@
+// Copyright Epic Games, Inc. All Rights Reserved.
+
+#include "StateTreeEditorSchema.h"
+#include "StateTree.h"
+#include "StateTreeEditorData.h"
+#include "StateTreeEditorModule.h"
+#include "StateTreeSchema.h"
+#include "StateTreeState.h"
+#include "StateTreeTaskBase.h"
+
+#include UE_INLINE_GENERATED_CPP_BY_NAME(StateTreeEditorSchema)
+
+void UStateTreeEditorSchema::Validate(TNotNull<UStateTree*> StateTree)
+{
+	UStateTreeEditorData* TreeData = Cast<UStateTreeEditorData>(StateTree->EditorData);
+	if (!TreeData)
+	{
+		return;
+	}
+
+	const UStateTreeSchema* Schema = TreeData->Schema;
+	if (!Schema)
+	{
+		return;
+	}
+
+	// Clear evaluators if not allowed.
+	if (Schema->AllowEvaluators() == false && TreeData->Evaluators.Num() > 0)
+	{
+		UE_LOG(LogStateTreeEditor, Warning, TEXT("%s: Resetting Evaluators due to current schema restrictions."), *StateTree->GetName());
+		TreeData->Evaluators.Reset();
+	}
+
+	TreeData->VisitHierarchy(
+		[&StateTree, Schema](UStateTreeState& State, UStateTreeState* /*ParentState*/)
+		{
+			constexpr bool bMarkDirty = false;
+			State.Modify(bMarkDirty);
+
+			// Clear enter conditions if not allowed.
+			if (Schema->AllowEnterConditions() == false && State.EnterConditions.Num() > 0)
+			{
+				UE_LOG(LogStateTreeEditor, Warning, TEXT("%s: Resetting Enter Conditions in state %s due to current schema restrictions."), *StateTree->GetName(), *State.GetName());
+				State.EnterConditions.Reset();
+			}
+
+			// Clear Utility if not allowed
+			if (Schema->AllowUtilityConsiderations() == false && State.Considerations.Num() > 0)
+			{
+				UE_LOG(LogStateTreeEditor, Warning, TEXT("%s: Resetting Utility Considerations in state %s due to current schema restrictions."), *StateTree->GetName(), *State.GetName());
+				State.Considerations.Reset();
+			}
+
+			// Keep single and many tasks based on what is allowed.
+			if (Schema->AllowMultipleTasks() == false)
+			{
+				if (State.Tasks.Num() > 0)
+				{
+					State.Tasks.Reset();
+					UE_LOG(LogStateTreeEditor, Warning, TEXT("%s: Resetting Tasks in state %s due to current schema restrictions."), *StateTree->GetName(), *State.GetName());
+				}
+
+				// Task name is the same as state name.
+				if (FStateTreeTaskBase* Task = State.SingleTask.Node.GetMutablePtr<FStateTreeTaskBase>())
+				{
+					Task->Name = State.Name;
+				}
+			}
+			else
+			{
+				if (State.SingleTask.Node.IsValid())
+				{
+					State.SingleTask.Reset();
+					UE_LOG(LogStateTreeEditor, Warning, TEXT("%s: Resetting Single Task in state %s due to current schema restrictions."), *StateTree->GetName(), *State.GetName());
+				}
+			}
+
+			return EStateTreeVisitor::Continue;
+		});
+}
diff --git a/Engine/Plugins/Runtime/StateTree/Source/StateTreeEditorModule/Public/StateTreeCompiler.h b/Engine/Plugins/Runtime/StateTree/Source/StateTreeEditorModule/Public/StateTreeCompiler.h
--- a/Engine/Plugins/Runtime/StateTree/Source/StateTreeEditorModule/Public/StateTreeCompiler.h
+++ b/Engine/Plugins/Runtime/StateTree/Source/StateTreeEditorModule/Public/StateTreeCompiler.h
@@ -16,11 +16,32 @@ struct FStateTreeStateHandle;
 class UStateTree;
 class UStateTreeState;
 class UStateTreeEditorData;
+class UStateTreeExtension;
 class UStateTreeSchema;
 struct FStateTreeEditorNode;
 struct FStateTreeStateLink;
 struct FStateTreeNodeBase;
 
+namespace UE::StateTree::Compiler
+{
+	// Compiler steps: PrePublic, PostPublic, PreInternal, PostInternal
+
+	/**
+	 * Compiler context for when the state tree fully compiled and succeeded.
+	 * This is the last compilation step.
+	 */
+	struct FPostInternalContext
+	{
+		virtual ~FPostInternalContext() = default;
+
+		virtual TNotNull<const UStateTree*> GetStateTree() const = 0;
+		virtual TNotNull<const UStateTreeEditorData*> GetEditorData() const = 0;
+		virtual FStateTreeCompilerLog& GetLog() const = 0;
+
+		virtual UStateTreeExtension* AddExtension(TNotNull<TSubclassOf<UStateTreeExtension>> ExtensionType) const = 0;
+	};
+}
+
 /**
  * Helper class to convert StateTree editor representation into a compact data.
  * Holds data needed during compiling.
@@ -37,37 +58,37 @@ public:
 	UE_API bool Compile(UStateTree& InStateTree);
 	
 private:
-
 	/** Resolves the state a transition points to, and the optional fallback for failing to enter the state. SourceState is nullptr for global tasks. */
-	UE_API bool ResolveTransitionStateAndFallback(const UStateTreeState* SourceState, const FStateTreeStateLink& Link, FStateTreeStateHandle& OutTransitionHandle, EStateTreeSelectionFallback& OutFallback) const;
-	UE_API FStateTreeStateHandle GetStateHandle(const FGuid& StateID) const;
-	UE_API UStateTreeState* GetState(const FGuid& StateID) const;
-
-	UE_API bool CreateParameters();
-	UE_API bool CreateStates();
-	UE_API bool CreateStateRecursive(UStateTreeState& State, const FStateTreeStateHandle Parent);
-	UE_API bool CreateEvaluators();
-	UE_API bool CreateGlobalTasks();
-	UE_API bool CreateStateTasksAndParameters();
-	UE_API bool CreateStateTransitions();
-	UE_API bool CreateStateConsiderations();
+	bool ResolveTransitionStateAndFallback(const UStateTreeState* SourceState, const FStateTreeStateLink& Link, FStateTreeStateHandle& OutTransitionHandle, EStateTreeSelectionFallback& OutFallback) const;
+
+	FStateTreeStateHandle GetStateHandle(const FGuid& StateID) const;
+	UStateTreeState* GetState(const FGuid& StateID) const;
+
+	bool CreateParameters();
+	bool CreateStates();
+	bool CreateStateRecursive(UStateTreeState& State, const FStateTreeStateHandle Parent);
+	bool CreateEvaluators();
+	bool CreateGlobalTasks();
+	bool CreateStateTasksAndParameters();
+	bool CreateStateTransitions();
+	bool CreateStateConsiderations();
 	
-	UE_API bool CreateBindingsForNodes(TConstArrayView<FStateTreeEditorNode> EditorNodes, FStateTreeIndex16 NodesBegin, TArray<FInstancedStruct>& Instances);
-	UE_API bool CreateBindingsForStruct(const FStateTreeBindableStructDesc& TargetStruct, FStateTreeDataView TargetValue, FStateTreeIndex16 PropertyFuncsBegin, FStateTreeIndex16 PropertyFuncsEnd, FStateTreeIndex16& OutBatchIndex);
+	bool CreateBindingsForNodes(TConstArrayView<FStateTreeEditorNode> EditorNodes, FStateTreeIndex16 NodesBegin, TArray<FInstancedStruct>& Instances);
+	bool CreateBindingsForStruct(const FStateTreeBindableStructDesc& TargetStruct, FStateTreeDataView TargetValue, FStateTreeIndex16 PropertyFuncsBegin, FStateTreeIndex16 PropertyFuncsEnd, FStateTreeIndex16& OutBatchIndex);
 
-	UE_API bool CreatePropertyFunctionsForStruct(FGuid StructID);
-	UE_API bool CreatePropertyFunction(const FStateTreeEditorNode& FuncEditorNode);
+	bool CreatePropertyFunctionsForStruct(FGuid StructID);
+	bool CreatePropertyFunction(const FStateTreeEditorNode& FuncEditorNode);
 	
-	UE_API bool CreateConditions(UStateTreeState& State, const FString& StatePath, TConstArrayView<FStateTreeEditorNode> Conditions);
-	UE_API bool CreateCondition(UStateTreeState& State, const FString& StatePath, const FStateTreeEditorNode& CondNode, const EStateTreeExpressionOperand Operand, const int8 DeltaIndent);
-	UE_API bool CreateConsiderations(UStateTreeState& State, const FString& StatePath, TConstArrayView<FStateTreeEditorNode> Considerations);
-	UE_API bool CreateConsideration(UStateTreeState& State, const FString& StatePath, const FStateTreeEditorNode& ConsiderationNode, const EStateTreeExpressionOperand Operand, const int8 DeltaIndent);
+	bool CreateConditions(UStateTreeState& State, const FString& StatePath, TConstArrayView<FStateTreeEditorNode> Conditions);
+	bool CreateCondition(UStateTreeState& State, const FString& StatePath, const FStateTreeEditorNode& CondNode, const EStateTreeExpressionOperand Operand, const int8 DeltaIndent);
+	bool CreateConsiderations(UStateTreeState& State, const FString& StatePath, TConstArrayView<FStateTreeEditorNode> Considerations);
+	bool CreateConsideration(UStateTreeState& State, const FString& StatePath, const FStateTreeEditorNode& ConsiderationNode, const EStateTreeExpressionOperand Operand, const int8 DeltaIndent);
 
 	template<class T = FStateTreeNodeBase>
 	T* CreateNodeWithSharedInstanceData(UStateTreeState* State, const FStateTreeEditorNode& EditorNode, FStateTreeBindableStructDesc& StructDesc);
 
-	UE_API bool CreateTask(UStateTreeState* State, const FStateTreeEditorNode& TaskNode, const FStateTreeDataHandle TaskDataHandle);
-	UE_API bool CreateEvaluator(const FStateTreeEditorNode& EvalNode, const FStateTreeDataHandle EvalDataHandle);
+	bool CreateTask(UStateTreeState* State, const FStateTreeEditorNode& TaskNode, const FStateTreeDataHandle TaskDataHandle);
+	bool CreateEvaluator(const FStateTreeEditorNode& EvalNode, const FStateTreeDataHandle EvalDataHandle);
 
 	struct FValidatedPathBindings
 	{
@@ -77,14 +98,17 @@ private:
 		TArray<FStateTreePropertyPathBinding> ReferenceBindings;
 	};
 
-	UE_API bool GetAndValidateBindings(const FStateTreeBindableStructDesc& TargetStruct, FStateTreeDataView TargetValue, FValidatedPathBindings& OutValidatedBindings) const;
-	UE_API bool IsPropertyOfTypeOrChild(UScriptStruct& Type, const FStateTreeBindableStructDesc& Struct, FPropertyBindingPath Path) const;
-	UE_API bool ValidateStructRef(const FStateTreeBindableStructDesc& SourceStruct, FPropertyBindingPath SourcePath,
-							const FStateTreeBindableStructDesc& TargetStruct, FPropertyBindingPath TargetPath) const;
-	UE_API bool CompileAndValidateNode(const UStateTreeState* SourceState, const FStateTreeBindableStructDesc& NodeDesc, FStructView NodeView, const FStateTreeDataView InstanceData);
+	bool GetAndValidateBindings(const FStateTreeBindableStructDesc& TargetStruct, FStateTreeDataView TargetValue, FValidatedPathBindings& OutValidatedBindings) const;
+	bool IsPropertyOfTypeOrChild(UScriptStruct& Type, const FStateTreeBindableStructDesc& Struct, FPropertyBindingPath Path) const;
+	bool ValidateStructRef(const FStateTreeBindableStructDesc& SourceStruct, FPropertyBindingPath SourcePath,
+					const FStateTreeBindableStructDesc& TargetStruct, FPropertyBindingPath TargetPath) const;
+	bool CompileAndValidateNode(const UStateTreeState* SourceState, const FStateTreeBindableStructDesc& NodeDesc, FStructView NodeView, const FStateTreeDataView InstanceData);
 
-	UE_API void InstantiateStructSubobjects(FStructView Struct);
-	
+	void InstantiateStructSubobjects(FStructView Struct);
+
+	bool NotifyInternalPost();
+
+private:
 	FStateTreeCompilerLog& Log;
 	UStateTree* StateTree = nullptr;
 	UStateTreeEditorData* EditorData = nullptr;
diff --git a/Engine/Plugins/Runtime/StateTree/Source/StateTreeEditorModule/Public/StateTreeEditorData.h b/Engine/Plugins/Runtime/StateTree/Source/StateTreeEditorModule/Public/StateTreeEditorData.h
--- a/Engine/Plugins/Runtime/StateTree/Source/StateTreeEditorModule/Public/StateTreeEditorData.h
+++ b/Engine/Plugins/Runtime/StateTree/Source/StateTreeEditorModule/Public/StateTreeEditorData.h
@@ -12,6 +12,7 @@
 
 struct FStateTreeBindableStructDesc;
 
+class UStateTreeEditorDataExtension;
 class UStateTreeSchema;
 
 namespace UE::StateTree::Editor
@@ -195,6 +196,17 @@ public:
 		GetAccessibleStructsInExecutionPath(Path, TargetStructID, OutStructDescs);
 	}
 
+	/** Find the first extension of the requested type. */
+	template<typename ExtensionType>
+	ExtensionType* GetExtension()
+	{
+		return CastChecked<ExtensionType>(K2_GetExtension(ExtensionType::StaticClass()), ECastCheckedType::NullAllowed);
+	}
+
+	/** Find the first extension of the requested type. */
+	UFUNCTION(BlueprintCallable, Category = "StateTree|Extension", Meta = (DisplayName="Get Extension", DeterminesOutputType = "ExtensionType"))
+	UE_API UStateTreeEditorDataExtension* K2_GetExtension(TSubclassOf<UStateTreeEditorDataExtension> ExtensionType);
+
 	UE_API void ReparentStates();
 	
 	// StateTree Builder API
@@ -341,10 +353,14 @@ private:
 #endif
 
 public:
-	/** Schema describing which inputs, evaluators, and tasks a StateTree can contain */	
-	UPROPERTY(EditDefaultsOnly, Instanced, Category = Common, NoClear)
+	/** Schema describing which inputs, evaluators, and tasks a StateTree can contain. */
+	UPROPERTY(EditDefaultsOnly, Instanced, Category = "Common", NoClear)
 	TObjectPtr<UStateTreeSchema> Schema = nullptr;
 
+	/** The editor data extensions. A place to add extra information for plugins. */
+	UPROPERTY(EditDefaultsOnly, Category = "Extension", NoClear)
+	TArray<TObjectPtr<UStateTreeEditorDataExtension>> Extensions;
+
 	/** Public parameters that could be used for bindings within the Tree. */
 	UE_DEPRECATED(5.6, "Public access to RootParameters is deprecated. Use GetRootParametersPropertyBag")
 	UPROPERTY(meta = (DeprecatedProperty))
diff --git a/Engine/Plugins/Runtime/StateTree/Source/StateTreeEditorModule/Public/StateTreeEditorDataExtension.h b/Engine/Plugins/Runtime/StateTree/Source/StateTreeEditorModule/Public/StateTreeEditorDataExtension.h
new file mode 100644
--- /dev/null
+++ b/Engine/Plugins/Runtime/StateTree/Source/StateTreeEditorModule/Public/StateTreeEditorDataExtension.h
@@ -0,0 +1,34 @@
+// Copyright Epic Games, Inc. All Rights Reserved.
+
+#pragma once
+
+#include "UObject/Object.h"
+#include "StateTreeEditorDataExtension.generated.h"
+
+namespace UE::StateTree::Compiler
+{
+	struct FPostInternalContext;
+}
+
+class UStateTreeEditorData;
+
+/**
+ * Extension for the editor data of the state tree asset.
+ */
+UCLASS(Abstract, DefaultToInstanced, EditInlineNew, Within=StateTreeEditorData, MinimalAPI)
+class UStateTreeEditorDataExtension : public UObject
+{
+	GENERATED_BODY()
+
+public:
+	virtual bool HandlePostInternalCompile(const UE::StateTree::Compiler::FPostInternalContext& Context)
+	{
+		return true;
+	}
+
+protected:
+	UStateTreeEditorData* GetStateTreeEditorData() const
+	{
+		return GetOuterUStateTreeEditorData();
+	}
+};
diff --git a/Engine/Plugins/Runtime/StateTree/Source/StateTreeEditorModule/Public/StateTreeEditorSchema.h b/Engine/Plugins/Runtime/StateTree/Source/StateTreeEditorModule/Public/StateTreeEditorSchema.h
new file mode 100644
--- /dev/null
+++ b/Engine/Plugins/Runtime/StateTree/Source/StateTreeEditorModule/Public/StateTreeEditorSchema.h
@@ -0,0 +1,28 @@
+// Copyright Epic Games, Inc. All Rights Reserved.
+
+#pragma once
+
+#include "UObject/Object.h"
+#include "StateTreeEditorSchema.generated.h"
+
+#define UE_API STATETREEEDITORMODULE_API
+
+class UStateTree;
+
+/**
+ * Schema describing how a StateTree is edited.
+ */
+UCLASS(MinimalAPI)
+class UStateTreeEditorSchema : public UObject
+{
+	GENERATED_BODY()
+	
+public:
+	/*
+	 * Validates and applies the schema restrictions on the StateTree.
+	 * Also serves as the "safety net" of fixing up editor data following an editor operation.
+	 */
+	UE_API virtual void Validate(TNotNull<UStateTree*> StateTree);
+};
+
+#undef UE_API
diff --git a/Engine/Plugins/Runtime/StateTree/Source/StateTreeModule/Private/StateTree.cpp b/Engine/Plugins/Runtime/StateTree/Source/StateTreeModule/Private/StateTree.cpp
--- a/Engine/Plugins/Runtime/StateTree/Source/StateTreeModule/Private/StateTree.cpp
+++ b/Engine/Plugins/Runtime/StateTree/Source/StateTreeModule/Private/StateTree.cpp
@@ -7,6 +7,7 @@
 #include "StateTreeEvaluatorBase.h"
 #include "StateTreeConditionBase.h"
 #include "StateTreeConsiderationBase.h"
+#include "StateTreeExtension.h"
 #include "StateTreeModuleImpl.h"
 #include "AssetRegistry/AssetData.h"
 #include "Misc/ScopeRWLock.h"
@@ -102,6 +103,21 @@ FGuid UStateTree::GetTransitionIdFromIndex(const FStateTreeIndex16 Index) const
 	return Entry != nullptr ? Entry->Id : FGuid();
 }
 
+const UStateTreeExtension* UStateTree::K2_GetExtension(TSubclassOf<UStateTreeExtension> InExtensionType) const
+{
+	for (const UStateTreeExtension* Extension : Extensions)
+	{
+		if (ensureMsgf(Extension, TEXT("The extension is invalid. Make sure it's not created in an editor only module.")))
+		{
+			if (Extension->IsA(InExtensionType))
+			{
+				return Extension;
+			}
+		}
+	}
+	return nullptr;
+}
+
 UE_AUTORTFM_ALWAYS_OPEN
 static int32 GetThreadIndexForSharedInstanceData()
 {
@@ -230,6 +246,7 @@ void UStateTree::ResetCompiled()
 	SharedInstanceData.Reset();
 	ContextDataDescs.Reset();
 	PropertyBindings.Reset();
+	Extensions.Reset();
 	Parameters.Reset();
 	ParameterDataType = EStateTreeParameterDataType::GlobalParameterData;
 	IDToStateMappings.Reset();
@@ -394,6 +411,19 @@ void UStateTree::PostLoad()
 	}
 #endif
 
+	if (Schema)
+	{
+		Schema->ConditionalPostLoad();
+	}
+
+	for (UStateTreeExtension* Extension : Extensions)
+	{
+		if (Extension)
+		{
+			Extension->ConditionalPostLoad();
+		}
+	}
+
 	if (!Link())
 	{
 		UE_LOG(LogStateTree, Log, TEXT("%s failed to link. Asset will not be usable at runtime."), *GetPathName());
@@ -597,6 +627,7 @@ bool UStateTree::Link()
 			}
 		}
 
+		// Schema
 		if (Schema)
 		{
 			const bool bSchemaLinkSucceeded = Schema->Link(Linker);
@@ -607,6 +638,23 @@ bool UStateTree::Link()
 			}
 		}
 
+		// Extensions
+		if (Extensions.Num())
+		{
+			for (UStateTreeExtension* Extension : Extensions)
+			{
+				if (Extension)
+				{
+					const bool bLinkSucceeded = Extension->Link(Linker);
+					if (!bLinkSucceeded || Linker.GetStatus() == EStateTreeLinkerStatus::Failed)
+					{
+						UE_LOG(LogStateTree, Error, TEXT("%s: extension failed to resolve its references."), *GetPathName());
+						return false;
+					}
+				}
+			}
+		}
+
 		ExternalDataDescs = Linker.GetExternalDataDescs();
 	}
 
@@ -1005,6 +1053,7 @@ TArray<FStateTreeMemoryUsage> UStateTree::CalculateEstimatedMemoryUsage() const
 	const int32 EvalMemUsageIndex = MemoryUsages.Emplace(TEXT("Evaluators"));
 	const int32 GlobalTaskMemUsageIndex = MemoryUsages.Emplace(TEXT("GlobalTask"));
 	const int32 SharedMemUsageIndex = MemoryUsages.Emplace(TEXT("Shared Data"));
+	const int32 ExtensionMemUsageIndex = MemoryUsages.Emplace(TEXT("Extensions"));
 
 	auto GetRootStateHandle = [this](const FStateTreeStateHandle InState) -> FStateTreeStateHandle
 	{
@@ -1190,6 +1239,17 @@ TArray<FStateTreeMemoryUsage> UStateTree::CalculateEstimatedMemoryUsage() const
 	SharedMemUsage.NodeCount = SharedInstanceData.Num();
 	SharedMemUsage.EstimatedMemoryUsage = SharedInstanceData.GetEstimatedMemoryUsage();
 
+	// Extensions
+	FStateTreeMemoryUsage& ExtensionMemUsage = MemoryUsages[ExtensionMemUsageIndex];
+	for (UStateTreeExtension* Extension : Extensions)
+	{
+		if (Extension)
+		{
+			ExtensionMemUsage.AddUsage(Extension);
+			++ExtensionMemUsage.NodeCount;
+		}
+	}
+
 	return MemoryUsages;
 }
 
@@ -1265,6 +1325,15 @@ FString UStateTree::DebugInternalLayoutAsString() const
 			DebugString << TEXT('\n');
 		};
 
+	if (Schema)
+	{
+		DebugString.Appendf(TEXT("Schema: %s\n"), *WriteToString<128>(Schema->GetFName()));
+	}
+	else
+	{
+		DebugString.Append(TEXT("Schema: [None]\n"));
+	}
+
 	// Tree items (e.g. tasks, evaluators, conditions)
 	DebugString.Appendf(TEXT("\nNodes(%d)\n"), Nodes.Num());
 	for (int32 Index = 0; Index < Nodes.Num(); Index++)
@@ -1420,8 +1489,14 @@ FString UStateTree::DebugInternalLayoutAsString() const
 	}
 
 	// Tasks
-	DebugString.Appendf(TEXT("\nTasks\n  [ (Idx) | %-30s | %-30s | %8s | %14s ]\n"),
-		TEXT("State"), TEXT("Name"), TEXT("Bindings"), TEXT("Struct Idx"));
+	int32 NumberOfTasks = GlobalTasksNum;
+	for (const FCompactStateTreeState& State : States)
+	{
+		NumberOfTasks += State.TasksNum;
+	}
+
+	DebugString.Appendf(TEXT("\nTasks(%d)\n  [ (Idx) | %-30s | %-30s | %8s | %14s ]\n"),
+		NumberOfTasks, TEXT("State"), TEXT("Name"), TEXT("Bindings"), TEXT("Struct Idx"));
 	for (const FCompactStateTreeState& State : States)
 	{
 		if (State.TasksNum)
@@ -1467,6 +1542,23 @@ FString UStateTree::DebugInternalLayoutAsString() const
 		}
 	}
 
+	// Extensions
+	DebugString.Appendf(TEXT("\nExtensions(%d)\n"), Extensions.Num());
+	for (int32 Index = 0; Index < Extensions.Num(); ++Index)
+	{
+		if (const UStateTreeExtension* Extension = Extensions[Index])
+		{
+			if (Extension)
+			{
+				DebugString.Appendf(TEXT("  %s\n"), *WriteToString<128>(Extension->GetFName()));
+			}
+			else
+			{
+				DebugString.Append(TEXT("  [None]\n"));
+			}
+		}
+	}
+
 	return DebugString.ToString();
 }
 #endif // WITH_EDITOR || WITH_STATETREE_DEBUG
diff --git a/Engine/Plugins/Runtime/StateTree/Source/StateTreeModule/Private/StateTreeExtension.cpp b/Engine/Plugins/Runtime/StateTree/Source/StateTreeModule/Private/StateTreeExtension.cpp
new file mode 100644
--- /dev/null
+++ b/Engine/Plugins/Runtime/StateTree/Source/StateTreeModule/Private/StateTreeExtension.cpp
@@ -0,0 +1,5 @@
+// Copyright Epic Games, Inc. All Rights Reserved.
+
+#include "StateTreeExtension.h"
+
+#include UE_INLINE_GENERATED_CPP_BY_NAME(StateTreeExtension)
diff --git a/Engine/Plugins/Runtime/StateTree/Source/StateTreeModule/Public/StateTree.h b/Engine/Plugins/Runtime/StateTree/Source/StateTreeModule/Public/StateTree.h
--- a/Engine/Plugins/Runtime/StateTree/Source/StateTreeModule/Public/StateTree.h
+++ b/Engine/Plugins/Runtime/StateTree/Source/StateTreeModule/Public/StateTree.h
@@ -13,6 +13,7 @@
 
 #define UE_API STATETREEMODULE_API
 
+class UStateTreeExtension;
 class UUserDefinedStruct;
 
 template<bool>
@@ -233,6 +234,23 @@ public:
 		return PropertyBindings;
 	}
 
+	/** @return View of all extensions. */
+	TConstArrayView<UStateTreeExtension*> GetExtensions() const
+	{
+		return Extensions;
+	}
+
+	/** Find the first extension of the requested type. */
+	template<typename ExtensionType>
+	const ExtensionType* GetExtension() const
+	{
+		return CastChecked<const ExtensionType>(K2_GetExtension(ExtensionType::StaticClass()), ECastCheckedType::NullAllowed);
+	}
+
+	/** Find the first extension of the requested type. */
+	UFUNCTION(BlueprintCallable, Category = "StateTree|Extension", Meta = (DisplayName="Get Extension", DeterminesOutputType = "ExtensionType"))
+	UE_API const UStateTreeExtension* K2_GetExtension(TSubclassOf<UStateTreeExtension> ExtensionType) const;
+
 	/** @return True if there is any global tasks need ticking. */
 	bool DoesRequestTickGlobalTasks(bool bHasEvents) const
 	{
@@ -385,6 +403,10 @@ private:
 	UPROPERTY()
 	FStateTreePropertyBindings PropertyBindings;
 
+	/** The asset extensions. A place to add extra information for plugins. */
+	UPROPERTY()
+	TArray<TObjectPtr<UStateTreeExtension>> Extensions;
+
 	/** Mapping of state guid for the Editor and state handles, created at compilation. */
 	UPROPERTY()
 	TArray<FStateTreeStateIdToHandle> IDToStateMappings;
diff --git a/Engine/Plugins/Runtime/StateTree/Source/StateTreeModule/Public/StateTreeExtension.h b/Engine/Plugins/Runtime/StateTree/Source/StateTreeModule/Public/StateTreeExtension.h
new file mode 100644
--- /dev/null
+++ b/Engine/Plugins/Runtime/StateTree/Source/StateTreeModule/Public/StateTreeExtension.h
@@ -0,0 +1,31 @@
+// Copyright Epic Games, Inc. All Rights Reserved.
+
+#pragma once
+
+#include "UObject/Object.h"
+#include "StateTreeExtension.generated.h"
+
+class UStateTree;
+struct FStateTreeLinker;
+
+/**
+ * Extension for the state tree asset.
+ */
+UCLASS(Abstract, DefaultToInstanced, Within=StateTree, MinimalAPI)
+class UStateTreeExtension : public UObject
+{
+	GENERATED_BODY()
+
+public:
+	/** Resolves references to other StateTree data. */
+	virtual bool Link(FStateTreeLinker& Linker)
+	{
+		return true;
+	}
+
+protected:
+	UStateTree* GetStateTree() const
+	{
+		return GetOuterUStateTree();
+	}
+};
