[StateTree] Added StateTreeEditorSchema. Similar to the regular schema, but for editor module. It will be the hook to add/modify editor data and to modify the editor.

diff --git a/Engine/Plugins/Runtime/StateTree/Source/StateTreeEditorModule/Private/StateTreeEditingSubsystem.cpp b/Engine/Plugins/Runtime/StateTree/Source/StateTreeEditorModule/Private/StateTreeEditingSubsystem.cpp
--- a/Engine/Plugins/Runtime/StateTree/Source/StateTreeEditorModule/Private/StateTreeEditingSubsystem.cpp
+++ b/Engine/Plugins/Runtime/StateTree/Source/StateTreeEditorModule/Private/StateTreeEditingSubsystem.cpp
@@ -8,6 +8,7 @@
 #include "StateTreeDelegates.h"
 #include "StateTreeEditorData.h"
 #include "StateTreeEditorModule.h"
+#include "StateTreeEditorSchema.h"
 #include "StateTreeObjectHash.h"
 #include "StateTreeTaskBase.h"
 #include "UObject/UObjectGlobals.h"
@@ -100,62 +101,56 @@ TSharedRef<SWidget> UStateTreeEditingSubsystem::GetStateTreeView(TSharedRef<FSta
 void UStateTreeEditingSubsystem::ValidateStateTree(TNotNull<UStateTree*> InStateTree)
 {
 	auto FixChangedStateLinkName = [](FStateTreeStateLink& StateLink, const TMap<FGuid, FName>& IDToName) -> bool
-	{
-		if (StateLink.ID.IsValid())
 		{
-			const FName* Name = IDToName.Find(StateLink.ID);
-			if (Name == nullptr)
-			{
-				// Missing link, we'll show these in the UI
-				return false;
-			}
-			if (StateLink.Name != *Name)
+			if (StateLink.ID.IsValid())
 			{
-				// Name changed, fix!
-				StateLink.Name = *Name;
-				return true;
+				const FName* Name = IDToName.Find(StateLink.ID);
+				if (Name == nullptr)
+				{
+					// Missing link, we'll show these in the UI
+					return false;
+				}
+				if (StateLink.Name != *Name)
+				{
+					// Name changed, fix!
+					StateLink.Name = *Name;
+					return true;
+				}
 			}
-		}
-		return false;
-	};
+			return false;
+		};
 
-	auto ValidateLinkedStates = [FixChangedStateLinkName](const UStateTree& StateTree)
-	{
-		UStateTreeEditorData* TreeData = Cast<UStateTreeEditorData>(StateTree.EditorData);
-		if (!TreeData)
+	auto ValidateLinkedStates = [&FixChangedStateLinkName](TNotNull<UStateTreeEditorData*> TreeData)
 		{
-			return;
-		}
-
-		// Make sure all state links are valid and update the names if needed.
+			// Make sure all state links are valid and update the names if needed.
 
-		// Create ID to state name map.
-		TMap<FGuid, FName> IDToName;
+			// Create ID to state name map.
+			TMap<FGuid, FName> IDToName;
 
-		TreeData->VisitHierarchy([&IDToName](const UStateTreeState& State, UStateTreeState* /*ParentState*/)
-		{
-			IDToName.Add(State.ID, State.Name);
-			return EStateTreeVisitor::Continue;
-		});
+			TreeData->VisitHierarchy([&IDToName](const UStateTreeState& State, UStateTreeState* /*ParentState*/)
+			{
+				IDToName.Add(State.ID, State.Name);
+				return EStateTreeVisitor::Continue;
+			});
 		
-		// Fix changed names.
-		TreeData->VisitHierarchy([&IDToName, FixChangedStateLinkName](UStateTreeState& State, UStateTreeState* /*ParentState*/)
-		{
-			constexpr bool bMarkDirty = false;
-			State.Modify(bMarkDirty);
-			if (State.Type == EStateTreeStateType::Linked)
+			// Fix changed names.
+			TreeData->VisitHierarchy([&IDToName, FixChangedStateLinkName](UStateTreeState& State, UStateTreeState* /*ParentState*/)
 			{
-				FixChangedStateLinkName(State.LinkedSubtree, IDToName);
-			}
+				constexpr bool bMarkDirty = false;
+				State.Modify(bMarkDirty);
+				if (State.Type == EStateTreeStateType::Linked)
+				{
+					FixChangedStateLinkName(State.LinkedSubtree, IDToName);
+				}
 					
-			for (FStateTreeTransition& Transition : State.Transitions)
-			{
-				FixChangedStateLinkName(Transition.State, IDToName);
-			}
+				for (FStateTreeTransition& Transition : State.Transitions)
+				{
+					FixChangedStateLinkName(Transition.State, IDToName);
+				}
 
-			return EStateTreeVisitor::Continue;
-		});
-	};
+				return EStateTreeVisitor::Continue;
+			});
+		};
 
 	auto FixEditorData = [](TNotNull<UStateTree*> StateTree)
 	{
@@ -209,131 +204,65 @@ void UStateTreeEditingSubsystem::ValidateStateTree(TNotNull<UStateTree*> InState
 			StateTree->EditorData = EditorData;
 
 			// Trash the previous EditorData
-			FName TrashName = MakeUniqueObjectName(GetTransientPackage(), UStateTree::StaticClass(), *FString::Printf(TEXT("TRASH_%s"), *UStateTree::StaticClass()->GetName()));
-			UStateTree* TransientOuter = NewObject<UStateTree>(GetTransientPackage(), TrashName, RF_Transient);
-			ERenameFlags RenameFlags = REN_DoNotDirty | REN_DontCreateRedirectors;
-			PreviousEditorData->Rename(*TrashName.ToString(), TransientOuter, RenameFlags);
+			const FName TrashStateTreeName = MakeUniqueObjectName(GetTransientPackage(), UStateTree::StaticClass(), *FString::Printf(TEXT("TRASH_%s"), *UStateTree::StaticClass()->GetName()));
+			UStateTree* TransientOuter = NewObject<UStateTree>(GetTransientPackage(), TrashStateTreeName, RF_Transient);
+			const FName TrashSchemaName = *FString::Printf(TEXT("TRASH_%s"), *PreviousEditorData->GetName());
+			constexpr ERenameFlags RenameFlags = REN_DoNotDirty | REN_DontCreateRedirectors;
+			PreviousEditorData->Rename(*TrashSchemaName.ToString(), TransientOuter, RenameFlags);
 			PreviousEditorData->SetFlags(RF_Transient);
 		}
 	};
 
-	auto UpdateParents = [](const UStateTree& StateTree)
-	{
-		UStateTreeEditorData* TreeData = Cast<UStateTreeEditorData>(StateTree.EditorData);
-		if (!TreeData)
-		{
-			return;
-		}
-
-		TreeData->ReparentStates();
-	};
-
-	auto ApplySchema =[](const UStateTree& StateTree)
-	{
-		UStateTreeEditorData* TreeData = Cast<UStateTreeEditorData>(StateTree.EditorData);
-		if (!TreeData)
-		{
-			return;
-		
-		}
-		const UStateTreeSchema* Schema = TreeData->Schema;
-		if (!Schema)
-		{
-			return;
-		}
-		
-		// Clear evaluators if not allowed.
-		if (Schema->AllowEvaluators() == false && TreeData->Evaluators.Num() > 0)
-		{
-			UE_LOG(LogStateTreeEditor, Warning, TEXT("%s: Resetting Evaluators due to current schema restrictions."), *GetNameSafe(&StateTree));
-			TreeData->Evaluators.Reset();
-		}
-
-		TreeData->VisitHierarchy([&StateTree, Schema](UStateTreeState& State, UStateTreeState* /*ParentState*/)
+	auto FixEditorSchema = [](TNotNull<UStateTreeEditorData*> EditorData)
 		{
-			constexpr bool bMarkDirty = false;
-			State.Modify(bMarkDirty);
-
-			// Clear enter conditions if not allowed.
-			if (Schema->AllowEnterConditions() == false && State.EnterConditions.Num() > 0)
+			TSubclassOf<const UStateTreeSchema> SchemaClass = EditorData->Schema ? EditorData->Schema->GetClass() : nullptr;
+			if (SchemaClass.Get() == nullptr)
 			{
-				UE_LOG(LogStateTreeEditor, Warning, TEXT("%s: Resetting Enter Conditions in state %s due to current schema restrictions."), *GetNameSafe(&StateTree), *GetNameSafe(&State));
-				State.EnterConditions.Reset();
+				return;
 			}
 
-			// Clear Utility if not allowed
-			if (Schema->AllowUtilityConsiderations() == false && State.Considerations.Num() > 0)
+			TNonNullSubclassOf<UStateTreeEditorSchema> EditorSchemaClass = FStateTreeEditorModule::GetModule().GetEditorSchemaClass(SchemaClass.Get());
+			if (EditorData->EditorSchema == nullptr)
 			{
-				UE_LOG(LogStateTreeEditor, Warning, TEXT("%s: Resetting Utility Considerations in state %s due to current schema restrictions."), *GetNameSafe(&StateTree), *GetNameSafe(&State));
-				State.Considerations.Reset();
-			}
-
-			// Keep single and many tasks based on what is allowed.
-			if (Schema->AllowMultipleTasks() == false)
-			{
-				if (State.Tasks.Num() > 0)
-				{
-					State.Tasks.Reset();
-					UE_LOG(LogStateTreeEditor, Warning, TEXT("%s: Resetting Tasks in state %s due to current schema restrictions."), *GetNameSafe(&StateTree), *GetNameSafe(&State));
-				}
-				
-				// Task name is the same as state name.
-				if (FStateTreeTaskBase* Task = State.SingleTask.Node.GetMutablePtr<FStateTreeTaskBase>())
-				{
-					Task->Name = State.Name;
-				}
+				EditorData->EditorSchema = NewObject<UStateTreeEditorSchema>(EditorData, EditorSchemaClass.Get(), FName(), RF_Transactional);
 			}
-			else
+			else if (!EditorData->EditorSchema->IsA(EditorSchemaClass.Get()))
 			{
-				if (State.SingleTask.Node.IsValid())
-				{
-					State.SingleTask.Reset();
-					UE_LOG(LogStateTreeEditor, Warning, TEXT("%s: Resetting Single Task in state %s due to current schema restrictions."), *GetNameSafe(&StateTree), *GetNameSafe(&State));
-				}
+				// The current EditorSchema is not of the correct type. The data needs to be patched by the schema desired editor data subclass.
+				UStateTreeEditorSchema* PreviousEditorSchema = EditorData->EditorSchema;
+				EditorData->EditorSchema = CastChecked<UStateTreeEditorSchema>(StaticDuplicateObject(PreviousEditorSchema, EditorData, FName(), RF_Transactional, EditorSchemaClass.Get()));
+
+				// Trash the previous EditorData
+				const FName TrashName = MakeUniqueObjectName(GetTransientPackage(), UStateTreeEditorSchema::StaticClass(), *FString::Printf(TEXT("TRASH_%s"), *PreviousEditorSchema->GetName()));
+				constexpr ERenameFlags RenameFlags = REN_DoNotDirty | REN_DontCreateRedirectors;
+				PreviousEditorSchema->Rename(*TrashName.ToString(), PreviousEditorSchema, RenameFlags);
+				PreviousEditorSchema->SetFlags(RF_Transient);
 			}
-			
-			return EStateTreeVisitor::Continue;
-		});
-	};
+		};
 
-	auto RemoveUnusedBindings = [](const UStateTree& StateTree)
-	{
-		UStateTreeEditorData* TreeData = Cast<UStateTreeEditorData>(StateTree.EditorData);
-		if (!TreeData)
-		{
-			return;
-		}
-
-		TMap<FGuid, const FPropertyBindingDataView> AllStructValues;
-		TreeData->GetAllStructValues(AllStructValues);
-		TreeData->GetPropertyEditorBindings()->RemoveInvalidBindings(AllStructValues);
-	};
-
-	auto UpdateLinkedStateParameters = [](const UStateTree& StateTree)
-	{
-		UStateTreeEditorData* TreeData = Cast<UStateTreeEditorData>(StateTree.EditorData);
-		if (!TreeData)
+	auto RemoveUnusedBindings = [](TNotNull<UStateTreeEditorData*> EditorData)
 		{
-			return;
-		}
+			TMap<FGuid, const FPropertyBindingDataView> AllStructValues;
+			EditorData->GetAllStructValues(AllStructValues);
+			EditorData->GetPropertyEditorBindings()->RemoveInvalidBindings(AllStructValues);
+		};
 
-		const EStateTreeVisitor Result = TreeData->VisitHierarchy([](UStateTreeState& State, UStateTreeState* /*ParentState*/)
+	auto UpdateLinkedStateParameters = [](TNotNull<UStateTreeEditorData*> TreeData)
 		{
-			if (State.Type == EStateTreeStateType::Linked
-				|| State.Type == EStateTreeStateType::LinkedAsset)
+			const EStateTreeVisitor Result = TreeData->VisitHierarchy([](UStateTreeState& State, UStateTreeState* /*ParentState*/)
 			{
-				constexpr bool bMarkDirty = false;
-				State.Modify(bMarkDirty);
-				State.UpdateParametersFromLinkedSubtree();
-			}
-			return EStateTreeVisitor::Continue;
-		});
-	};
+				if (State.Type == EStateTreeStateType::Linked
+					|| State.Type == EStateTreeStateType::LinkedAsset)
+				{
+					constexpr bool bMarkDirty = false;
+					State.Modify(bMarkDirty);
+					State.UpdateParametersFromLinkedSubtree();
+				}
+				return EStateTreeVisitor::Continue;
+			});
+		};
 
-	auto UpdateTransactionalFlags = [](TNotNull<UStateTree*> StateTree)
-	{
-		UStateTreeEditorData* EditorData = Cast<UStateTreeEditorData>(StateTree->EditorData);
-		if (EditorData)
+	auto UpdateTransactionalFlags = [](TNotNull<UStateTreeEditorData*> EditorData)
 		{
 			for (UStateTreeState* SubTree : EditorData->SubTrees)
 			{
@@ -353,23 +282,28 @@ void UStateTreeEditingSubsystem::ValidateStateTree(TNotNull<UStateTree*> InState
 					}
 				}
 			}
-		}
-	};
+		};
 
 	FixEditorData(InStateTree);
-
 	if (InStateTree->EditorData)
-		{
-			constexpr bool bMarkDirty = false;
+	{
+		constexpr bool bMarkDirty = false;
 		InStateTree->EditorData->Modify(bMarkDirty);
-	}
 
-	UpdateParents(*InStateTree);
-	ApplySchema(*InStateTree);
-	RemoveUnusedBindings(*InStateTree);
-	ValidateLinkedStates(*InStateTree);
-	UpdateLinkedStateParameters(*InStateTree);
-	UpdateTransactionalFlags(InStateTree);
+		UStateTreeEditorData* EditorData = CastChecked<UStateTreeEditorData>(InStateTree->EditorData);
+		FixEditorSchema(EditorData);
+
+		EditorData->ReparentStates();
+		if (EditorData->EditorSchema)
+		{
+			EditorData->EditorSchema->Validate(InStateTree);
+		}
+
+		RemoveUnusedBindings(EditorData);
+		ValidateLinkedStates(EditorData);
+		UpdateLinkedStateParameters(EditorData);
+		UpdateTransactionalFlags(EditorData);
+	}
 }
 
 uint32 UStateTreeEditingSubsystem::CalculateStateTreeHash(TNotNull<const UStateTree*> InStateTree)
diff --git a/Engine/Plugins/Runtime/StateTree/Source/StateTreeEditorModule/Private/StateTreeEditorData.cpp b/Engine/Plugins/Runtime/StateTree/Source/StateTreeEditorModule/Private/StateTreeEditorData.cpp
--- a/Engine/Plugins/Runtime/StateTree/Source/StateTreeEditorModule/Private/StateTreeEditorData.cpp
+++ b/Engine/Plugins/Runtime/StateTree/Source/StateTreeEditorModule/Private/StateTreeEditorData.cpp
@@ -7,6 +7,7 @@
 #include "StateTreeDelegates.h"
 #include "StateTreeEditorDataExtension.h"
 #include "StateTreeEditorModule.h"
+#include "StateTreeEditorSchema.h"
 #include "StateTreeEvaluatorBase.h"
 #include "StateTreeNodeClassCache.h"
 #include "StateTreePropertyFunctionBase.h"
@@ -163,6 +164,10 @@ void UStateTreeEditorData::PostLoad()
 	{
 		Schema->ConditionalPostLoad();
 	}
+	if (EditorSchema)
+	{
+		EditorSchema->ConditionalPostLoad();
+	}
 
 	for (UStateTreeEditorDataExtension* Extension : Extensions)
 	{
@@ -177,12 +182,12 @@ void UStateTreeEditorData::PostLoad()
 		State.ConditionalPostLoad();
 		return EStateTreeVisitor::Continue;
 	});
+	CallPostLoadOnNodes();
 
 	ReparentStates();
 	FixObjectNodes();
 	FixDuplicateIDs();
 	UpdateBindingsInstanceStructs();
-	CallPostLoadOnNodes();
 }
 
 void UStateTreeEditorData::PostEditChangeChainProperty(FPropertyChangedChainEvent& PropertyChangedEvent)
@@ -202,7 +207,8 @@ void UStateTreeEditorData::PostEditChangeChainProperty(FPropertyChangedChainEven
 		checkf(StateTree, TEXT("UStateTreeEditorData should only be allocated within a UStateTree"));
 		
 		const FName MemberName = MemberProperty->GetFName();
-		if (MemberName == GET_MEMBER_NAME_CHECKED(UStateTreeEditorData, Schema))
+		if (MemberName == GET_MEMBER_NAME_CHECKED(UStateTreeEditorData, Schema)
+			|| MemberName == GET_MEMBER_NAME_CHECKED(UStateTreeEditorData, EditorSchema))
 		{			
 			UE::StateTree::Delegates::OnSchemaChanged.Broadcast(*StateTree);
 		}
diff --git a/Engine/Plugins/Runtime/StateTree/Source/StateTreeEditorModule/Private/StateTreeEditorModule.cpp b/Engine/Plugins/Runtime/StateTree/Source/StateTreeEditorModule/Private/StateTreeEditorModule.cpp
--- a/Engine/Plugins/Runtime/StateTree/Source/StateTreeEditorModule/Private/StateTreeEditorModule.cpp
+++ b/Engine/Plugins/Runtime/StateTree/Source/StateTreeEditorModule/Private/StateTreeEditorModule.cpp
@@ -33,6 +33,7 @@
 #include "StateTreeEditor.h"
 #include "StateTreeEditorCommands.h"
 #include "StateTreeEditorData.h"
+#include "StateTreeEditorSchema.h"
 #include "StateTreeEditorStyle.h"
 #include "StateTreeNodeClassCache.h"
 #include "StateTreePropertyFunctionBase.h"
@@ -186,39 +187,144 @@ TSharedPtr<FStateTreeNodeClassCache> FStateTreeEditorModule::GetNodeClassCache()
 	return NodeClassCache;
 }
 
+bool FStateTreeEditorModule::FEditorTypes::HasData() const
+{
+	return EditorData.IsValid() || EditorSchema.IsValid();
+}
+
+namespace UE::StateTreeEditor::Private
+{
+	TOptional<int32> GetDepth(TNotNull<const UStruct*> Struct, TNotNull<const UStruct*> MatchingParent)
+	{
+		int32 Depth = 0;
+		if (Struct == MatchingParent)
+		{
+			return Depth;
+		}
+
+		for (const UStruct* TempStruct : Struct->GetSuperStructIterator())
+		{
+			++Depth;
+			if (TempStruct == MatchingParent)
+			{
+				return Depth;
+			}
+		}
+
+		return {};
+	}
+}
+
 void FStateTreeEditorModule::RegisterEditorDataClass(TNonNullSubclassOf<const UStateTreeSchema> Schema, TNonNullSubclassOf<const UStateTreeEditorData> EditorData)
 {
-	const bool bContains = EditorDataTypes.ContainsByPredicate(
-		[Schema](const FEditorDataType& Other)
+	FEditorTypes* EditorType = EditorTypes.FindByPredicate(
+		[Schema](const FEditorTypes& Other)
 		{
 			return Other.Schema == Schema.Get();
 		});
-	if (ensureMsgf(!bContains, TEXT("The type %s is already registered."), *Schema.Get()->GetName()))
+	if (EditorType)
 	{
-		EditorDataTypes.Emplace(TWeakObjectPtr<const UClass>(Schema.Get()), TWeakObjectPtr<const UClass>(EditorData.Get()));
+		ensureMsgf(EditorType->EditorData.IsExplicitlyNull(), TEXT("The type %s is already registered."), *Schema.Get()->GetName());
+		EditorType->EditorData = EditorData.Get();
+	}
+	else
+	{
+		FEditorTypes& NewEditorType = EditorTypes.AddDefaulted_GetRef();
+		NewEditorType.Schema = Schema.Get();
+		NewEditorType.EditorData = EditorData.Get();
 	}
 }
 
 void FStateTreeEditorModule::UnregisterEditorDataClass(TNonNullSubclassOf<const UStateTreeSchema> Schema)
 {
-	EditorDataTypes.RemoveAllSwap(
-		[Schema](const FEditorDataType& Other)
+	const int32 FoundIndex = EditorTypes.IndexOfByPredicate(
+		[Schema](const FEditorTypes& Other)
 		{
 			return Other.Schema == Schema.Get();
 		});
+	if (FoundIndex != INDEX_NONE)
+	{
+		EditorTypes[FoundIndex].EditorData.Reset();
+		if (!EditorTypes[FoundIndex].HasData())
+		{
+			EditorTypes.RemoveAtSwap(FoundIndex);
+		}
+	}
 }
 
 TNonNullSubclassOf<UStateTreeEditorData> FStateTreeEditorModule::GetEditorDataClass(TNonNullSubclassOf<const UStateTreeSchema> Schema) const
 {
-	const FEditorDataType* FoundEditorData = EditorDataTypes.FindByPredicate(
-		[Schema](const FEditorDataType& Other)
+	int32 BestDepth = INT_MAX;
+	const FEditorTypes* BestEditorType = nullptr;
+	for (const FEditorTypes& EditorType : EditorTypes)
+	{
+		const UClass* OtherSchema = EditorType.Schema.Get();
+		TOptional<int32> Depth = UE::StateTreeEditor::Private::GetDepth(OtherSchema, Schema.Get());
+		if (Depth.IsSet() && Depth.GetValue() < BestDepth)
 		{
-			const UClass* OtherSchema = Other.Schema.Get();
-			return OtherSchema && Schema.Get()->IsChildOf(OtherSchema);
-		});
-	const UClass* Result = FoundEditorData && FoundEditorData->EditorData.Get() ? FoundEditorData->EditorData.Get() : UStateTreeEditorData::StaticClass();
+			BestDepth = Depth.GetValue();
+			BestEditorType = &EditorType;
+		}
+	}
+
+	const UClass* Result = BestEditorType && BestEditorType->EditorData.Get() ? BestEditorType->EditorData.Get() : UStateTreeEditorData::StaticClass();
 	return const_cast<UClass*>(Result); // NewObject wants none const UClass
 }
 
+void FStateTreeEditorModule::RegisterEditorSchemaClass(TNonNullSubclassOf<const UStateTreeSchema> Schema, TNonNullSubclassOf<const UStateTreeEditorSchema> EditorSchema)
+{
+	FEditorTypes* EditorType = EditorTypes.FindByPredicate(
+		[Schema](const FEditorTypes& Other)
+		{
+			return Other.Schema == Schema.Get();
+		});
+	if (EditorType)
+	{
+		ensureMsgf(EditorType->EditorSchema.IsExplicitlyNull(), TEXT("The type %s is already registered."), *Schema.Get()->GetName());
+		EditorType->EditorSchema = EditorSchema.Get();
+	}
+	else
+	{
+		FEditorTypes& NewEditorType = EditorTypes.AddDefaulted_GetRef();
+		NewEditorType.Schema = Schema.Get();
+		NewEditorType.EditorSchema = EditorSchema.Get();
+	}
+}
+
+void FStateTreeEditorModule::UnregisterEditorSchemaClass(TNonNullSubclassOf<const UStateTreeSchema> Schema)
+{
+	const int32 FoundIndex = EditorTypes.IndexOfByPredicate(
+		[Schema](const FEditorTypes& Other)
+		{
+			return Other.Schema == Schema.Get();
+		});
+	if (FoundIndex != INDEX_NONE)
+	{
+		EditorTypes[FoundIndex].EditorSchema.Reset();
+		if (!EditorTypes[FoundIndex].HasData())
+		{
+			EditorTypes.RemoveAtSwap(FoundIndex);
+		}
+	}
+}
+
+TNonNullSubclassOf<UStateTreeEditorSchema> FStateTreeEditorModule::GetEditorSchemaClass(TNonNullSubclassOf<const UStateTreeSchema> Schema) const
+{
+	int32 BestDepth = INT_MAX;
+	const FEditorTypes* BestEditorType = nullptr;
+	for (const FEditorTypes& OtherEditorType : EditorTypes)
+	{
+		const UClass* OtherSchema = OtherEditorType.Schema.Get();
+		TOptional<int32> Depth = UE::StateTreeEditor::Private::GetDepth(OtherSchema, Schema.Get());
+		if (Depth.IsSet() && Depth.GetValue() < BestDepth)
+		{
+			BestDepth = Depth.GetValue();
+			BestEditorType = &OtherEditorType;
+		}
+	}
+
+	const UClass* Result = BestEditorType && BestEditorType->EditorSchema.Get() ? BestEditorType->EditorSchema.Get() : UStateTreeEditorSchema::StaticClass();
+	return const_cast<UClass*>(Result); // NewObject wants none const UClass
+}
 
 #undef LOCTEXT_NAMESPACE
diff --git a/Engine/Plugins/Runtime/StateTree/Source/StateTreeEditorModule/Private/StateTreeFactory.cpp b/Engine/Plugins/Runtime/StateTree/Source/StateTreeEditorModule/Private/StateTreeFactory.cpp
--- a/Engine/Plugins/Runtime/StateTree/Source/StateTreeEditorModule/Private/StateTreeFactory.cpp
+++ b/Engine/Plugins/Runtime/StateTree/Source/StateTreeEditorModule/Private/StateTreeFactory.cpp
@@ -5,13 +5,14 @@
 #include "Kismet2/SClassPickerDialog.h"
 #include "ClassViewerFilter.h"
 #include "StateTreeCompilerLog.h"
-#include "StateTreeEditorModule.h"
-#include "StateTreeNodeClassCache.h"
-#include "StateTreeEditorData.h"
-#include "StateTreeEditor.h"
 #include "StateTreeCompiler.h"
 #include "StateTreeDelegates.h"
+#include "StateTreeEditor.h"
+#include "StateTreeEditorData.h"
+#include "StateTreeEditorModule.h"
+#include "StateTreeEditorSchema.h"
 #include "StateTreeEditingSubsystem.h"
+#include "StateTreeNodeClassCache.h"
 #include "Modules/ModuleManager.h"
 
 #include UE_INLINE_GENERATED_CPP_BY_NAME(StateTreeFactory)
@@ -122,14 +123,15 @@ UObject* UStateTreeFactory::FactoryCreateNew(UClass* Class, UObject* InParent, F
 	// Create and init new editor data.
 	TNonNullSubclassOf<UStateTreeEditorData> EditorDataClass = FStateTreeEditorModule::GetModule().GetEditorDataClass(StateTreeSchemaClass.Get());
 	UStateTreeEditorData* EditorData = NewObject<UStateTreeEditorData>(NewStateTree, EditorDataClass, FName(), RF_Transactional);
-	EditorData->AddRootState();
-	EditorData->Schema = NewObject<UStateTreeSchema>(EditorData, StateTreeSchemaClass);
-	
 	NewStateTree->EditorData = EditorData;
 
-	FStateTreeCompilerLog Log;
-	FStateTreeCompiler Compiler(Log);
+	EditorData->Schema = NewObject<UStateTreeSchema>(EditorData, StateTreeSchemaClass, FName(), RF_Transactional);
+	EditorData->AddRootState();
+
+	TNonNullSubclassOf<UStateTreeEditorSchema> EditorSchemaClass = FStateTreeEditorModule::GetModule().GetEditorSchemaClass(StateTreeSchemaClass.Get());
+	EditorData->EditorSchema = NewObject<UStateTreeEditorSchema>(EditorData, EditorSchemaClass, FName(), RF_Transactional);
 
+	FStateTreeCompilerLog Log;
 	const bool bSuccess = UStateTreeEditingSubsystem::CompileStateTree(NewStateTree, Log);
 
 	if (!bSuccess)
diff --git a/Engine/Plugins/Runtime/StateTree/Source/StateTreeEditorModule/Public/StateTreeEditorData.h b/Engine/Plugins/Runtime/StateTree/Source/StateTreeEditorModule/Public/StateTreeEditorData.h
--- a/Engine/Plugins/Runtime/StateTree/Source/StateTreeEditorModule/Public/StateTreeEditorData.h
+++ b/Engine/Plugins/Runtime/StateTree/Source/StateTreeEditorModule/Public/StateTreeEditorData.h
@@ -13,6 +13,7 @@
 struct FStateTreeBindableStructDesc;
 
 class UStateTreeEditorDataExtension;
+class UStateTreeEditorSchema;
 class UStateTreeSchema;
 
 namespace UE::StateTree::Editor
@@ -357,8 +358,12 @@ public:
 	UPROPERTY(EditDefaultsOnly, Instanced, Category = "Common", NoClear)
 	TObjectPtr<UStateTreeSchema> Schema = nullptr;
 
+	/** Schema describing how the editor schema is customized. */
+	UPROPERTY(Instanced, NoClear)
+	TObjectPtr<UStateTreeEditorSchema> EditorSchema = nullptr;
+
 	/** The editor data extensions. A place to add extra information for plugins. */
-	UPROPERTY(EditDefaultsOnly, Category = "Extension", NoClear)
+	UPROPERTY(EditDefaultsOnly, Instanced, Category = "Extension", NoClear)
 	TArray<TObjectPtr<UStateTreeEditorDataExtension>> Extensions;
 
 	/** Public parameters that could be used for bindings within the Tree. */
diff --git a/Engine/Plugins/Runtime/StateTree/Source/StateTreeEditorModule/Public/StateTreeEditorModule.h b/Engine/Plugins/Runtime/StateTree/Source/StateTreeEditorModule/Public/StateTreeEditorModule.h
--- a/Engine/Plugins/Runtime/StateTree/Source/StateTreeEditorModule/Public/StateTreeEditorModule.h
+++ b/Engine/Plugins/Runtime/StateTree/Source/StateTreeEditorModule/Public/StateTreeEditorModule.h
@@ -12,6 +12,7 @@
 
 class UStateTree;
 class UStateTreeEditorData;
+class UStateTreeEditorSchema;
 class UStateTreeSchema;
 class UUserDefinedStruct;
 class IStateTreeEditor;
@@ -64,17 +65,26 @@ public:
 	/** Get the editor data type for a specific schema. */
 	UE_API TNonNullSubclassOf<UStateTreeEditorData> GetEditorDataClass(TNonNullSubclassOf<const UStateTreeSchema> Schema) const;
 
+	/** Register the editor schema type for a specific schema. */
+	void RegisterEditorSchemaClass(TNonNullSubclassOf<const UStateTreeSchema> Schema, TNonNullSubclassOf<const UStateTreeEditorSchema> EditorSchema);
+	/** Unregister the editor schema type for a specific schema. */
+	void UnregisterEditorSchemaClass(TNonNullSubclassOf<const UStateTreeSchema> Schema);
+	/** Get the editor data type for a specific schema. */
+	TNonNullSubclassOf<UStateTreeEditorSchema> GetEditorSchemaClass(TNonNullSubclassOf<const UStateTreeSchema> Schema) const;
+
 protected:
 	TSharedPtr<FExtensibilityManager> MenuExtensibilityManager;
 	TSharedPtr<FExtensibilityManager> ToolBarExtensibilityManager;
 	TSharedPtr<FStateTreeNodeClassCache> NodeClassCache;
 
-	struct FEditorDataType
+	struct FEditorTypes
 	{
 		TWeakObjectPtr<const UClass> Schema;
 		TWeakObjectPtr<const UClass> EditorData;
+		TWeakObjectPtr<const UClass> EditorSchema;
+		bool HasData() const;
 	};
-	TArray<FEditorDataType> EditorDataTypes;
+	TArray<FEditorTypes> EditorTypes;
 
 #if WITH_STATETREE_TRACE_DEBUGGER
 	TPimplPtr<UE::StateTreeDebugger::FRewindDebuggerPlaybackExtension> RewindDebuggerPlaybackExtension;
