[StateTree] Fix: revert wrong state selection behavior for groups introduced at CL 43606513.

diff --git a/Engine/Plugins/Runtime/StateTree/Source/StateTreeEditorModule/Private/StateTreeCompiler.cpp b/Engine/Plugins/Runtime/StateTree/Source/StateTreeEditorModule/Private/StateTreeCompiler.cpp
--- a/Engine/Plugins/Runtime/StateTree/Source/StateTreeEditorModule/Private/StateTreeCompiler.cpp
+++ b/Engine/Plugins/Runtime/StateTree/Source/StateTreeEditorModule/Private/StateTreeCompiler.cpp
@@ -649,7 +649,7 @@ bool FStateTreeCompiler::CreateStateRecursive(UStateTreeState& State, const FSta
 		return false;
 	}
 
-	const bool bHasPredefinedSelectionBehavior = (CompactState.Type == EStateTreeStateType::Group || CompactState.Type == EStateTreeStateType::Linked || CompactState.Type == EStateTreeStateType::LinkedAsset);
+	const bool bHasPredefinedSelectionBehavior = (CompactState.Type == EStateTreeStateType::Linked || CompactState.Type == EStateTreeStateType::LinkedAsset);
 	if (bHasPredefinedSelectionBehavior)
 	{
 		CompactState.SelectionBehavior = EStateTreeStateSelectionBehavior::TryEnterState;
diff --git a/Engine/Plugins/Runtime/StateTree/Source/StateTreeEditorModule/Private/StateTreeState.cpp b/Engine/Plugins/Runtime/StateTree/Source/StateTreeEditorModule/Private/StateTreeState.cpp
--- a/Engine/Plugins/Runtime/StateTree/Source/StateTreeEditorModule/Private/StateTreeState.cpp
+++ b/Engine/Plugins/Runtime/StateTree/Source/StateTreeEditorModule/Private/StateTreeState.cpp
@@ -157,7 +157,13 @@ void UStateTreeState::PostEditChangeChainProperty(FPropertyChangedChainEvent& Pr
 	
 	if (ChangePropertyPath.IsPathExact(StateTypePath))
 	{
-		const bool bHasPredefinedSelectionBehavior = Type == EStateTreeStateType::Group || Type == EStateTreeStateType::Linked || Type == EStateTreeStateType::LinkedAsset;
+		if (Type == EStateTreeStateType::Group)
+		{
+			// Group should not have tasks.
+			Tasks.Reset();
+		}
+
+		const bool bHasPredefinedSelectionBehavior = Type == EStateTreeStateType::Linked || Type == EStateTreeStateType::LinkedAsset;
 		if (bHasPredefinedSelectionBehavior)
 		{
 			// Reset Selection Behavior back to Try Enter State for group and linked types
