[StateTree] Refactor StateTreeExecutionContext::CopyBatch and EvaluatePropertyFunctions to the same function. It helps with the ExecutionScopeInstanceData new feature.

diff --git a/Engine/Plugins/Runtime/StateTree/Source/StateTreeModule/Private/StateTreeExecutionContext.cpp b/Engine/Plugins/Runtime/StateTree/Source/StateTreeModule/Private/StateTreeExecutionContext.cpp
--- a/Engine/Plugins/Runtime/StateTree/Source/StateTreeModule/Private/StateTreeExecutionContext.cpp
+++ b/Engine/Plugins/Runtime/StateTree/Source/StateTreeModule/Private/StateTreeExecutionContext.cpp
@@ -2950,25 +2950,18 @@ void FStateTreeExecutionContext::CopyAllBindingsOnActiveInstances(const ECopyBin
 
 bool FStateTreeExecutionContext::CopyBatchOnActiveInstances(const FStateTreeExecutionFrame* ParentFrame, const FStateTreeExecutionFrame& CurrentFrame, const FStateTreeDataView TargetView, const FStateTreeIndex16 BindingsBatch)
 {
-	const FPropertyBindingCopyInfoBatch& Batch = CurrentFrame.StateTree->PropertyBindings.Super::GetBatch(BindingsBatch);
-	check(TargetView.GetStruct() == Batch.TargetStruct.Get().Struct);
-
-	if (Batch.PropertyFunctionsBegin != Batch.PropertyFunctionsEnd)
-	{
-		check(Batch.PropertyFunctionsBegin.IsValid() && Batch.PropertyFunctionsEnd.IsValid());
-		EvaluatePropertyFunctionsOnActiveInstances(ParentFrame, CurrentFrame, FStateTreeIndex16(Batch.PropertyFunctionsBegin), Batch.PropertyFunctionsEnd.Get() - Batch.PropertyFunctionsBegin.Get());
-	}
-
-	bool bSucceed = true;
-	for (const FPropertyBindingCopyInfo& Copy : CurrentFrame.StateTree->PropertyBindings.Super::GetBatchCopies(Batch))
-	{
-		const FStateTreeDataView SourceView = GetDataView(ParentFrame, CurrentFrame, Copy);
-		bSucceed &= CurrentFrame.StateTree->PropertyBindings.Super::CopyProperty(Copy, SourceView, TargetView);
-	}
-	return bSucceed;
+	constexpr bool bOnActiveInstances = true;
+	return CopyBatchInternal<bOnActiveInstances>(ParentFrame, CurrentFrame, TargetView, BindingsBatch);
 }
 
 bool FStateTreeExecutionContext::CopyBatchWithValidation(const FStateTreeExecutionFrame* ParentFrame, const FStateTreeExecutionFrame& CurrentFrame, const FStateTreeDataView TargetView, const FStateTreeIndex16 BindingsBatch)
+{
+	constexpr bool bOnActiveInstances = false;
+	return CopyBatchInternal<bOnActiveInstances>(ParentFrame, CurrentFrame, TargetView, BindingsBatch);
+}
+
+template<bool bOnActiveInstances>
+bool FStateTreeExecutionContext::CopyBatchInternal(const FStateTreeExecutionFrame* ParentFrame, const FStateTreeExecutionFrame& CurrentFrame, const FStateTreeDataView TargetView, const FStateTreeIndex16 BindingsBatch)
 {
 	const FPropertyBindingCopyInfoBatch& Batch = CurrentFrame.StateTree->PropertyBindings.Super::GetBatch(BindingsBatch);
 	check(TargetView.GetStruct() == Batch.TargetStruct.Get().Struct);
@@ -2976,20 +2969,34 @@ bool FStateTreeExecutionContext::CopyBatchWithValidation(const FStateTreeExecuti
 	if (Batch.PropertyFunctionsBegin != Batch.PropertyFunctionsEnd)
 	{
 		check(Batch.PropertyFunctionsBegin.IsValid() && Batch.PropertyFunctionsEnd.IsValid());
-		EvaluatePropertyFunctionsWithValidation(ParentFrame, CurrentFrame, FStateTreeIndex16(Batch.PropertyFunctionsBegin), Batch.PropertyFunctionsEnd.Get() - Batch.PropertyFunctionsBegin.Get());
+		if constexpr (bOnActiveInstances)
+		{
+			EvaluatePropertyFunctionsOnActiveInstances(ParentFrame, CurrentFrame, FStateTreeIndex16(Batch.PropertyFunctionsBegin), Batch.PropertyFunctionsEnd.Get() - Batch.PropertyFunctionsBegin.Get());
+		}
+		else
+		{
+			EvaluatePropertyFunctionsWithValidation(ParentFrame, CurrentFrame, FStateTreeIndex16(Batch.PropertyFunctionsBegin), Batch.PropertyFunctionsEnd.Get() - Batch.PropertyFunctionsBegin.Get());
+		}
 	}
 
 	bool bSucceed = true;
 	for (const FPropertyBindingCopyInfo& Copy : CurrentFrame.StateTree->PropertyBindings.Super::GetBatchCopies(Batch))
 	{
-		const FStateTreeDataView SourceView = GetDataViewOrTemporary(ParentFrame, CurrentFrame, Copy);
-		if (!SourceView.IsValid())
+		if constexpr (bOnActiveInstances)
 		{
-			bSucceed = false;
-			break;
+			const FStateTreeDataView SourceView = GetDataView(ParentFrame, CurrentFrame, Copy);
+			bSucceed &= CurrentFrame.StateTree->PropertyBindings.Super::CopyProperty(Copy, SourceView, TargetView);
+		}
+		else
+		{
+			const FStateTreeDataView SourceView = GetDataViewOrTemporary(ParentFrame, CurrentFrame, Copy);
+			if (!SourceView.IsValid())
+			{
+				bSucceed = false;
+				break;
+			}
+			bSucceed &= CurrentFrame.StateTree->PropertyBindings.Super::CopyProperty(Copy, SourceView, TargetView);
 		}
-
-		bSucceed &= CurrentFrame.StateTree->PropertyBindings.Super::CopyProperty(Copy, SourceView, TargetView);
 	}
 	return bSucceed;
 }
@@ -4596,10 +4603,10 @@ bool FStateTreeExecutionContext::TestAllConditionsInternal(const FStateTreeExecu
 
 	int32 Level = 0;
 
-	for (int32 Index = 0; Index < ConditionsNum; Index++)
-	{
-		const int32 ConditionIndex = ConditionsOffset + Index;
-		const FStateTreeConditionBase& Cond = CurrentFrame.StateTree->Nodes[ConditionIndex].Get<const FStateTreeConditionBase>();
+		for (int32 Index = 0; Index < ConditionsNum; Index++)
+		{
+			const int32 ConditionIndex = ConditionsOffset + Index;
+			const FStateTreeConditionBase& Cond = CurrentFrame.StateTree->Nodes[ConditionIndex].Get<const FStateTreeConditionBase>();
 		checkf(Cond.InstanceDataHandle.GetSource() == EStateTreeDataSourceType::SharedInstanceData
 			|| Cond.InstanceDataHandle.GetSource() == EStateTreeDataSourceType::SharedInstanceDataObject,
 			TEXT("Condition only support SharedInstanceData"));
@@ -4780,32 +4787,21 @@ float FStateTreeExecutionContext::EvaluateUtilityWithValidation(const FStateTree
 
 void FStateTreeExecutionContext::EvaluatePropertyFunctionsOnActiveInstances(const FStateTreeExecutionFrame* CurrentParentFrame, const FStateTreeExecutionFrame& CurrentFrame, FStateTreeIndex16 FuncsBegin, uint16 FuncsNum)
 {
-	for (int32 FuncIndex = FuncsBegin.Get(); FuncIndex < FuncsBegin.Get() + FuncsNum; ++FuncIndex)
-	{
-		const FStateTreePropertyFunctionBase& Func = CurrentFrame.StateTree->Nodes[FuncIndex].Get<const FStateTreePropertyFunctionBase>();
-		const FStateTreeDataView FuncInstanceView = GetDataView(CurrentParentFrame, CurrentFrame, Func.InstanceDataHandle);
-		FNodeInstanceDataScope DataScope(*this, &Func, FuncIndex, Func.InstanceDataHandle, FuncInstanceView);
-
-		// Copy bound properties.
-		if (Func.BindingsBatch.IsValid())
-		{
-			// Use validated copy, since we test in situations where the sources are not always valid (e.g. enter conditions may try to access inactive parent state). 
-			CopyBatchOnActiveInstances(CurrentParentFrame, CurrentFrame, FuncInstanceView, Func.BindingsBatch);
-		}
-
-		Func.Execute(*this);
-
-		// Reset copied properties that might contain object references.
-		if (Func.BindingsBatch.IsValid())
-		{
-			CurrentFrame.StateTree->PropertyBindings.Super::ResetObjects(Func.BindingsBatch, FuncInstanceView);
-		}
-	}
+	constexpr bool bOnActiveInstances = true;
+	EvaluatePropertyFunctionsInternal<bOnActiveInstances>(CurrentParentFrame, CurrentFrame, FuncsBegin, FuncsNum);
 }
 
 void FStateTreeExecutionContext::EvaluatePropertyFunctionsWithValidation(const FStateTreeExecutionFrame* CurrentParentFrame, const FStateTreeExecutionFrame& CurrentFrame, FStateTreeIndex16 FuncsBegin, uint16 FuncsNum)
 {
-	for (int32 FuncIndex = FuncsBegin.Get(); FuncIndex < FuncsBegin.Get() + FuncsNum; ++FuncIndex)
+	constexpr bool bOnActiveInstances = false;
+	EvaluatePropertyFunctionsInternal<bOnActiveInstances>(CurrentParentFrame, CurrentFrame, FuncsBegin, FuncsNum);
+}
+
+template<bool bOnActiveInstances>
+void FStateTreeExecutionContext::EvaluatePropertyFunctionsInternal(const FStateTreeExecutionFrame* CurrentParentFrame, const FStateTreeExecutionFrame& CurrentFrame, FStateTreeIndex16 FuncsBegin, uint16 FuncsNum)
+{
+	const int32 FuncsEnd = FuncsBegin.AsInt32() + FuncsNum;
+	for (int32 FuncIndex = FuncsBegin.AsInt32(); FuncIndex < FuncsEnd; ++FuncIndex)
 	{
 		const FStateTreePropertyFunctionBase& Func = CurrentFrame.StateTree->Nodes[FuncIndex].Get<const FStateTreePropertyFunctionBase>();
 		const FStateTreeDataView FuncInstanceView = GetDataView(CurrentParentFrame, CurrentFrame, Func.InstanceDataHandle);
@@ -4814,8 +4810,14 @@ void FStateTreeExecutionContext::EvaluatePropertyFunctionsWithValidation(const F
 		// Copy bound properties.
 		if (Func.BindingsBatch.IsValid())
 		{
-			// Use validated copy, since we test in situations where the sources are not always valid (e.g. enter conditions may try to access inactive parent state). 
-			CopyBatchWithValidation(CurrentParentFrame, CurrentFrame, FuncInstanceView, Func.BindingsBatch);
+			if constexpr (bOnActiveInstances)
+			{
+				CopyBatchOnActiveInstances(CurrentParentFrame, CurrentFrame, FuncInstanceView, Func.BindingsBatch);
+			}
+			else
+			{
+				CopyBatchWithValidation(CurrentParentFrame, CurrentFrame, FuncInstanceView, Func.BindingsBatch);
+			}
 		}
 
 		Func.Execute(*this);
@@ -7610,4 +7612,3 @@ FString FStateTreeExecutionContext::GetInstanceDescription() const
 
 #undef STATETREE_LOG
 #undef STATETREE_CLOG
-#undef UE_STATETREE_ENSURE_COMPLETED_STATE_RUN_STATUS
diff --git a/Engine/Plugins/Runtime/StateTree/Source/StateTreeModule/Public/StateTreeExecutionContext.h b/Engine/Plugins/Runtime/StateTree/Source/StateTreeModule/Public/StateTreeExecutionContext.h
--- a/Engine/Plugins/Runtime/StateTree/Source/StateTreeModule/Public/StateTreeExecutionContext.h
+++ b/Engine/Plugins/Runtime/StateTree/Source/StateTreeModule/Public/StateTreeExecutionContext.h
@@ -870,8 +870,8 @@ PRAGMA_ENABLE_DEPRECATION_WARNINGS
 		FStateTreeTransitionSource Source;
 	};
 
-	/** @return Prefix that will be used by STATETREE_LOG and STATETREE_CLOG, Owner name by default. */
 	UE_DEPRECATED(5.6, "Use FStateTreeExecutionExtension::GetInstanceDescription instead")
+	/** @return Prefix that will be used by STATETREE_LOG and STATETREE_CLOG, Owner name by default. */
 	UE_API virtual FString GetInstanceDescription() const final;
 
 	/** Callback when delayed transition is triggered. Contexts that are event based can use this to trigger a future event. */
@@ -1078,6 +1078,11 @@ protected:
 	UE_DEPRECATED(5.7, "Use EvaluateUtilityWithValidation. It can get the data from a temporary frame in state selection.")
 	UE_API float EvaluateUtility(const FStateTreeExecutionFrame* CurrentParentFrame, const FStateTreeExecutionFrame& CurrentFrame, const int32 ConsiderationsOffset, const int32 ConsiderationsNum, const float StateWeight);
 
+private:
+	template<bool bOnActiveInstances>
+	void EvaluatePropertyFunctionsInternal(const FStateTreeExecutionFrame* CurrentParentFrame, const FStateTreeExecutionFrame& CurrentFrame, FStateTreeIndex16 FuncsBegin, uint16 FuncsNum);
+
+protected:
 	/*
 	 * Evaluate all function at given range.
 	 * Should be used only on active node instances, assumes valid data handles and does not consider temporary node instances.
@@ -1334,7 +1339,7 @@ PRAGMA_ENABLE_DEPRECATION_WARNINGS
 	/** @return true if handle source is valid cified handle relative to given frame. */
 	[[nodiscard]] UE_API bool IsHandleSourceValid(const FStateTreeExecutionFrame* ParentFrame, const FStateTreeExecutionFrame& CurrentFrame, const FStateTreeDataHandle Handle) const;
 	[[nodiscard]] UE_API bool IsHandleSourceValid(const FStateTreeExecutionFrame* ParentFrame, const FStateTreeExecutionFrame& CurrentFrame, const FPropertyBindingCopyInfo& CopyInfo) const;
-
+	
 	/** @return data view of the specified handle relative to the given frame, or tries to find a matching temporary instance. */
 	[[nodiscard]] UE_API FStateTreeDataView GetDataViewOrTemporary(const FStateTreeExecutionFrame* ParentFrame, const FStateTreeExecutionFrame& CurrentFrame, const FStateTreeDataHandle Handle);
 	[[nodiscard]] UE_API FStateTreeDataView GetDataViewOrTemporary(const FStateTreeExecutionFrame* ParentFrame, const FStateTreeExecutionFrame& CurrentFrame, const FPropertyBindingCopyInfo& CopyInfo);
@@ -1358,6 +1363,11 @@ PRAGMA_ENABLE_DEPRECATION_WARNINGS
 	/** Copy the binding on all valid active instances. */
 	UE_API void CopyAllBindingsOnActiveInstances(ECopyBindings CopyType);
 
+private:
+	template<bool bOnActiveInstances>
+	bool CopyBatchInternal(const FStateTreeExecutionFrame* ParentFrame, const FStateTreeExecutionFrame& CurrentFrame, const FStateTreeDataView TargetView, const FStateTreeIndex16 BindingsBatch);
+
+protected:
 	/**
 	 * Copies a batch of properties to the data in TargetView.
 	 * Should be used only on active node instances, assumes valid data handles and does not consider temporary node instances.
@@ -1464,7 +1474,7 @@ PRAGMA_ENABLE_DEPRECATION_WARNINGS
 	const FStateTreeExecutionFrame* CurrentlyProcessedParentFrame = nullptr; 
 	const FStateTreeExecutionFrame* CurrentlyProcessedFrame = nullptr; 
 
-	/** Pointer to the shared instance data of the current frame we're processing.  */
+	/** Pointer to the shared instance data of the current frame we're processing. */
 	FStateTreeInstanceStorage* CurrentlyProcessedSharedInstanceStorage = nullptr;
 
 	/** Helper struct to track currently processed frame. */
