[StateTree] Notify the schema on post compilation. Add option to display extensions in the editor.

diff --git a/Engine/Plugins/Runtime/StateTree/Source/StateTreeEditorModule/Private/Customizations/StateTreeEditorDataDetails.cpp b/Engine/Plugins/Runtime/StateTree/Source/StateTreeEditorModule/Private/Customizations/StateTreeEditorDataDetails.cpp
--- a/Engine/Plugins/Runtime/StateTree/Source/StateTreeEditorModule/Private/Customizations/StateTreeEditorDataDetails.cpp
+++ b/Engine/Plugins/Runtime/StateTree/Source/StateTreeEditorModule/Private/Customizations/StateTreeEditorDataDetails.cpp
@@ -10,6 +10,7 @@
 #include "PropertyCustomizationHelpers.h"
 #include "StateTreeEditorData.h"
 #include "StateTreeEditorNodeUtils.h"
+#include "StateTreeEditorSchema.h"
 #include "StateTreeEditorStyle.h"
 #include "StateTreeSchema.h"
 #include "Styling/StyleColors.h"
@@ -42,6 +43,7 @@ void FStateTreeEditorDataDetails::CustomizeDetails(IDetailLayoutBuilder& DetailB
 		}
 	}
 	const UStateTreeSchema* Schema = EditorData ? EditorData->Schema.Get() : nullptr;
+	const UStateTreeEditorSchema* EditorSchema = EditorData ? EditorData->EditorSchema.Get() : nullptr;
 
 	// Common category
 	IDetailCategoryBuilder& CommonCategory = DetailBuilder.EditCategory(TEXT("Common"), LOCTEXT("EditorDataDetailsCommon", "Common"));
@@ -236,7 +238,15 @@ void FStateTreeEditorDataDetails::CustomizeDetails(IDetailLayoutBuilder& DetailB
 		/*SortOrder*/4);
 	UE::StateTreeEditor::EditorNodeUtils::MakeArrayItems(GlobalTasksCategory, GlobalTasksProperty);
 
-	// Refresh the UI when the Schema changes.	
+	// Extensions
+	if (EditorSchema && !EditorSchema->AllowExtensions())
+	{
+		TSharedPtr<IPropertyHandle> ExtensionsProperty = DetailBuilder.GetProperty(GET_MEMBER_NAME_CHECKED(UStateTreeEditorData, Extensions));
+		check(ExtensionsProperty.IsValid());
+		DetailBuilder.EditCategory(ExtensionsProperty->GetDefaultCategoryName()).SetCategoryVisibility(false);
+	}
+
+	// Refresh the UI when the Schema changes.
 	TSharedPtr<IPropertyHandle> SchemaProperty = DetailBuilder.GetProperty(GET_MEMBER_NAME_CHECKED(UStateTreeEditorData, Schema));
 	check(SchemaProperty.IsValid());
 	SchemaProperty->SetOnPropertyValueChanged(FSimpleDelegate::CreateLambda([WeakPropertyUtilities = WeakPropertyUtilities] ()
diff --git a/Engine/Plugins/Runtime/StateTree/Source/StateTreeEditorModule/Private/StateTreeCompiler.cpp b/Engine/Plugins/Runtime/StateTree/Source/StateTreeEditorModule/Private/StateTreeCompiler.cpp
--- a/Engine/Plugins/Runtime/StateTree/Source/StateTreeEditorModule/Private/StateTreeCompiler.cpp
+++ b/Engine/Plugins/Runtime/StateTree/Source/StateTreeEditorModule/Private/StateTreeCompiler.cpp
@@ -6,6 +6,7 @@
 #include "StateTreeCompilerLog.h"
 #include "StateTreeEditorData.h"
 #include "StateTreeEditorDataExtension.h"
+#include "StateTreeEditorSchema.h"
 #include "StateTreeExtension.h"
 #include "StateTreeEditorModule.h"
 #include "StateTreeEvaluatorBase.h"
@@ -2696,10 +2697,19 @@ bool FStateTreeCompiler::NotifyInternalPost()
 		FStateTreeCompilerLog& Log;
 	};
 
-	if (EditorData->Extensions.Num())
+	if (EditorData->EditorSchema || EditorData->Extensions.Num() > 0)
 	{
 		FPostInternalContextImpl Context = FPostInternalContextImpl(StateTree, EditorData, Log);
 
+		if (EditorData->EditorSchema)
+		{
+			if (!EditorData->EditorSchema->HandlePostInternalCompile(Context))
+			{
+				Log.Reportf(EMessageSeverity::Error, TEXT("The schema failed compilation. See log for more info."));
+				return false;
+			}
+		}
+
 		for (UStateTreeEditorDataExtension* Extension : EditorData->Extensions)
 		{
 			if (Extension)
diff --git a/Engine/Plugins/Runtime/StateTree/Source/StateTreeEditorModule/Private/StateTreeEditorData.cpp b/Engine/Plugins/Runtime/StateTree/Source/StateTreeEditorModule/Private/StateTreeEditorData.cpp
--- a/Engine/Plugins/Runtime/StateTree/Source/StateTreeEditorModule/Private/StateTreeEditorData.cpp
+++ b/Engine/Plugins/Runtime/StateTree/Source/StateTreeEditorModule/Private/StateTreeEditorData.cpp
@@ -209,7 +209,11 @@ void UStateTreeEditorData::PostEditChangeChainProperty(FPropertyChangedChainEven
 		const FName MemberName = MemberProperty->GetFName();
 		if (MemberName == GET_MEMBER_NAME_CHECKED(UStateTreeEditorData, Schema)
 			|| MemberName == GET_MEMBER_NAME_CHECKED(UStateTreeEditorData, EditorSchema))
-		{			
+		{
+			if (EditorSchema && !EditorSchema->AllowExtensions())
+			{
+				Extensions.Reset();
+			}
 			UE::StateTree::Delegates::OnSchemaChanged.Broadcast(*StateTree);
 		}
 		else if (MemberName == GET_MEMBER_NAME_CHECKED(UStateTreeEditorData, RootParameterPropertyBag))
diff --git a/Engine/Plugins/Runtime/StateTree/Source/StateTreeEditorModule/Public/StateTreeEditorSchema.h b/Engine/Plugins/Runtime/StateTree/Source/StateTreeEditorModule/Public/StateTreeEditorSchema.h
--- a/Engine/Plugins/Runtime/StateTree/Source/StateTreeEditorModule/Public/StateTreeEditorSchema.h
+++ b/Engine/Plugins/Runtime/StateTree/Source/StateTreeEditorModule/Public/StateTreeEditorSchema.h
@@ -8,6 +8,10 @@
 #define UE_API STATETREEEDITORMODULE_API
 
 class UStateTree;
+namespace UE::StateTree::Compiler
+{
+	struct FPostInternalContext;
+}
 
 /**
  * Schema describing how a StateTree is edited.
@@ -16,13 +20,30 @@ UCLASS(MinimalAPI)
 class UStateTreeEditorSchema : public UObject
 {
 	GENERATED_BODY()
-	
+
+public:
+	/** @returns True if modifying extensions in the editor is allowed. An extension can be added by code. */
+	virtual bool AllowExtensions() const
+	{
+		return true;
+	}
+
 public:
 	/*
 	 * Validates and applies the schema restrictions on the StateTree.
 	 * Also serves as the "safety net" of fixing up editor data following an editor operation.
 	 */
 	UE_API virtual void Validate(TNotNull<UStateTree*> StateTree);
+
+
+	/**
+	 * Handle compilation for the owning state tree asset.
+	 * The state tree asset compiled successfully.
+	 */
+	virtual bool HandlePostInternalCompile(const UE::StateTree::Compiler::FPostInternalContext& Context)
+	{
+		return true;
+	}
 };
 
 #undef UE_API
