[StateTree] Add EvaluationScopeInstanceData. A temporary instance data for nodes that are used on the stack. The instance data is created on the stack and destroyed right away after use. PropertyFunction is using EvaluationScopeInstanceData. Condition and Consideration will also use it in the future.
Add EvaluationScopeInstanceContainer, a container to help construct and destroy the evaluation scope instance data.
Add InstanceContainer, a wrapper around FInstancedStructContainer. Use to store the default value instead of using the FStateTreeInstanceStorage, which stores more information. In the future, the storage should be refactored to the InstanceContainer.

diff --git a/Engine/Plugins/Runtime/StateTree/Source/StateTreeEditorModule/Private/StateTreeCompiler.cpp b/Engine/Plugins/Runtime/StateTree/Source/StateTreeEditorModule/Private/StateTreeCompiler.cpp
--- a/Engine/Plugins/Runtime/StateTree/Source/StateTreeEditorModule/Private/StateTreeCompiler.cpp
+++ b/Engine/Plugins/Runtime/StateTree/Source/StateTreeEditorModule/Private/StateTreeCompiler.cpp
@@ -446,6 +446,7 @@ bool FStateTreeCompiler::Compile(UStateTree& InStateTree)
 	StateTree->Nodes = Nodes;
 	StateTree->DefaultInstanceData.Init(*StateTree, InstanceStructs, FStateTreeInstanceData::FAddArgs{ .bDuplicateWrappedObject = false });
 	StateTree->SharedInstanceData.Init(*StateTree, SharedInstanceStructs, FStateTreeInstanceData::FAddArgs{ .bDuplicateWrappedObject = false });
+	StateTree->DefaultEvaluationScopeInstanceData.Init(StateTree, EvaluationScopeStructs, UE::StateTree::InstanceData::FInstanceContainer::FAddArgs{ .bDuplicateWrappedObject = false });
 
 	// Store the new compiled dispatchers.
 	EditorData->CompiledDispatchers = BindingsCompiler.GetCompiledDelegateDispatchers();
@@ -1703,8 +1704,15 @@ bool FStateTreeCompiler::CreatePropertyFunction(const FStateTreeEditorNode& Func
 	StructDesc.Name = FuncEditorNode.GetName();
 	StructDesc.DataSource = EStateTreeBindableStructSource::PropertyFunction;
 
-	FStateTreeNodeBase* Node = CreateNodeWithSharedInstanceData(nullptr, FuncEditorNode, StructDesc);
-	if(Node == nullptr)
+	const FStateTreeDataHandle DataHandle = FStateTreeDataHandle(EStateTreeDataSourceType::EvaluationScopeInstanceData, EvaluationScopeStructs.Num());
+	FInstancedStruct* CreatedNode = CreateNode(nullptr, FuncEditorNode, StructDesc, DataHandle, EvaluationScopeStructs);
+	if (CreatedNode == nullptr)
+	{
+		return false;
+	}
+
+	FStateTreeNodeBase* Function = CreatedNode->GetMutablePtr<FStateTreeNodeBase>();
+	if (Function == nullptr)
 	{
 		return false;
 	}
@@ -1713,9 +1721,9 @@ bool FStateTreeCompiler::CreatePropertyFunction(const FStateTreeEditorNode& Func
 	check(BindableStruct);
 
 	FStateTreeDataView InstanceView;
-	check(SharedInstanceStructs.IsValidIndex(Node->InstanceTemplateIndex.Get()));
+	check(EvaluationScopeStructs.IsValidIndex(Function->InstanceTemplateIndex.Get()));
 
-	FInstancedStruct& Instance = SharedInstanceStructs[Node->InstanceTemplateIndex.Get()];
+	FInstancedStruct& Instance = EvaluationScopeStructs[Function->InstanceTemplateIndex.Get()];
 	if (FStateTreeInstanceObjectWrapper* ObjectWrapper = Instance.GetMutablePtr<FStateTreeInstanceObjectWrapper>())
 	{
 		check(FuncEditorNode.InstanceObject->GetClass() == ObjectWrapper->InstanceObject->GetClass());
@@ -1727,19 +1735,16 @@ bool FStateTreeCompiler::CreatePropertyFunction(const FStateTreeEditorNode& Func
 		InstanceView = FStateTreeDataView(Instance);
 	}
 
-	return CreateBindingsForStruct(*BindableStruct, InstanceView, FStateTreeIndex16::Invalid, FStateTreeIndex16::Invalid, Node->BindingsBatch);
+	return CreateBindingsForStruct(*BindableStruct, InstanceView, FStateTreeIndex16::Invalid, FStateTreeIndex16::Invalid, Function->BindingsBatch);
 }
 
-template<class T>
-T* FStateTreeCompiler::CreateNodeWithSharedInstanceData(UStateTreeState* State, const FStateTreeEditorNode& EditorNode, FStateTreeBindableStructDesc& StructDesc)
+FInstancedStruct* FStateTreeCompiler::CreateNode(UStateTreeState* State, const FStateTreeEditorNode& EditorNode, FStateTreeBindableStructDesc& StructDesc, const FStateTreeDataHandle DataHandle, TArray<FInstancedStruct>& InstancedStructContainer)
 {
 	if (!EditorNode.Node.IsValid())
 	{
 		return nullptr;
 	}
 
-	check(EditorNode.Node.GetScriptStruct()->IsChildOf<T>());
-
 	// Check that item has valid instance initialized.
 	if (!EditorNode.Instance.IsValid() && EditorNode.InstanceObject == nullptr)
 	{
@@ -1755,16 +1760,41 @@ T* FStateTreeCompiler::CreateNodeWithSharedInstanceData(UStateTreeState* State,
 
 	FStateTreeNodeBase& Node = RawNode.GetMutable<FStateTreeNodeBase>();
 
-	// Update node name as description for runtime.
-	Node.Name = EditorNode.GetName();
+	TOptional<FStateTreeDataView> InstanceDataView = CreateNodeInstanceData(EditorNode, Node, StructDesc, DataHandle, InstancedStructContainer);
+	if (!InstanceDataView.IsSet())
+	{
+		return nullptr;
+	}
 
+	if (!CompileAndValidateNode(State, StructDesc, RawNode, InstanceDataView.GetValue()))
+	{
+		return nullptr;
+	}
+
+	// Mark the struct as binding source.
+	BindingsCompiler.AddSourceStruct(StructDesc);
+
+	return &RawNode;
+}
+
+FInstancedStruct* FStateTreeCompiler::CreateNodeWithSharedInstanceData(UStateTreeState* State, const FStateTreeEditorNode& EditorNode, FStateTreeBindableStructDesc& StructDesc)
+{
+	const FStateTreeDataHandle DataHandle = FStateTreeDataHandle(EStateTreeDataSourceType::SharedInstanceData, SharedInstanceStructs.Num());
+	return CreateNode(State, EditorNode, StructDesc, DataHandle, SharedInstanceStructs);
+}
+
+TOptional<FStateTreeDataView> FStateTreeCompiler::CreateNodeInstanceData(const FStateTreeEditorNode& EditorNode, FStateTreeNodeBase& Node, FStateTreeBindableStructDesc& StructDesc, const FStateTreeDataHandle DataHandle, TArray<FInstancedStruct>& InstancedStructContainer)
+{
 	FStateTreeDataView InstanceDataView;
-	
+
+	// Update Node name as description for runtime.
+	Node.Name = EditorNode.GetName();
+
 	if (EditorNode.Instance.IsValid())
 	{
-		// Struct instance
-		const int32 InstanceIndex = SharedInstanceStructs.Add(EditorNode.Instance);
-		InstantiateStructSubobjects(SharedInstanceStructs[InstanceIndex]);
+		// Struct Instance
+		const int32 InstanceIndex = InstancedStructContainer.Add(EditorNode.Instance);
+		InstantiateStructSubobjects(InstancedStructContainer[InstanceIndex]);
 
 		// Create binding source struct descriptor.
 		StructDesc.Struct = EditorNode.Instance.GetScriptStruct();
@@ -1772,47 +1802,38 @@ T* FStateTreeCompiler::CreateNodeWithSharedInstanceData(UStateTreeState* State,
 		if (const auto Validation = UE::StateTree::Compiler::IsValidIndex16(InstanceIndex); Validation.DidFail())
 		{
 			Validation.Log(Log, TEXT("InstanceIndex"), StructDesc);
-			return nullptr;
+			return {};
 		}
 		Node.InstanceTemplateIndex = FStateTreeIndex16(InstanceIndex);
-		Node.InstanceDataHandle = FStateTreeDataHandle(EStateTreeDataSourceType::SharedInstanceData, InstanceIndex);
-		InstanceDataView = FStateTreeDataView(SharedInstanceStructs[InstanceIndex]);
+		Node.InstanceDataHandle = DataHandle;
+		InstanceDataView = FStateTreeDataView(InstancedStructContainer[InstanceIndex]);
 	}
 	else
 	{
 		// Object Instance
 		check(EditorNode.InstanceObject != nullptr);
-
 		UObject* Instance = UE::StateTree::Compiler::DuplicateInstanceObject(Log, StructDesc, EditorNode.ID, EditorNode.InstanceObject, StateTree);
+
 		FInstancedStruct Wrapper;
 		Wrapper.InitializeAs<FStateTreeInstanceObjectWrapper>(Instance);
-		const int32 InstanceIndex = SharedInstanceStructs.Add(MoveTemp(Wrapper));
-		
+		const int32 InstanceIndex = InstancedStructContainer.Add(MoveTemp(Wrapper));
+
 		// Create binding source struct descriptor.
 		StructDesc.Struct = Instance->GetClass();
 
 		if (const auto Validation = UE::StateTree::Compiler::IsValidIndex16(InstanceIndex); Validation.DidFail())
 		{
 			Validation.Log(Log, TEXT("InstanceIndex"), StructDesc);
-			return nullptr;
+			return {};
 		}
 		Node.InstanceTemplateIndex = FStateTreeIndex16(InstanceIndex);
-		Node.InstanceDataHandle = FStateTreeDataHandle(EStateTreeDataSourceType::SharedInstanceDataObject, InstanceIndex);
+		Node.InstanceDataHandle = DataHandle.ToObjectSource();
 		InstanceDataView = FStateTreeDataView(Instance);
 	}
 
 	StructDesc.DataHandle = Node.InstanceDataHandle;
-	
-	if (!CompileAndValidateNode(State, StructDesc, RawNode, InstanceDataView))
-	{
-		return nullptr;
-	}
 
-	// Mark the struct as binding source.
-	BindingsCompiler.AddSourceStruct(StructDesc);
-	
-	check(RawNode.GetScriptStruct()->IsChildOf<T>());
-	return RawNode.GetMutablePtr<T>();
+	return TOptional<FStateTreeDataView>(InstanceDataView);
 }
 
 bool FStateTreeCompiler::ResolveTransitionStateAndFallback(const UStateTreeState* SourceState, const FStateTreeStateLink& Link, FStateTreeStateHandle& OutTransitionHandle, EStateTreeSelectionFallback& OutFallback) const 
@@ -1905,19 +1926,24 @@ bool FStateTreeCompiler::CreateCondition(UStateTreeState& State, const FString&
 	StructDesc.Name = CondNode.GetName();
 	StructDesc.DataSource = EStateTreeBindableStructSource::Condition;
 
-	if (FStateTreeConditionBase* Cond = CreateNodeWithSharedInstanceData<FStateTreeConditionBase>(&State, CondNode, StructDesc))
+	FInstancedStruct* CreatedNode = CreateNodeWithSharedInstanceData(&State, CondNode, StructDesc);
+	if (CreatedNode)
 	{
-		if (Cond->EvaluationMode == EStateTreeConditionEvaluationMode::ForcedFalse
-			|| Cond->EvaluationMode == EStateTreeConditionEvaluationMode::ForcedTrue)
+		FStateTreeConditionBase* Cond = CreatedNode->GetMutablePtr<FStateTreeConditionBase>();
+		if (ensure(Cond))
 		{
-			Log.Reportf(EMessageSeverity::Info, StructDesc, 
-						TEXT("The condition result will always be %s."),
-						Cond->EvaluationMode == EStateTreeConditionEvaluationMode::ForcedTrue ? TEXT("True") : TEXT("False"));
-		}
+			if (Cond->EvaluationMode == EStateTreeConditionEvaluationMode::ForcedFalse
+				|| Cond->EvaluationMode == EStateTreeConditionEvaluationMode::ForcedTrue)
+			{
+				Log.Reportf(EMessageSeverity::Info, StructDesc,
+					TEXT("The condition result will always be %s."),
+					Cond->EvaluationMode == EStateTreeConditionEvaluationMode::ForcedTrue ? TEXT("True") : TEXT("False"));
+			}
 
-		Cond->Operand = Operand;
-		Cond->DeltaIndent = DeltaIndent;
-		return true;
+			Cond->Operand = Operand;
+			Cond->DeltaIndent = DeltaIndent;
+			return true;
+		}
 	}
 
 	return false;
@@ -1967,11 +1993,16 @@ bool FStateTreeCompiler::CreateConsideration(UStateTreeState& State, const FStri
 	StructDesc.Name = ConsiderationNode.GetName();
 	StructDesc.DataSource = EStateTreeBindableStructSource::Consideration;
 
-	if (FStateTreeConsiderationBase* Consideration = CreateNodeWithSharedInstanceData<FStateTreeConsiderationBase>(&State, ConsiderationNode, StructDesc))
+	FInstancedStruct* CreatedNode = CreateNodeWithSharedInstanceData(&State, ConsiderationNode, StructDesc);
+	if (CreatedNode)
 	{
-		Consideration->Operand = Operand;
-		Consideration->DeltaIndent = DeltaIndent;
-		return true;
+		FStateTreeConsiderationBase * Consideration = CreatedNode->GetMutablePtr<FStateTreeConsiderationBase>();
+		if (ensure(Consideration))
+		{
+			Consideration->Operand = Operand;
+			Consideration->DeltaIndent = DeltaIndent;
+			return true;
+		}
 	}
 
 	return false;
@@ -2074,59 +2105,18 @@ bool FStateTreeCompiler::CreateTask(UStateTreeState* State, const FStateTreeEdit
 	InstantiateStructSubobjects(Node);
 	
 	FStateTreeTaskBase& Task = Node.GetMutable<FStateTreeTaskBase>();
-	FStateTreeDataView InstanceDataView;
-
-	// Update task name as description for runtime.
-	Task.Name = TaskNode.GetName();
-	
-	if (TaskNode.Instance.IsValid())
-	{
-		// Struct Instance
-		const int32 InstanceIndex = InstanceStructs.Add(TaskNode.Instance);
-		InstantiateStructSubobjects(InstanceStructs[InstanceIndex]);
-
-		// Create binding source struct descriptor.
-		StructDesc.Struct = TaskNode.Instance.GetScriptStruct();
-
-		if (const auto Validation = UE::StateTree::Compiler::IsValidIndex16(InstanceIndex); Validation.DidFail())
-		{
-			Validation.Log(Log, TEXT("InstanceIndex"), StructDesc);
-			return false;
-		}
-		Task.InstanceTemplateIndex = FStateTreeIndex16(InstanceIndex);
-		Task.InstanceDataHandle = TaskDataHandle;
-		InstanceDataView = FStateTreeDataView(InstanceStructs[InstanceIndex]);
-	}
-	else
+	TOptional<FStateTreeDataView> InstanceDataView = CreateNodeInstanceData(TaskNode, Task, StructDesc, TaskDataHandle, InstanceStructs);
+	if (!InstanceDataView.IsSet())
 	{
-		// Object Instance
-		check(TaskNode.InstanceObject != nullptr);
-		UObject* Instance = UE::StateTree::Compiler::DuplicateInstanceObject(Log, StructDesc, TaskNode.ID, TaskNode.InstanceObject, StateTree);
-		FInstancedStruct Wrapper;
-		Wrapper.InitializeAs<FStateTreeInstanceObjectWrapper>(Instance);
-		const int32 InstanceIndex = InstanceStructs.Add(MoveTemp(Wrapper));
-
-		// Create binding source struct descriptor.
-		StructDesc.Struct = Instance->GetClass();
-
-		if (const auto Validation = UE::StateTree::Compiler::IsValidIndex16(InstanceIndex); Validation.DidFail())
-		{
-			Validation.Log(Log, TEXT("InstanceIndex"), StructDesc);
-			return false;
-		}
-		Task.InstanceTemplateIndex = FStateTreeIndex16(InstanceIndex);
-		Task.InstanceDataHandle = TaskDataHandle.ToObjectSource();
-		InstanceDataView = FStateTreeDataView(Instance);
+		return false;
 	}
 
-	StructDesc.DataHandle = Task.InstanceDataHandle;
-
 	if (!Task.bTaskEnabled)
 	{
 		Log.Reportf(EMessageSeverity::Info, StructDesc, TEXT("Task is disabled and will have no effect."));
 	}
 
-	if (!CompileAndValidateNode(State, StructDesc, Node,  InstanceDataView))
+	if (!CompileAndValidateNode(State, StructDesc, Node,  InstanceDataView.GetValue()))
 	{
 		return false;
 	}
@@ -2166,55 +2156,14 @@ bool FStateTreeCompiler::CreateEvaluator(const FStateTreeEditorNode& EvalNode, c
 	InstantiateStructSubobjects(Node);
 	
 	FStateTreeEvaluatorBase& Eval = Node.GetMutable<FStateTreeEvaluatorBase>();
-	FStateTreeDataView InstanceDataView;
-
-	// Update eval name as description for runtime.
-	Eval.Name = EvalNode.GetName();
-
-	if (EvalNode.Instance.IsValid())
+	TOptional<FStateTreeDataView> InstanceDataView = CreateNodeInstanceData(EvalNode, Eval, StructDesc, EvalDataHandle, InstanceStructs);
+	if (!InstanceDataView.IsSet())
 	{
-		// Struct Instance
-		const int32 InstanceIndex = InstanceStructs.Add(EvalNode.Instance);
-		InstantiateStructSubobjects(InstanceStructs[InstanceIndex]);
-
-		// Create binding source struct descriptor.
-		StructDesc.Struct = EvalNode.Instance.GetScriptStruct();
-
-		if (const auto Validation = UE::StateTree::Compiler::IsValidIndex16(InstanceIndex); Validation.DidFail())
-		{
-			Validation.Log(Log, TEXT("InstanceIndex"), StructDesc);
-			return false;
-		}
-		Eval.InstanceTemplateIndex = FStateTreeIndex16(InstanceIndex);
-		Eval.InstanceDataHandle = EvalDataHandle;
-		InstanceDataView = FStateTreeDataView(InstanceStructs[InstanceIndex]);
-	}
-	else
-	{
-		// Object Instance
-		check(EvalNode.InstanceObject != nullptr);
-
-		UObject* Instance = UE::StateTree::Compiler::DuplicateInstanceObject(Log, StructDesc, EvalNode.ID, EvalNode.InstanceObject, StateTree);
-		FInstancedStruct Wrapper;
-		Wrapper.InitializeAs<FStateTreeInstanceObjectWrapper>(Instance);
-		const int32 InstanceIndex = InstanceStructs.Add(MoveTemp(Wrapper));
-		
-		// Create binding source struct descriptor.
-		StructDesc.Struct = Instance->GetClass();
-
-		if (const auto Validation = UE::StateTree::Compiler::IsValidIndex16(InstanceIndex); Validation.DidFail())
-		{
-			Validation.Log(Log, TEXT("InstanceIndex"), StructDesc);
-			return false;
-		}
-		Eval.InstanceTemplateIndex = FStateTreeIndex16(InstanceIndex);
-		Eval.InstanceDataHandle = EvalDataHandle.ToObjectSource();
-		InstanceDataView = FStateTreeDataView(Instance);
+		return false;
 	}
 
-	StructDesc.DataHandle = Eval.InstanceDataHandle;
 
-	if (!CompileAndValidateNode(nullptr, StructDesc, Node,  InstanceDataView))
+	if (!CompileAndValidateNode(nullptr, StructDesc, Node,  InstanceDataView.GetValue()))
 	{
 		return false;
 	}
diff --git a/Engine/Plugins/Runtime/StateTree/Source/StateTreeEditorModule/Public/StateTreeCompiler.h b/Engine/Plugins/Runtime/StateTree/Source/StateTreeEditorModule/Public/StateTreeCompiler.h
--- a/Engine/Plugins/Runtime/StateTree/Source/StateTreeEditorModule/Public/StateTreeCompiler.h
+++ b/Engine/Plugins/Runtime/StateTree/Source/StateTreeEditorModule/Public/StateTreeCompiler.h
@@ -83,13 +83,13 @@ private:
 	bool CreateCondition(UStateTreeState& State, const FString& StatePath, const FStateTreeEditorNode& CondNode, const EStateTreeExpressionOperand Operand, const int8 DeltaIndent);
 	bool CreateConsiderations(UStateTreeState& State, const FString& StatePath, TConstArrayView<FStateTreeEditorNode> Considerations);
 	bool CreateConsideration(UStateTreeState& State, const FString& StatePath, const FStateTreeEditorNode& ConsiderationNode, const EStateTreeExpressionOperand Operand, const int8 DeltaIndent);
-
-	template<class T = FStateTreeNodeBase>
-	T* CreateNodeWithSharedInstanceData(UStateTreeState* State, const FStateTreeEditorNode& EditorNode, FStateTreeBindableStructDesc& StructDesc);
-
 	bool CreateTask(UStateTreeState* State, const FStateTreeEditorNode& TaskNode, const FStateTreeDataHandle TaskDataHandle);
 	bool CreateEvaluator(const FStateTreeEditorNode& EvalNode, const FStateTreeDataHandle EvalDataHandle);
 
+	FInstancedStruct* CreateNode(UStateTreeState* State, const FStateTreeEditorNode& EditorNode, FStateTreeBindableStructDesc& StructDesc, const FStateTreeDataHandle DataHandle, TArray<FInstancedStruct>& InstancedStructContainer);
+	FInstancedStruct* CreateNodeWithSharedInstanceData(UStateTreeState* State, const FStateTreeEditorNode& EditorNode, FStateTreeBindableStructDesc& StructDesc);
+	TOptional<FStateTreeDataView> CreateNodeInstanceData(const FStateTreeEditorNode& EditorNode, FStateTreeNodeBase& Node, FStateTreeBindableStructDesc& StructDesc, const FStateTreeDataHandle DataHandle, TArray<FInstancedStruct>& InstancedStructContainer);
+
 	struct FValidatedPathBindings
 	{
 		TArray<FStateTreePropertyPathBinding> CopyBindings;
@@ -122,6 +122,7 @@ private:
 	TArray<FInstancedStruct> Nodes;
 	TArray<FInstancedStruct> InstanceStructs;
 	TArray<FInstancedStruct> SharedInstanceStructs;
+	TArray<FInstancedStruct> EvaluationScopeStructs;
 	
 	/** Cached result of MakeCompletionTasksMask for global tasks. Indicates where state tasks should start. */
 	int32 GlobalTaskEndBit = 0;
diff --git a/Engine/Plugins/Runtime/StateTree/Source/StateTreeModule/Private/StateTree.cpp b/Engine/Plugins/Runtime/StateTree/Source/StateTreeModule/Private/StateTree.cpp
--- a/Engine/Plugins/Runtime/StateTree/Source/StateTreeModule/Private/StateTree.cpp
+++ b/Engine/Plugins/Runtime/StateTree/Source/StateTreeModule/Private/StateTree.cpp
@@ -243,9 +243,11 @@ void UStateTree::ResetCompiled()
 	Transitions.Reset();
 	Nodes.Reset();
 	DefaultInstanceData.Reset();
+	DefaultEvaluationScopeInstanceData.Reset();
 	SharedInstanceData.Reset();
 	ContextDataDescs.Reset();
 	PropertyBindings.Reset();
+	PropertyFunctionEvaluationScopeMemoryRequirements.Reset();
 	Extensions.Reset();
 	Parameters.Reset();
 	ParameterDataType = EStateTreeParameterDataType::GlobalParameterData;
@@ -377,17 +379,32 @@ void UStateTree::PostLoad()
 		{
 			if (Node->InstanceTemplateIndex.IsValid())
 			{
-				const bool bIsUsingSharedInstanceData = NodeView.GetScriptStruct()->IsChildOf<FStateTreeConditionBase>()
-														|| NodeView.GetScriptStruct()->IsChildOf<FStateTreeConsiderationBase>()
-														|| NodeView.GetScriptStruct()->IsChildOf<FStateTreePropertyFunctionBase>();
-				FStateTreeInstanceData& SourceInstanceData = bIsUsingSharedInstanceData ? SharedInstanceData : DefaultInstanceData;
-				if (SourceInstanceData.IsObject(Node->InstanceTemplateIndex.Get()))
+				const bool bIsUsingSharedInstanceData = Node->InstanceDataHandle.GetSource() == EStateTreeDataSourceType::SharedInstanceData
+					|| Node->InstanceDataHandle.GetSource() == EStateTreeDataSourceType::SharedInstanceDataObject;
+				const bool bIsUsingEvaluationScopeInstanceData = Node->InstanceDataHandle.GetSource() == EStateTreeDataSourceType::EvaluationScopeInstanceData
+					|| Node->InstanceDataHandle.GetSource() == EStateTreeDataSourceType::EvaluationScopeInstanceDataObject;
+				if (bIsUsingEvaluationScopeInstanceData)
 				{
-					Node->PostLoad(SourceInstanceData.GetMutableObject(Node->InstanceTemplateIndex.Get()));
+					if (DefaultEvaluationScopeInstanceData.IsObject(Node->InstanceTemplateIndex.Get()))
+					{
+						Node->PostLoad(DefaultEvaluationScopeInstanceData.GetMutableObject(Node->InstanceTemplateIndex.Get()));
+					}
+					else
+					{
+						Node->PostLoad(DefaultEvaluationScopeInstanceData.GetMutableStruct(Node->InstanceTemplateIndex.Get()));
+					}
 				}
 				else
 				{
-					Node->PostLoad(SourceInstanceData.GetMutableStruct(Node->InstanceTemplateIndex.Get()));
+					FStateTreeInstanceData& SourceInstanceData = bIsUsingSharedInstanceData ? SharedInstanceData : DefaultInstanceData;
+					if (SourceInstanceData.IsObject(Node->InstanceTemplateIndex.Get()))
+					{
+						Node->PostLoad(SourceInstanceData.GetMutableObject(Node->InstanceTemplateIndex.Get()));
+					}
+					else
+					{
+						Node->PostLoad(SourceInstanceData.GetMutableStruct(Node->InstanceTemplateIndex.Get()));
+					}
 				}
 			}
 		}
@@ -487,16 +504,34 @@ bool UStateTree::ValidateInstanceData()
 		{
 			const UStruct* CurrentInstanceDataType = nullptr;
 			{
-				const bool bUseSharedInstanceData = NodeView.GetPtr<const FStateTreeConditionBase>() || NodeView.GetPtr<const FStateTreeConsiderationBase>() || NodeView.GetPtr<const FStateTreePropertyFunctionBase>();
-				const FStateTreeInstanceData& SourceInstanceData = bUseSharedInstanceData ? SharedInstanceData : DefaultInstanceData;
-				if (SourceInstanceData.IsObject(Node->InstanceTemplateIndex.Get()))
+				 const bool bIsUsingSharedInstanceData = Node->InstanceDataHandle.GetSource() == EStateTreeDataSourceType::SharedInstanceData
+					|| Node->InstanceDataHandle.GetSource() == EStateTreeDataSourceType::SharedInstanceDataObject;
+				 const bool bIsUsingEvaluationScopeInstanceData = Node->InstanceDataHandle.GetSource() == EStateTreeDataSourceType::EvaluationScopeInstanceData
+					 || Node->InstanceDataHandle.GetSource() == EStateTreeDataSourceType::EvaluationScopeInstanceDataObject;
+				if (bIsUsingEvaluationScopeInstanceData)
 				{
-					const UObject* InstanceObject = SourceInstanceData.GetObject(Node->InstanceTemplateIndex.Get());
-					CurrentInstanceDataType = InstanceObject ? InstanceObject->GetClass() : nullptr;
+					if (DefaultEvaluationScopeInstanceData.IsObject(Node->InstanceTemplateIndex.Get()))
+					{
+						const UObject* InstanceObject = DefaultEvaluationScopeInstanceData.GetObject(Node->InstanceTemplateIndex.Get());
+						CurrentInstanceDataType = InstanceObject ? InstanceObject->GetClass() : nullptr;
+					}
+					else
+					{
+						CurrentInstanceDataType = DefaultEvaluationScopeInstanceData.GetStruct(Node->InstanceTemplateIndex.Get()).GetScriptStruct();
+					}
 				}
 				else
 				{
-					CurrentInstanceDataType = SourceInstanceData.GetStruct(Node->InstanceTemplateIndex.Get()).GetScriptStruct();
+					const FStateTreeInstanceData& SourceInstanceData = bIsUsingSharedInstanceData ? SharedInstanceData : DefaultInstanceData;
+					if (SourceInstanceData.IsObject(Node->InstanceTemplateIndex.Get()))
+					{
+						const UObject* InstanceObject = SourceInstanceData.GetObject(Node->InstanceTemplateIndex.Get());
+						CurrentInstanceDataType = InstanceObject ? InstanceObject->GetClass() : nullptr;
+					}
+					else
+					{
+						CurrentInstanceDataType = SourceInstanceData.GetStruct(Node->InstanceTemplateIndex.Get()).GetScriptStruct();
+					}
 				}
 			}
 
@@ -671,6 +706,12 @@ bool UStateTree::Link()
 		UE_LOG(LogStateTree, Error, TEXT("%s: State Tree asset was not properly loaded (missing shared instance data). See log for loading failures, or recompile the StateTree asset."), *GetPathName());
 		return false;
 	}
+
+	if (!DefaultEvaluationScopeInstanceData.AreAllInstancesValid())
+	{
+		UE_LOG(LogStateTree, Error, TEXT("%s: State Tree asset was not properly loaded (missing evaluation scope instance data). See log for loading failures, or recompile the StateTree asset."), *GetPathName());
+		return false;
+	}
 	
 	if (!PatchBindings())
 	{
@@ -780,6 +821,7 @@ bool UStateTree::PatchBindings()
 			}
 		}
 	}
+
 	for (FPropertyBindingCopyInfoBatch& CopyBatch : CopyBatches)
 	{
 		if (const UClass* TargetClass = Cast<UClass>(CopyBatch.TargetStruct.Get().Struct))
@@ -923,19 +965,27 @@ bool UStateTree::PatchBindings()
 	{
 		const FStateTreeNodeBase& Node = NodeView.Get<const FStateTreeNodeBase>();
 
-		FStateTreeInstanceData* SourceInstanceData = &DefaultInstanceData;
-		if (NodeView.GetScriptStruct()->IsChildOf<FStateTreeConditionBase>()
-            || NodeView.GetScriptStruct()->IsChildOf<FStateTreeConsiderationBase>()
-            || NodeView.GetScriptStruct()->IsChildOf<FStateTreePropertyFunctionBase>())
+		FStateTreeDataView NodeDataView;
+		if (Node.InstanceDataHandle.GetSource() == EStateTreeDataSourceType::SharedInstanceData
+			|| Node.InstanceDataHandle.GetSource() == EStateTreeDataSourceType::SharedInstanceDataObject)
 		{
-			// Conditions, Considerations, and PropertyFunctions are stored in shared instance data.
-			SourceInstanceData = &SharedInstanceData;
+			NodeDataView = Node.InstanceDataHandle.IsObjectSource()
+				? FStateTreeDataView(SharedInstanceData.GetMutableObject(Node.InstanceTemplateIndex.Get()))
+				: FStateTreeDataView(SharedInstanceData.GetMutableStruct(Node.InstanceTemplateIndex.Get()));
+		}
+		else if (Node.InstanceDataHandle.GetSource() == EStateTreeDataSourceType::EvaluationScopeInstanceData
+			|| Node.InstanceDataHandle.GetSource() == EStateTreeDataSourceType::EvaluationScopeInstanceDataObject)
+		{
+			NodeDataView = Node.InstanceDataHandle.IsObjectSource()
+				? FStateTreeDataView(DefaultEvaluationScopeInstanceData.GetMutableObject(Node.InstanceTemplateIndex.Get()))
+				: FStateTreeDataView(DefaultEvaluationScopeInstanceData.GetMutableStruct(Node.InstanceTemplateIndex.Get()));
+		}
+		else
+		{
+			NodeDataView = Node.InstanceDataHandle.IsObjectSource()
+				? FStateTreeDataView(DefaultInstanceData.GetMutableObject(Node.InstanceTemplateIndex.Get()))
+				: FStateTreeDataView(DefaultInstanceData.GetMutableStruct(Node.InstanceTemplateIndex.Get()));
 		}
-
-		FStateTreeDataView NodeDataView = Node.InstanceDataHandle.IsObjectSource()
-						? FStateTreeDataView(SourceInstanceData->GetMutableObject(Node.InstanceTemplateIndex.Get()))
-						: FStateTreeDataView(SourceInstanceData->GetMutableStruct(Node.InstanceTemplateIndex.Get()));
-
 		DataViews.Add(Node.InstanceDataHandle, NodeDataView);
 
 		if (Node.BindingsBatch.IsValid())
@@ -1013,6 +1063,28 @@ bool UStateTree::PatchBindings()
 		}
 	}
 
+	// Cache the amount of memory needed to execute property functions for each batch.
+	PropertyFunctionEvaluationScopeMemoryRequirements.Reset(CopyBatches.Num());
+	for (FPropertyBindingCopyInfoBatch& CopyBatch : CopyBatches)
+	{
+		UE::StateTree::InstanceData::FEvaluationScopeInstanceContainer::FMemoryRequirementBuilder CopyBatchFunctionRequirement;
+
+		if (CopyBatch.PropertyFunctionsBegin != CopyBatch.PropertyFunctionsEnd)
+		{
+			const int32 FuncsBegin = CopyBatch.PropertyFunctionsBegin.Get();
+			const int32 FuncsEnd = CopyBatch.PropertyFunctionsEnd.Get();
+			for (int32 FuncIndex = FuncsBegin; FuncIndex < FuncsEnd; ++FuncIndex)
+			{
+				const FStateTreePropertyFunctionBase& Func = Nodes[FuncIndex].Get<const FStateTreePropertyFunctionBase>();
+				const FConstStructView DefaultInstanceView = DefaultEvaluationScopeInstanceData.GetStruct(Func.InstanceTemplateIndex.Get());
+				CopyBatchFunctionRequirement.Add(DefaultInstanceView.GetScriptStruct());
+			}
+		}
+
+		PropertyFunctionEvaluationScopeMemoryRequirements.Add(CopyBatchFunctionRequirement.Build());
+	}
+	check(CopyBatches.Num() == PropertyFunctionEvaluationScopeMemoryRequirements.Num());
+
 	return true;
 }
 
@@ -1374,6 +1446,22 @@ FString UStateTree::DebugInternalLayoutAsString() const
 		}
 	}
 
+	// Execution Scope InstanceData data (e.g. tasks)
+	DebugString.Appendf(TEXT("\nExecution Scope(%d)\n"), DefaultEvaluationScopeInstanceData.Num());
+	for (int32 Index = 0; Index < DefaultEvaluationScopeInstanceData.Num(); Index++)
+	{
+		if (DefaultEvaluationScopeInstanceData.IsObject(Index))
+		{
+			const UObject* Data = DefaultEvaluationScopeInstanceData.GetObject(Index);
+			PrintObjectNameSafe(Index, Data);
+		}
+		else
+		{
+			const FConstStructView Data = DefaultEvaluationScopeInstanceData.GetStruct(Index);
+			PrintViewNameSafe(Index, Data);
+		}
+	}
+
 	// External data (e.g. fragments, subsystems)
 	DebugString.Appendf(TEXT("\nExternal Data(%d)\n"), ExternalDataDescs.Num());
 	if (ExternalDataDescs.Num())
diff --git a/Engine/Plugins/Runtime/StateTree/Source/StateTreeModule/Private/StateTreeEvaluationScopeInstanceContainer.cpp b/Engine/Plugins/Runtime/StateTree/Source/StateTreeModule/Private/StateTreeEvaluationScopeInstanceContainer.cpp
new file mode 100644
--- /dev/null
+++ b/Engine/Plugins/Runtime/StateTree/Source/StateTreeModule/Private/StateTreeEvaluationScopeInstanceContainer.cpp
@@ -0,0 +1,195 @@
+// Copyright Epic Games, Inc. All Rights Reserved.
+
+#include "StateTreeEvaluationScopeInstanceContainer.h"
+
+#include "UObject/Class.h"
+#include "UObject/Package.h"
+#include <type_traits>
+
+namespace UE::StateTree::InstanceData
+{
+namespace Private
+{
+	static constexpr uint32 DebugTableEndTag = 0x99AABBCC;
+	static constexpr uint32 DebugStructEndTag = 0xFFEEDDCC;
+	void Add(FEvaluationScopeInstanceContainer::FMemoryRequirement& MemoryRequirement, TNotNull<const UScriptStruct*> Struct)
+	{
+		const int32 StructSize = Struct->GetStructureSize();
+		const int32 StructAlignment = Struct->GetMinAlignment();
+		MemoryRequirement.Alignment = FMath::Max(StructAlignment, MemoryRequirement.Alignment);
+		MemoryRequirement.Size = Align(MemoryRequirement.Size, StructAlignment);
+		MemoryRequirement.Size += StructSize;
+		MemoryRequirement.Size += sizeof(DebugStructEndTag); // for struct end tag
+	}
+}
+
+void FEvaluationScopeInstanceContainer::FMemoryRequirementBuilder::Add(TNotNull<const UScriptStruct*> Struct)
+{
+	if (!Struct->IsChildOf(FStateTreeInstanceObjectWrapper::StaticStruct()))
+	{
+		Private::Add(MemoryRequirement, Struct);
+		if (FirstStructAlignment <= 0)
+		{
+			FirstStructAlignment = Struct->GetMinAlignment();
+		}
+	}
+	++MemoryRequirement.NumberOfElements;
+}
+
+FEvaluationScopeInstanceContainer::FMemoryRequirement FEvaluationScopeInstanceContainer::FMemoryRequirementBuilder::Build()
+{
+	int32 ContainerSize = sizeof(FItem) * MemoryRequirement.NumberOfElements;
+	ContainerSize += sizeof(Private::DebugTableEndTag); // table for end tag
+	ContainerSize = Align(ContainerSize, FirstStructAlignment);
+	FMemoryRequirement Result = MemoryRequirement;
+	Result.Size += ContainerSize;
+	return Result;
+}
+
+FEvaluationScopeInstanceContainer::FEvaluationScopeInstanceContainer(TNotNull<void*> InMemory, const FMemoryRequirement& InRequirement)
+	: Memory(InMemory)
+	, MemoryRequirement(InRequirement)
+{
+	AddTableDebugTag();
+}
+
+FEvaluationScopeInstanceContainer::~FEvaluationScopeInstanceContainer()
+{
+	Reset();
+}
+
+void FEvaluationScopeInstanceContainer::Add(FStateTreeDataHandle DataHandle, FConstStructView DefaultInstance)
+{
+	check(NumberOfElements < MemoryRequirement.NumberOfElements);
+
+	FItem* ItemContainer = static_cast<FItem*>(Memory);
+	const UScriptStruct* Struct = DefaultInstance.GetScriptStruct();
+	if (!Struct->IsChildOf(FStateTreeInstanceObjectWrapper::StaticStruct()))
+	{
+		void* MemoryForNewItem = (NumberOfElements > 0)
+			? (uint8*)ItemContainer[NumberOfElements - 1].Instance.GetMutableMemory()
+				+ ItemContainer[NumberOfElements - 1].Instance.GetStruct()->GetStructureSize()
+				+ sizeof(Private::DebugStructEndTag)
+			: (uint8*)Memory + (MemoryRequirement.NumberOfElements * sizeof(FItem))
+				+ sizeof(Private::DebugTableEndTag);
+		MemoryForNewItem = Align(MemoryForNewItem, Struct->GetMinAlignment());
+
+		// Initialize the instance data
+		constexpr int32 ArrayDim = 1;
+		Struct->InitializeStruct(MemoryForNewItem);
+		Struct->CopyScriptStruct(MemoryForNewItem, DefaultInstance.GetMemory(), ArrayDim);
+
+		// Initialize the item in the table
+		new (ItemContainer + NumberOfElements)FItem();
+		ItemContainer[NumberOfElements].Instance = FStateTreeDataView(Struct, (uint8*)MemoryForNewItem);
+		ItemContainer[NumberOfElements].DataHandle = DataHandle;
+
+		AddStructDebugTag(NumberOfElements);
+
+		if ((Struct->StructFlags & (STRUCT_IsPlainOldData | STRUCT_NoDestructor)) == 0)
+		{
+			bStructsHaveDestructor = true;
+		}
+	}
+	else
+	{
+		const FStateTreeInstanceObjectWrapper& Wrapper = DefaultInstance.Get<const FStateTreeInstanceObjectWrapper>();
+		if (Wrapper.InstanceObject)
+		{
+			UObject* DuplicatedObject = ::DuplicateObject(Wrapper.InstanceObject, GetTransientPackage());
+
+			new (ItemContainer + NumberOfElements)FItem();
+			ItemContainer[NumberOfElements].Instance = FStateTreeDataView(DuplicatedObject);
+			ItemContainer[NumberOfElements].DataHandle = DataHandle;
+		}
+	}
+	++NumberOfElements;
+}
+
+void FEvaluationScopeInstanceContainer::Reset()
+{
+	TestDebugTags();
+
+	if (bStructsHaveDestructor)
+	{
+		for (; NumberOfElements > 0; --NumberOfElements)
+		{
+			FItem* ItemContainer = static_cast<FItem*>(Memory);
+			if (!ItemContainer[NumberOfElements - 1].DataHandle.IsObjectSource())
+			{
+				FStateTreeDataView& Instance = ItemContainer[NumberOfElements - 1].Instance;
+				Instance.GetStruct()->DestroyStruct(Instance.GetMutableMemory());
+			}
+			if constexpr (!std::is_trivially_destructible_v<FItem>)
+			{
+				ItemContainer->~FItem();
+			}
+		}
+		check(NumberOfElements == 0);
+	}
+	else if constexpr (!std::is_trivially_destructible_v<FItem>)
+	{
+		for (; NumberOfElements > 0; --NumberOfElements)
+		{
+			FItem* ItemContainer = static_cast<FItem*>(Memory);
+			ItemContainer->~FItem();
+		}
+		check(NumberOfElements == 0);
+	}
+	NumberOfElements = 0;
+}
+
+FEvaluationScopeInstanceContainer::FItem* FEvaluationScopeInstanceContainer::Get(const FStateTreeDataHandle& DataHandle) const
+{
+	FItem* ItemContainer = static_cast<FItem*>(Memory);
+	for (int32 Index = 0; Index < NumberOfElements; ++Index)
+	{
+		if (ItemContainer[Index].DataHandle == DataHandle)
+		{
+			return &ItemContainer[Index];
+		}
+	}
+	return nullptr;
+}
+
+void FEvaluationScopeInstanceContainer::AddTableDebugTag()
+{
+#if WITH_STATETREE_DEBUG
+	void* MemoryForDebugTag = (uint8*)Memory + (sizeof(FItem) * MemoryRequirement.NumberOfElements);
+	*reinterpret_cast<uint32*>(MemoryForDebugTag) = Private::DebugTableEndTag;
+#endif
+}
+
+void FEvaluationScopeInstanceContainer::AddStructDebugTag(int32 Index)
+{
+#if WITH_STATETREE_DEBUG
+	FItem* ItemContainer = static_cast<FItem*>(Memory);
+	void* MemoryForDebugTag = (uint8*)ItemContainer[Index].Instance.GetMutableMemory()
+		+ ItemContainer[Index].Instance.GetStruct()->GetStructureSize();
+	*reinterpret_cast<uint32*>(MemoryForDebugTag) = Private::DebugStructEndTag;
+#endif
+}
+
+void FEvaluationScopeInstanceContainer::TestDebugTags() const
+{
+#if WITH_STATETREE_DEBUG && DO_ENSURE
+	CA_SUPPRESS(6269); // warning C6269: Possibly incorrect order of operations.
+	if (Memory)
+	{
+		const void* MemoryForTableDebugTag = (const uint8*)Memory + (sizeof(FItem) * MemoryRequirement.NumberOfElements);
+		const int32 TableEndTag = *reinterpret_cast<const uint32*>(MemoryForTableDebugTag);
+		ensure(TableEndTag == Private::DebugTableEndTag);
+
+		FItem* ItemContainer = static_cast<FItem*>(Memory);
+		for (int32 Index = 0; Index < NumberOfElements; ++Index)
+		{
+			const void* MemoryForStructDebugTag = (const uint8*)ItemContainer[Index].Instance.GetMutableMemory()
+				+ ItemContainer[Index].Instance.GetStruct()->GetStructureSize();
+			const int32 StructEndTag = *reinterpret_cast<const uint32*>(MemoryForStructDebugTag);
+			ensure(StructEndTag == Private::DebugStructEndTag);
+		}
+	}
+#endif
+}
+
+} // namespace UE::StateTree::InstanceData
diff --git a/Engine/Plugins/Runtime/StateTree/Source/StateTreeModule/Private/StateTreeExecutionContext.cpp b/Engine/Plugins/Runtime/StateTree/Source/StateTreeModule/Private/StateTreeExecutionContext.cpp
--- a/Engine/Plugins/Runtime/StateTree/Source/StateTreeModule/Private/StateTreeExecutionContext.cpp
+++ b/Engine/Plugins/Runtime/StateTree/Source/StateTreeModule/Private/StateTreeExecutionContext.cpp
@@ -2482,6 +2482,24 @@ FStateTreeDataView FStateTreeExecutionContext::GetDataView(const FStateTreeExecu
 		check(!ContextAndExternalDataViews.IsEmpty());
 		return ContextAndExternalDataViews[Handle.GetIndex()];
 
+	case EStateTreeDataSourceType::EvaluationScopeInstanceData:
+	case EStateTreeDataSourceType::EvaluationScopeInstanceDataObject:
+	{
+		// The data can be accessed in any of the caches (depends on how the binding was constructed), but it is most likely on top of the stack.
+		for (int32 Index = EvaluationScopeInstanceCaches.Num() - 1; Index >= 0; --Index)
+		{
+			if (EvaluationScopeInstanceCaches[Index].StateTree == CurrentFrame.StateTree)
+			{
+				if (FStateTreeDataView* Result = EvaluationScopeInstanceCaches[Index].Container->GetDataViewPtr(Handle))
+				{
+					return *Result;
+				}
+			}
+		}
+		checkf(false, TEXT("The evaluation scope instance data needs to be constructed before you can access it."));
+		return nullptr;
+	}
+
 	case EStateTreeDataSourceType::ExternalData:
 		check(!ContextAndExternalDataViews.IsEmpty());
 		return ContextAndExternalDataViews[CurrentFrame.ExternalDataBaseIndex.Get() + Handle.GetIndex()];
@@ -2778,6 +2796,21 @@ bool FStateTreeExecutionContext::IsHandleSourceValid(const FStateTreeExecutionFr
 	case EStateTreeDataSourceType::ContextData:
 		return true;
 
+	case EStateTreeDataSourceType::EvaluationScopeInstanceData:
+	case EStateTreeDataSourceType::EvaluationScopeInstanceDataObject:
+		// The data can be accessed in any of the caches (depends on how the binding was constructed), but it is most likely on top of the stack.
+		for (int32 Index = EvaluationScopeInstanceCaches.Num() - 1; Index >= 0; --Index)
+		{
+			if (EvaluationScopeInstanceCaches[Index].StateTree == CurrentFrame.StateTree)
+			{
+				if (FStateTreeDataView* Result = EvaluationScopeInstanceCaches[Index].Container->GetDataViewPtr(Handle))
+				{
+					return true;
+				}
+			}
+		}
+		return false;
+
 	case EStateTreeDataSourceType::ExternalData:
 		return CurrentFrame.ExternalDataBaseIndex.IsValid()
 			&& ContextAndExternalDataViews.IsValidIndex(CurrentFrame.ExternalDataBaseIndex.Get() + Handle.GetIndex());
@@ -2843,6 +2876,10 @@ FStateTreeDataView FStateTreeExecutionContext::GetTemporaryDataView(const FState
 	case EStateTreeDataSourceType::ExternalGlobalParameterData:
 		checkf(false, TEXT("External global parameter data currently not supported for linked state-trees"));
 		return {};
+	case EStateTreeDataSourceType::EvaluationScopeInstanceData:
+	case EStateTreeDataSourceType::EvaluationScopeInstanceDataObject:
+		ensureMsgf(false, TEXT("The evaluation scope instance data needs to be constructed before you can access it."));
+		return {};
 
 	default:
 		return UE::StateTree::InstanceData::Private::GetTemporaryDataView(Storage, ParentFrame, CurrentFrame, Handle);
@@ -2859,6 +2896,19 @@ FStateTreeDataView FStateTreeExecutionContext::AddTemporaryInstance(const FState
 	return NewInstance;
 }
 
+void FStateTreeExecutionContext::PushEvaluationScopeInstanceContainer(UE::StateTree::InstanceData::FEvaluationScopeInstanceContainer& Container, const FStateTreeExecutionFrame& Frame)
+{
+	EvaluationScopeInstanceCaches.Emplace(&Container, Frame.StateTree);
+}
+
+void FStateTreeExecutionContext::PopEvaluationScopeInstanceContainer(UE::StateTree::InstanceData::FEvaluationScopeInstanceContainer& Container)
+{
+	if (ensure(EvaluationScopeInstanceCaches.Num() > 0 && EvaluationScopeInstanceCaches.Last().Container == &Container))
+	{
+		EvaluationScopeInstanceCaches.Pop();
+	}
+}
+
 void FStateTreeExecutionContext::CopyAllBindingsOnActiveInstances(const ECopyBindings CopyType)
 {
 	FStateTreeExecutionState& Exec = GetExecState();
@@ -2960,15 +3010,43 @@ bool FStateTreeExecutionContext::CopyBatchWithValidation(const FStateTreeExecuti
 	return CopyBatchInternal<bOnActiveInstances>(ParentFrame, CurrentFrame, TargetView, BindingsBatch);
 }
 
+namespace UE::StateTree::ExecutionContext::Private
+{
+	void InitPropertyFunctions(UE::StateTree::InstanceData::FEvaluationScopeInstanceContainer& Container, TNotNull<const UStateTree*> StateTree, int32 FuncBegin, int32 FuncEnd)
+	{
+		for (int32 FuncIndex = FuncBegin; FuncIndex != FuncEnd; ++FuncIndex)
+		{
+			const FStateTreePropertyFunctionBase& Func = StateTree->GetNodes()[FuncIndex].Get<const FStateTreePropertyFunctionBase>();
+			Container.Add(Func.InstanceDataHandle, StateTree->GetDefaultEvaluationScopeInstanceData().GetStruct(Func.InstanceTemplateIndex.Get()));
+		}
+	}
+}
+
 template<bool bOnActiveInstances>
 bool FStateTreeExecutionContext::CopyBatchInternal(const FStateTreeExecutionFrame* ParentFrame, const FStateTreeExecutionFrame& CurrentFrame, const FStateTreeDataView TargetView, const FStateTreeIndex16 BindingsBatch)
 {
+	using namespace UE::StateTree::InstanceData;
+	using namespace UE::StateTree::ExecutionContext::Private;
+
 	const FPropertyBindingCopyInfoBatch& Batch = CurrentFrame.StateTree->PropertyBindings.Super::GetBatch(BindingsBatch);
 	check(TargetView.GetStruct() == Batch.TargetStruct.Get().Struct);
 
+	FEvaluationScopeInstanceContainer EvaluationScopeContainer;
+	bool bEvaluationScopeContainerPushed = false;
 	if (Batch.PropertyFunctionsBegin != Batch.PropertyFunctionsEnd)
 	{
 		check(Batch.PropertyFunctionsBegin.IsValid() && Batch.PropertyFunctionsEnd.IsValid());
+
+		const FEvaluationScopeInstanceContainer::FMemoryRequirement& MemoryRequirement = CurrentFrame.StateTree->PropertyFunctionEvaluationScopeMemoryRequirements[BindingsBatch.Get()];
+		if (ensure(MemoryRequirement.Size > 0))
+		{
+			void* FunctionBindingMemory = FMemory_Alloca_Aligned(MemoryRequirement.Size, MemoryRequirement.Alignment);
+			EvaluationScopeContainer = FEvaluationScopeInstanceContainer(FunctionBindingMemory, MemoryRequirement);
+			PushEvaluationScopeInstanceContainer(EvaluationScopeContainer, CurrentFrame);
+			InitPropertyFunctions(EvaluationScopeContainer, CurrentFrame.StateTree, Batch.PropertyFunctionsBegin.Get(), Batch.PropertyFunctionsEnd.Get());
+			bEvaluationScopeContainerPushed = true;
+		}
+
 		if constexpr (bOnActiveInstances)
 		{
 			EvaluatePropertyFunctionsOnActiveInstances(ParentFrame, CurrentFrame, FStateTreeIndex16(Batch.PropertyFunctionsBegin), Batch.PropertyFunctionsEnd.Get() - Batch.PropertyFunctionsBegin.Get());
@@ -2995,9 +3073,16 @@ bool FStateTreeExecutionContext::CopyBatchInternal(const FStateTreeExecutionFram
 				bSucceed = false;
 				break;
 			}
+
 			bSucceed &= CurrentFrame.StateTree->PropertyBindings.Super::CopyProperty(Copy, SourceView, TargetView);
 		}
 	}
+
+	if (bEvaluationScopeContainerPushed)
+	{
+		PopEvaluationScopeInstanceContainer(EvaluationScopeContainer);
+	}
+
 	return bSucceed;
 }
 
@@ -4603,10 +4688,10 @@ bool FStateTreeExecutionContext::TestAllConditionsInternal(const FStateTreeExecu
 
 	int32 Level = 0;
 
-		for (int32 Index = 0; Index < ConditionsNum; Index++)
-		{
-			const int32 ConditionIndex = ConditionsOffset + Index;
-			const FStateTreeConditionBase& Cond = CurrentFrame.StateTree->Nodes[ConditionIndex].Get<const FStateTreeConditionBase>();
+	for (int32 Index = 0; Index < ConditionsNum; Index++)
+	{
+		const int32 ConditionIndex = ConditionsOffset + Index;
+		const FStateTreeConditionBase& Cond = CurrentFrame.StateTree->Nodes[ConditionIndex].Get<const FStateTreeConditionBase>();
 		checkf(Cond.InstanceDataHandle.GetSource() == EStateTreeDataSourceType::SharedInstanceData
 			|| Cond.InstanceDataHandle.GetSource() == EStateTreeDataSourceType::SharedInstanceDataObject,
 			TEXT("Condition only support SharedInstanceData"));
@@ -4821,12 +4906,6 @@ void FStateTreeExecutionContext::EvaluatePropertyFunctionsInternal(const FStateT
 		}
 
 		Func.Execute(*this);
-
-		// Reset copied properties that might contain object references.
-		if (Func.BindingsBatch.IsValid())
-		{
-			CurrentFrame.StateTree->PropertyBindings.Super::ResetObjects(Func.BindingsBatch, FuncInstanceView);
-		}
 	}
 }
 
diff --git a/Engine/Plugins/Runtime/StateTree/Source/StateTreeModule/Private/StateTreeInstanceContainer.cpp b/Engine/Plugins/Runtime/StateTree/Source/StateTreeModule/Private/StateTreeInstanceContainer.cpp
new file mode 100644
--- /dev/null
+++ b/Engine/Plugins/Runtime/StateTree/Source/StateTreeModule/Private/StateTreeInstanceContainer.cpp
@@ -0,0 +1,26 @@
+// Copyright Epic Games, Inc. All Rights Reserved.
+
+#include "StateTreeInstanceContainer.h"
+#include "StateTreeInstanceDataHelpers.h"
+
+#include UE_INLINE_GENERATED_CPP_BY_NAME(StateTreeInstanceContainer)
+
+namespace UE::StateTree::InstanceData
+{
+
+void FInstanceContainer::Init(TNotNull<UObject*> InOwner, TConstArrayView<FInstancedStruct> InStructs, FAddArgs Args)
+{
+	Private::AppendToInstanceStructContainer(InstanceStructs, InOwner, InStructs, Args.bDuplicateWrappedObject);
+}
+
+void FInstanceContainer::Init(TNotNull<UObject*> InOwner, TConstArrayView<FConstStructView> InStructs, FAddArgs Args)
+{
+	Private::AppendToInstanceStructContainer(InstanceStructs, InOwner, InStructs, Args.bDuplicateWrappedObject);
+}
+
+bool FInstanceContainer::AreAllInstancesValid() const
+{
+	return Private::AreAllInstancesValid(InstanceStructs);
+}
+
+} // namespace UE::StateTree::InstanceData
diff --git a/Engine/Plugins/Runtime/StateTree/Source/StateTreeModule/Private/StateTreeInstanceData.cpp b/Engine/Plugins/Runtime/StateTree/Source/StateTreeModule/Private/StateTreeInstanceData.cpp
--- a/Engine/Plugins/Runtime/StateTree/Source/StateTreeModule/Private/StateTreeInstanceData.cpp
+++ b/Engine/Plugins/Runtime/StateTree/Source/StateTreeModule/Private/StateTreeInstanceData.cpp
@@ -19,11 +19,27 @@
 const FGuid FStateTreeInstanceStorageCustomVersion::GUID(0x60C4F0DE, 0x8B264C34, 0xAA937201, 0x5DFF09CC);
 FCustomVersionRegistration GRegisterStateTreeInstanceStorageCustomVersion(FStateTreeInstanceStorageCustomVersion::GUID, FStateTreeInstanceStorageCustomVersion::LatestVersion, TEXT("StateTreeInstanceStorage"));
 
-namespace UE::StateTree
+namespace UE::StateTree::InstanceData::Private
 {
-	/**
-	 * Duplicates object, and tries to covert old BP classes (REINST_*) to their newer version.
-	 */
+	bool AreAllInstancesValid(const FInstancedStructContainer& InstanceStructs)
+	{
+		for (FConstStructView Instance : InstanceStructs)
+		{
+			if (!Instance.IsValid())
+			{
+				return false;
+			}
+			if (const FStateTreeInstanceObjectWrapper* Wrapper = Instance.GetPtr<const FStateTreeInstanceObjectWrapper>())
+			{
+				if (!Wrapper->InstanceObject)
+				{
+					return false;
+				}
+			}
+		}
+		return true;
+	}
+
 	TNotNull<UObject*> CopyNodeInstance(TNotNull<UObject*> Instance, TNotNull<UObject*> InOwner, bool bDuplicate)
 	{
 		const UClass* InstanceClass = Instance->GetClass();
@@ -58,6 +74,25 @@ namespace UE::StateTree
 		return Instance;
 	}
 
+	void PostAppendToInstanceStructContainer(FInstancedStructContainer& InstanceStructs, TNotNull<UObject*> InOwner, bool bDuplicateWrappedObject, int32 StartIndex)
+	{
+		for (int32 Index = StartIndex; Index < InstanceStructs.Num(); ++Index)
+		{
+			if (FStateTreeInstanceObjectWrapper* Wrapper = InstanceStructs[Index].GetPtr<FStateTreeInstanceObjectWrapper>())
+			{
+				if (Wrapper->InstanceObject)
+				{
+					const bool bDuplicate = bDuplicateWrappedObject || InOwner != Wrapper->InstanceObject->GetOuter();
+					Wrapper->InstanceObject = CopyNodeInstance(Wrapper->InstanceObject, InOwner, bDuplicate);
+				}
+			}
+		}
+	}
+}
+
+namespace UE::StateTree
+{
+
 #if WITH_EDITORONLY_DATA
 	void GatherForLocalization(const FString& PathToParent, const UScriptStruct* Struct, const void* StructData, const void* DefaultStructData, FPropertyLocalizationDataGatherer& PropertyLocalizationDataGatherer, const EPropertyLocalizationGathererTextFlags GatherTextFlags)
 	{
@@ -448,21 +483,7 @@ void FStateTreeInstanceStorage::ResetTransitionRequests()
 
 bool FStateTreeInstanceStorage::AreAllInstancesValid() const
 {
-	for (FConstStructView Instance : InstanceStructs)
-	{
-		if (!Instance.IsValid())
-		{
-			return false;
-		}
-		if (const FStateTreeInstanceObjectWrapper* Wrapper = Instance.GetPtr<const FStateTreeInstanceObjectWrapper>())
-		{
-			if (!Wrapper->InstanceObject)
-			{
-				return false;
-			}
-		}
-	}
-	return true;
+	return UE::StateTree::InstanceData::Private::AreAllInstancesValid(InstanceStructs);
 }
 
 FStructView FStateTreeInstanceStorage::AddTemporaryInstance(UObject& InOwner, const FStateTreeExecutionFrame& Frame, const FStateTreeIndex16 OwnerNodeIndex, const FStateTreeDataHandle DataHandle, FConstStructView NewInstanceData)
@@ -496,7 +517,7 @@ FStructView FStateTreeInstanceStorage::AddTemporaryInstance(UObject& InOwner, co
 		if (Wrapper->InstanceObject)
 		{
 			constexpr bool bDuplicate = true;
-			Wrapper->InstanceObject = UE::StateTree::CopyNodeInstance(Wrapper->InstanceObject, &InOwner, bDuplicate);
+			Wrapper->InstanceObject = UE::StateTree::InstanceData::Private::CopyNodeInstance(Wrapper->InstanceObject, &InOwner, bDuplicate);
 		}
 	}
 
@@ -913,7 +934,7 @@ void FStateTreeInstanceData::CopyFrom(UObject& InOwner, const FStateTreeInstance
 			if (Wrapper->InstanceObject)
 			{
 				constexpr bool bDuplicate = true;
-				Wrapper->InstanceObject = UE::StateTree::CopyNodeInstance(Wrapper->InstanceObject, &InOwner, bDuplicate);
+				Wrapper->InstanceObject = UE::StateTree::InstanceData::Private::CopyNodeInstance(Wrapper->InstanceObject, &InOwner, bDuplicate);
 			}
 		}
 	}
@@ -933,42 +954,12 @@ void FStateTreeInstanceData::Init(UObject& InOwner, TConstArrayView<FConstStruct
 
 void FStateTreeInstanceData::Append(UObject& InOwner, TConstArrayView<FInstancedStruct> InStructs, FAddArgs Args)
 {
-	FStateTreeInstanceStorage& Storage = GetMutableStorage();
-
-	const int32 StartIndex = Storage.InstanceStructs.Num();
-	Storage.InstanceStructs.Append(InStructs);
-
-	for (int32 Index = StartIndex; Index < Storage.InstanceStructs.Num(); Index++)
-	{
-		if (FStateTreeInstanceObjectWrapper* Wrapper = Storage.InstanceStructs[Index].GetPtr<FStateTreeInstanceObjectWrapper>())
-		{
-			if (Wrapper->InstanceObject)
-			{
-				const bool bDuplicate = Args.bDuplicateWrappedObject || &InOwner != Wrapper->InstanceObject->GetOuter();
-				Wrapper->InstanceObject = UE::StateTree::CopyNodeInstance(Wrapper->InstanceObject, &InOwner, bDuplicate);
-			}
-		}
-	}
+	UE::StateTree::InstanceData::Private::AppendToInstanceStructContainer(GetMutableStorage().InstanceStructs, &InOwner, InStructs, Args.bDuplicateWrappedObject);
 }
 
 void FStateTreeInstanceData::Append(UObject& InOwner, TConstArrayView<FConstStructView> InStructs, FAddArgs Args)
 {
-	FStateTreeInstanceStorage& Storage = GetMutableStorage();
-
-	const int32 StartIndex = Storage.InstanceStructs.Num();
-	Storage.InstanceStructs.Append(InStructs);
-	
-	for (int32 Index = StartIndex; Index < Storage.InstanceStructs.Num(); Index++)
-	{
-		if (FStateTreeInstanceObjectWrapper* Wrapper = Storage.InstanceStructs[Index].GetPtr<FStateTreeInstanceObjectWrapper>())
-		{
-			if (Wrapper->InstanceObject)
-			{
-				const bool bDuplicate = Args.bDuplicateWrappedObject || &InOwner != Wrapper->InstanceObject->GetOuter();
-				Wrapper->InstanceObject = UE::StateTree::CopyNodeInstance(Wrapper->InstanceObject, &InOwner, bDuplicate);
-			}
-		}
-	}
+	UE::StateTree::InstanceData::Private::AppendToInstanceStructContainer(GetMutableStorage().InstanceStructs, &InOwner, InStructs, Args.bDuplicateWrappedObject);
 }
 
 void FStateTreeInstanceData::Append(UObject& InOwner, TConstArrayView<FConstStructView> InStructs, TConstArrayView<FInstancedStruct*> InInstancesToMove, FAddArgs Args)
@@ -1002,7 +993,7 @@ void FStateTreeInstanceData::Append(UObject& InOwner, TConstArrayView<FConstStru
 			if (Wrapper->InstanceObject)
 			{
 				const bool bDuplicate = Args.bDuplicateWrappedObject || &InOwner != Wrapper->InstanceObject->GetOuter();
-				Wrapper->InstanceObject = UE::StateTree::CopyNodeInstance(Wrapper->InstanceObject, &InOwner, bDuplicate);
+				Wrapper->InstanceObject = UE::StateTree::InstanceData::Private::CopyNodeInstance(Wrapper->InstanceObject, &InOwner, bDuplicate);
 			}
 		}
 	}
diff --git a/Engine/Plugins/Runtime/StateTree/Source/StateTreeModule/Private/StateTreeInstanceDataHelpers.h b/Engine/Plugins/Runtime/StateTree/Source/StateTreeModule/Private/StateTreeInstanceDataHelpers.h
--- a/Engine/Plugins/Runtime/StateTree/Source/StateTreeModule/Private/StateTreeInstanceDataHelpers.h
+++ b/Engine/Plugins/Runtime/StateTree/Source/StateTreeModule/Private/StateTreeInstanceDataHelpers.h
@@ -25,5 +25,22 @@ namespace UE::StateTree::InstanceData::Private
 		const FStateTreeExecutionFrame* ParentFrame,
 		const FStateTreeExecutionFrame& CurrentFrame,
 		const FStateTreeDataHandle& Handle);
-} // namespace
 
+	/** @return true if all instances are valid. */
+	[[nodiscard]] bool AreAllInstancesValid(const FInstancedStructContainer& InstanceStructs);
+
+	/** Duplicates object, and tries to covert old BP classes (REINST_*) to their newer version. */
+	[[nodiscard]] TNotNull<UObject*> CopyNodeInstance(TNotNull<UObject*> Instance, TNotNull<UObject*> InOwner, bool bDuplicate);
+
+	/** None generic code for AppendToInstanceStructContainer. Called after AppendToInstanceStructContainer to move code to cpp. */
+	void PostAppendToInstanceStructContainer(FInstancedStructContainer& InstanceStructs, TNotNull<UObject*> InOwner, bool bDuplicateWrappedObject, int32 StartIndex);
+
+	/** Appends new items to the instance, and moves existing data into the allocated instances. */
+	template<typename TOtherType>
+	void AppendToInstanceStructContainer(FInstancedStructContainer& InstanceStructs, TNotNull<UObject*> InOwner, TConstArrayView<TOtherType> InStructs, bool bDuplicateWrappedObject)
+	{
+		const int32 StartIndex = InstanceStructs.Num();
+		InstanceStructs.Append(InStructs);
+		PostAppendToInstanceStructContainer(InstanceStructs, InOwner, bDuplicateWrappedObject, StartIndex);
+	}
+} // namespace
diff --git a/Engine/Plugins/Runtime/StateTree/Source/StateTreeModule/Private/StateTreeTypes.cpp b/Engine/Plugins/Runtime/StateTree/Source/StateTreeModule/Private/StateTreeTypes.cpp
--- a/Engine/Plugins/Runtime/StateTree/Source/StateTreeModule/Private/StateTreeTypes.cpp
+++ b/Engine/Plugins/Runtime/StateTree/Source/StateTreeModule/Private/StateTreeTypes.cpp
@@ -144,6 +144,28 @@ PRAGMA_ENABLE_DEPRECATION_WARNINGS
 //////////////////////////////////////////////////////////////////////////
 // FStateTreeDataHandle
 
+bool FStateTreeDataHandle::IsObjectSource() const
+{
+	return Source == EStateTreeDataSourceType::GlobalInstanceDataObject
+		|| Source == EStateTreeDataSourceType::ActiveInstanceDataObject
+		|| Source == EStateTreeDataSourceType::EvaluationScopeInstanceDataObject;
+}
+
+FStateTreeDataHandle FStateTreeDataHandle::ToObjectSource() const
+{
+	switch (Source)
+	{
+	case EStateTreeDataSourceType::GlobalInstanceData:
+		return FStateTreeDataHandle(EStateTreeDataSourceType::GlobalInstanceDataObject, Index, StateHandle);
+	case EStateTreeDataSourceType::ActiveInstanceData:
+		return FStateTreeDataHandle(EStateTreeDataSourceType::ActiveInstanceDataObject, Index, StateHandle);
+	case EStateTreeDataSourceType::EvaluationScopeInstanceData:
+		return FStateTreeDataHandle(EStateTreeDataSourceType::EvaluationScopeInstanceDataObject, Index, StateHandle);
+	default:
+		return *this;
+	}
+}
+
 FString FStateTreeDataHandle::Describe() const
 {
 	switch (Source)
@@ -162,6 +184,10 @@ FString FStateTreeDataHandle::Describe() const
 		return FString::Printf(TEXT("Shared[%d]"), Index);
 	case EStateTreeDataSourceType::SharedInstanceDataObject:
 		return FString::Printf(TEXT("SharedObj[%d]"), Index);
+	case EStateTreeDataSourceType::EvaluationScopeInstanceData:
+		return FString::Printf(TEXT("EvalScope[%d]"), Index);
+	case EStateTreeDataSourceType::EvaluationScopeInstanceDataObject:
+		return FString::Printf(TEXT("EvalScopeObj[%d]"), Index);
 	case EStateTreeDataSourceType::ContextData:
 		return FString::Printf(TEXT("Context[%d]"), Index);
 	case EStateTreeDataSourceType::GlobalParameterData:
diff --git a/Engine/Plugins/Runtime/StateTree/Source/StateTreeModule/Public/StateTree.h b/Engine/Plugins/Runtime/StateTree/Source/StateTreeModule/Public/StateTree.h
--- a/Engine/Plugins/Runtime/StateTree/Source/StateTreeModule/Public/StateTree.h
+++ b/Engine/Plugins/Runtime/StateTree/Source/StateTreeModule/Public/StateTree.h
@@ -3,6 +3,7 @@
 #pragma once
 
 #include "Engine/DataAsset.h"
+#include "StateTreeEvaluationScopeInstanceContainer.h"
 #include "StateTreeSchema.h"
 #include "StateTreePropertyBindings.h"
 #include "StateTreeInstanceData.h"
@@ -129,6 +130,12 @@ public:
 		return NumContextData;
 	}
 
+	/** @return Default evaluation scope instance data. */
+	const UE::StateTree::InstanceData::FInstanceContainer& GetDefaultEvaluationScopeInstanceData() const
+	{
+		return DefaultEvaluationScopeInstanceData;
+	}
+
 	/** @return List of external data required by the state tree */
 	TConstArrayView<FStateTreeExternalDataDesc> GetExternalDataDescs() const
 	{
@@ -389,6 +396,10 @@ private:
 	UPROPERTY()
 	FStateTreeInstanceData DefaultInstanceData;
 
+	/** Default node instance data for evaluation scope (e.g. conditions, considerations, function) */
+	UPROPERTY()
+	UE::StateTree::InstanceData::FInstanceContainer DefaultEvaluationScopeInstanceData;
+
 	/** Shared node instance data (e.g. conditions, considerations). */
 	UPROPERTY()
 	FStateTreeInstanceData SharedInstanceData;
@@ -403,6 +414,10 @@ private:
 	UPROPERTY()
 	FStateTreePropertyBindings PropertyBindings;
 
+	using FMemoryRequirement = UE::StateTree::InstanceData::FEvaluationScopeInstanceContainer::FMemoryRequirement;
+	/** The amount of memory used by property binding copy info batch for property function. */
+	TArray<FMemoryRequirement> PropertyFunctionEvaluationScopeMemoryRequirements;
+
 	/** The asset extensions. A place to add extra information for plugins. */
 	UPROPERTY()
 	TArray<TObjectPtr<UStateTreeExtension>> Extensions;
@@ -443,7 +458,7 @@ private:
 	/** Number of global instance data. */
 	UPROPERTY()
 	uint16 NumGlobalInstanceData = 0;
-	
+
 	/** Index of first evaluator in Nodes. */
 	UPROPERTY()
 	uint16 EvaluatorsBegin = 0;
diff --git a/Engine/Plugins/Runtime/StateTree/Source/StateTreeModule/Public/StateTreeEvaluationScopeInstanceContainer.h b/Engine/Plugins/Runtime/StateTree/Source/StateTreeModule/Public/StateTreeEvaluationScopeInstanceContainer.h
new file mode 100644
--- /dev/null
+++ b/Engine/Plugins/Runtime/StateTree/Source/StateTreeModule/Public/StateTreeEvaluationScopeInstanceContainer.h
@@ -0,0 +1,131 @@
+// Copyright Epic Games, Inc. All Rights Reserved.
+
+#pragma once
+
+#include "StateTreeTypes.h"
+#include "StateTreeInstanceContainer.h"
+
+#define UE_API STATETREEMODULE_API
+
+namespace UE::StateTree::InstanceData
+{
+	/**
+	 * Container of instance data for evaluation scope data.
+	 * The memory is pre-allocated and only available for a short amount of time.
+	 */
+	struct FEvaluationScopeInstanceContainer
+	{
+	public:
+		/** The memory requirement for the container allocation. */
+		struct FMemoryRequirement
+		{
+			int32 Size = 0;
+			int32 Alignment = 0;
+			int32 NumberOfElements = 0;
+		};
+
+		/** Builds the amount of memory required to store the provided structures. */
+		struct FMemoryRequirementBuilder
+		{
+		public:
+			FMemoryRequirementBuilder() = default;
+
+			/** Add a structure to the current memory requirement. Takes into account the structure alignment. */
+			UE_API void Add(TNotNull<const UScriptStruct*> Struct);
+
+			/** @return the requirement memory for all added structures. */
+			[[nodiscard]] UE_API FMemoryRequirement Build();
+
+		private:
+			FMemoryRequirement MemoryRequirement;
+			int32 FirstStructAlignment = -1;
+		};
+
+	public:
+		FEvaluationScopeInstanceContainer() = default;
+		UE_API explicit FEvaluationScopeInstanceContainer(TNotNull<void*> InMemory, const FMemoryRequirement& InRequirement);
+		UE_API ~FEvaluationScopeInstanceContainer();
+		
+		FEvaluationScopeInstanceContainer(const FEvaluationScopeInstanceContainer& Other) = delete;
+		FEvaluationScopeInstanceContainer& operator=(const FEvaluationScopeInstanceContainer& Other) = delete;
+
+		FEvaluationScopeInstanceContainer(FEvaluationScopeInstanceContainer&& Other)
+		{
+			FMemory::Memcpy(this, &Other, sizeof(FEvaluationScopeInstanceContainer));
+			FMemory::Memzero(&Other, sizeof(FEvaluationScopeInstanceContainer));
+		}
+
+		FEvaluationScopeInstanceContainer& operator=(FEvaluationScopeInstanceContainer&& Other)
+		{
+			Reset();
+			FMemory::Memcpy(this, &Other, sizeof(FEvaluationScopeInstanceContainer));
+			FMemory::Memzero(&Other, sizeof(FEvaluationScopeInstanceContainer));
+			return *this;
+		}
+
+	public:
+		/** @return number of items in the storage. */
+		[[nodiscard]] int32 Num() const
+		{
+			return NumberOfElements;
+		}
+
+		/** @return true if the index can be used to get data. */
+		[[nodiscard]] bool IsValidDataHandle(FStateTreeDataHandle DataHandle) const
+		{
+			return Get(DataHandle) != nullptr;
+		}
+
+		/** @return specified item of an added structure. */
+		[[nodiscard]] FStateTreeDataView GetDataView(FStateTreeDataHandle DataHandle)
+		{
+			return Get(DataHandle)->Instance;
+		}
+
+		/** @return specified item if the structure was added. */
+		[[nodiscard]] FStateTreeDataView* GetDataViewPtr(FStateTreeDataHandle DataHandle)
+		{
+			FItem* FoundItem = Get(DataHandle);
+			return FoundItem ? &FoundItem->Instance : nullptr;
+		}
+
+		/**
+		 * Add a structure to the current memory requirement.
+		 * Takes into account the structure alignment.
+		 * UObject are created in the transient package and returned.
+		 */
+		UE_API void Add(FStateTreeDataHandle DataHandle, FConstStructView Default);
+
+		/** Resets the data to empty. */
+		void UE_API Reset();
+
+		/** Returns number of bytes allocated for the array  */
+		[[nodiscard]] int32 GetAllocatedMemory() const
+		{
+			return MemoryRequirement.Size;
+		}
+
+	private:
+		struct FItem
+		{
+			FStateTreeDataView Instance;
+			FStateTreeDataHandle DataHandle;
+		};
+
+		FItem* Get(const FStateTreeDataHandle& Handle) const;
+
+		//~ Debug Tags
+		void AddTableDebugTag();
+		void AddStructDebugTag(int32 Index);
+		void TestDebugTags() const;
+
+	private:
+		/** Struct instances (Not transient, as we use FStateTreeInstanceData to store default values for instance data) */
+		void* Memory = nullptr;
+		FMemoryRequirement MemoryRequirement;
+		int32 NumberOfElements = 0;
+		bool bStructsHaveDestructor = false;
+	};
+}
+
+#undef UE_API
diff --git a/Engine/Plugins/Runtime/StateTree/Source/StateTreeModule/Public/StateTreeExecutionContext.h b/Engine/Plugins/Runtime/StateTree/Source/StateTreeModule/Public/StateTreeExecutionContext.h
--- a/Engine/Plugins/Runtime/StateTree/Source/StateTreeModule/Public/StateTreeExecutionContext.h
+++ b/Engine/Plugins/Runtime/StateTree/Source/StateTreeModule/Public/StateTreeExecutionContext.h
@@ -31,6 +31,11 @@ struct FStateTreeTransitionRequest;
 struct FStateTreeInstanceDebugId;
 struct FStateTreeWeakExecutionContext;
 
+namespace UE::StateTree::InstanceData
+{
+	struct FEvaluationScopeInstanceContainer;
+}
+
 namespace UE::StateTree::ExecutionContext
 {
 	UE_API bool MarkDelegateAsBroadcasted(FStateTreeDelegateDispatcher Dispatcher, const FStateTreeExecutionFrame& CurrentFrame, FStateTreeInstanceStorage& Storage);
@@ -1447,9 +1452,22 @@ PRAGMA_ENABLE_DEPRECATION_WARNINGS
 		FStateTreeIndex16 BaseIndex;
 	};
 	TArray<FCollectedExternalDataCache> CollectedExternalCache;
-
 	bool bActiveExternalDataCollected = false;
 
+	/** Holds the container. Multiple instances with the same state tree can occur in a recursive call. */
+	struct FEvaluationScopeDataCache
+	{
+		UE::StateTree::InstanceData::FEvaluationScopeInstanceContainer* Container;
+		TObjectKey<UStateTree> StateTree;
+	};
+
+	static constexpr int32 ExpectedEvaluationScopeCacheLength = 4;
+	/** Data view of the evaluation scope instance data. */
+	TArray<FEvaluationScopeDataCache, TInlineAllocator<ExpectedEvaluationScopeCacheLength>> EvaluationScopeInstanceCaches;
+
+	UE_API void PushEvaluationScopeInstanceContainer(UE::StateTree::InstanceData::FEvaluationScopeInstanceContainer& Container, const FStateTreeExecutionFrame& Frame);
+	UE_API void PopEvaluationScopeInstanceContainer(UE::StateTree::InstanceData::FEvaluationScopeInstanceContainer& Container);
+
 #if WITH_EDITORONLY_DATA
 	UE_DEPRECATED(5.6, "CompletedStateRunStatus is not used.")
 	/** Run status set in FinishTask and that used after the task ticks. */
diff --git a/Engine/Plugins/Runtime/StateTree/Source/StateTreeModule/Public/StateTreeInstanceContainer.h b/Engine/Plugins/Runtime/StateTree/Source/StateTreeModule/Public/StateTreeInstanceContainer.h
new file mode 100644
--- /dev/null
+++ b/Engine/Plugins/Runtime/StateTree/Source/StateTreeModule/Public/StateTreeInstanceContainer.h
@@ -0,0 +1,113 @@
+// Copyright Epic Games, Inc. All Rights Reserved.
+
+#pragma once
+
+#include "StructUtils/InstancedStructContainer.h"
+#include "StateTreeTypes.h"
+#include "StateTreeInstanceContainer.generated.h"
+
+#define UE_API STATETREEMODULE_API
+
+/** Wrapper class to store an object amongst the structs. */
+USTRUCT()
+struct FStateTreeInstanceObjectWrapper
+{
+	GENERATED_BODY()
+
+	FStateTreeInstanceObjectWrapper() = default;
+	FStateTreeInstanceObjectWrapper(UObject* Object)
+		: InstanceObject(Object)
+	{}
+	
+	UPROPERTY()
+	TObjectPtr<UObject> InstanceObject = nullptr;
+};
+
+namespace UE::StateTree::InstanceData
+{
+	/** Container of instance data. */
+	USTRUCT()
+	struct FInstanceContainer
+	{
+		GENERATED_BODY()
+
+		FInstanceContainer() = default;
+
+	public:
+		struct FAddArgs
+		{
+			/** Duplicate the object contained by object wrapper. */
+			bool bDuplicateWrappedObject = true;
+		};
+
+		/** Initializes the array with specified items. */
+		UE_API void Init(TNotNull<UObject*> InOwner, TConstArrayView<FInstancedStruct> InStructs, FAddArgs Args);
+		/** Initializes the array with specified items. */
+		UE_API void Init(TNotNull<UObject*> InOwner, TConstArrayView<FConstStructView> InStructs, FAddArgs Args);
+
+		/** @return number of items in the storage. */
+		[[nodiscard]] int32 Num() const
+		{
+			return InstanceStructs.Num();
+		}
+
+		/** @return true if the index can be used to get data. */
+		[[nodiscard]] bool IsValidIndex(const int32 Index) const
+		{
+			return InstanceStructs.IsValidIndex(Index);
+		}
+
+		/** @return true if item at the specified index is object type. */
+		[[nodiscard]] bool IsObject(const int32 Index) const
+		{
+			return InstanceStructs[Index].GetScriptStruct() == TBaseStructure<FStateTreeInstanceObjectWrapper>::Get();
+		}
+
+		/** @return specified item as struct. */
+		[[nodiscard]] FConstStructView GetStruct(const int32 Index) const
+		{
+			return InstanceStructs[Index];
+		}
+
+		/** @return specified item as mutable struct. */
+		[[nodiscard]] FStructView GetMutableStruct(const int32 Index)
+		{
+			return InstanceStructs[Index];
+		}
+
+		/** @return specified item as object, will check() if the item is not an object. */
+		[[nodiscard]] const UObject* GetObject(const int32 Index) const
+		{
+			const FStateTreeInstanceObjectWrapper& Wrapper = InstanceStructs[Index].Get<const FStateTreeInstanceObjectWrapper>();
+			return Wrapper.InstanceObject;
+		}
+
+		/** @return specified item as mutable Object, will check() if the item is not an object. */
+		[[nodiscard]] UObject* GetMutableObject(const int32 Index) const
+		{
+			const FStateTreeInstanceObjectWrapper& Wrapper = InstanceStructs[Index].Get<const FStateTreeInstanceObjectWrapper>();
+			return Wrapper.InstanceObject;
+		}
+
+		/** @return true if all instances are valid. */
+		[[nodiscard]] UE_API bool AreAllInstancesValid() const;
+
+		/** Resets the data to empty. */
+		void Reset()
+		{
+			InstanceStructs.Reset();
+		}
+
+		/** Returns number of bytes allocated for the array  */
+		[[nodiscard]] int32 GetAllocatedMemory() const
+		{
+			return InstanceStructs.GetAllocatedMemory();
+		}
+
+	private:
+		UPROPERTY()
+		FInstancedStructContainer InstanceStructs;
+	};
+} // namespace UE::StateTree::InstanceData
+
+#undef UE_API
diff --git a/Engine/Plugins/Runtime/StateTree/Source/StateTreeModule/Public/StateTreeInstanceData.h b/Engine/Plugins/Runtime/StateTree/Source/StateTreeModule/Public/StateTreeInstanceData.h
--- a/Engine/Plugins/Runtime/StateTree/Source/StateTreeModule/Public/StateTreeInstanceData.h
+++ b/Engine/Plugins/Runtime/StateTree/Source/StateTreeModule/Public/StateTreeInstanceData.h
@@ -5,6 +5,7 @@
 #include "StructUtils/InstancedStructContainer.h"
 #include "Debugger/StateTreeRuntimeValidation.h"
 #include "StateTreeEvents.h"
+#include "StateTreeInstanceContainer.h"
 #include "StateTreeTypes.h"
 #include "StateTreeExecutionTypes.h"
 #include "Misc/MTAccessDetector.h"
@@ -44,21 +45,6 @@ namespace UE::StateTree::InstanceData
 		const FStateTreeDataHandle& Handle);
 }
 
-/** Wrapper class to store an object amongst the structs. */
-USTRUCT()
-struct FStateTreeInstanceObjectWrapper
-{
-	GENERATED_BODY()
-
-	FStateTreeInstanceObjectWrapper() = default;
-	FStateTreeInstanceObjectWrapper(UObject* Object)
-		: InstanceObject(Object)
-	{}
-	
-	UPROPERTY()
-	TObjectPtr<UObject> InstanceObject = nullptr;
-};
-
 /**
  * Holds temporary instance data created during state selection.
  * The data is identified by Frame (StateTree + RootState) and DataHandle.
diff --git a/Engine/Plugins/Runtime/StateTree/Source/StateTreeModule/Public/StateTreeTypes.h b/Engine/Plugins/Runtime/StateTree/Source/StateTreeModule/Public/StateTreeTypes.h
--- a/Engine/Plugins/Runtime/StateTree/Source/StateTreeModule/Public/StateTreeTypes.h
+++ b/Engine/Plugins/Runtime/StateTree/Source/StateTreeModule/Public/StateTreeTypes.h
@@ -298,12 +298,18 @@ enum class EStateTreeDataSourceType : uint8
 	/** Active State Tasks */
 	ActiveInstanceDataObject,
 
-	/** Conditions and Considerations */
+	/** Conditions, considerations and function bindings. */
 	SharedInstanceData,
 
-	/** Conditions and Considerations */
+	/** Conditions, considerations and function bindings*/
 	SharedInstanceDataObject,
 
+	/** Temporary data constructor, used and destroyed immediately. Used by conditions, considerations and function bindings.*/
+	EvaluationScopeInstanceData,
+
+	/** Temporary data constructor, used and destroyed immediately. Used by conditions, considerations and function bindings.*/
+	EvaluationScopeInstanceDataObject,
+
 	/** Context Data, Tree Parameters */
 	ContextData,
 
@@ -433,27 +439,9 @@ struct FStateTreeDataHandle
 		return StateHandle;
 	}
 
-	bool IsObjectSource() const
-	{
-		return Source == EStateTreeDataSourceType::GlobalInstanceDataObject
-			|| Source == EStateTreeDataSourceType::ActiveInstanceDataObject
-			|| Source == EStateTreeDataSourceType::SharedInstanceDataObject;
-	}
+	STATETREEMODULE_API bool IsObjectSource() const;
 
-	FStateTreeDataHandle ToObjectSource() const
-	{
-		switch (Source)
-		{
-		case EStateTreeDataSourceType::GlobalInstanceData:
-			return FStateTreeDataHandle(EStateTreeDataSourceType::GlobalInstanceDataObject, Index, StateHandle);
-		case EStateTreeDataSourceType::ActiveInstanceData:
-			return FStateTreeDataHandle(EStateTreeDataSourceType::ActiveInstanceDataObject, Index, StateHandle);
-		case EStateTreeDataSourceType::SharedInstanceData:
-			return FStateTreeDataHandle(EStateTreeDataSourceType::SharedInstanceDataObject, Index, StateHandle);
-		default:
-			return *this;
-		}
-	}
+	STATETREEMODULE_API FStateTreeDataHandle ToObjectSource() const;
 
 	STATETREEMODULE_API FString Describe() const;
 
diff --git a/Engine/Plugins/Runtime/StateTree/Source/StateTreeTestSuite/Private/StateTreeTestTypes.h b/Engine/Plugins/Runtime/StateTree/Source/StateTreeTestSuite/Private/StateTreeTestTypes.h
--- a/Engine/Plugins/Runtime/StateTree/Source/StateTreeTestSuite/Private/StateTreeTestTypes.h
+++ b/Engine/Plugins/Runtime/StateTree/Source/StateTreeTestSuite/Private/StateTreeTestTypes.h
@@ -649,6 +649,7 @@ struct FStateTreeTestConditionInstanceData
 
 	static std::atomic<int32> GlobalCounter;
 };
+STATETREE_POD_INSTANCEDATA(FStateTreeTestConditionInstanceData);
 
 USTRUCT(meta = (Hidden))
 struct FStateTreeTestCondition : public FStateTreeConditionCommonBase
@@ -866,6 +867,7 @@ struct FTestPropertyFunction_InstanceData
 	UPROPERTY()
 	int32 Result = 0;
 };
+STATETREE_POD_INSTANCEDATA(FTestPropertyFunction_InstanceData);
 
 USTRUCT()
 struct FTestPropertyFunction : public FStateTreePropertyFunctionBase
@@ -929,6 +931,7 @@ struct FStateTreeTestBooleanConditionInstanceData
 	UPROPERTY(EditAnywhere, Category = Parameters)
 	bool bSuccess = true;
 };
+STATETREE_POD_INSTANCEDATA(FStateTreeTestBooleanConditionInstanceData);
 
 USTRUCT(meta = (Hidden))
 struct FStateTreeTestBooleanCondition : public FStateTreeConditionCommonBase
