[StateTree] backing out part of CL 44451677 while investigating a build failure. PropertyFunction uses SharedInstanceData.

#rnx

diff --git a/Engine/Plugins/Runtime/StateTree/Source/StateTreeEditorModule/Private/StateTreeCompiler.cpp b/Engine/Plugins/Runtime/StateTree/Source/StateTreeEditorModule/Private/StateTreeCompiler.cpp
--- a/Engine/Plugins/Runtime/StateTree/Source/StateTreeEditorModule/Private/StateTreeCompiler.cpp
+++ b/Engine/Plugins/Runtime/StateTree/Source/StateTreeEditorModule/Private/StateTreeCompiler.cpp
@@ -1704,8 +1704,8 @@ bool FStateTreeCompiler::CreatePropertyFunction(const FStateTreeEditorNode& Func
 	StructDesc.Name = FuncEditorNode.GetName();
 	StructDesc.DataSource = EStateTreeBindableStructSource::PropertyFunction;
 
-	const FStateTreeDataHandle DataHandle = FStateTreeDataHandle(EStateTreeDataSourceType::EvaluationScopeInstanceData, EvaluationScopeStructs.Num());
-	FInstancedStruct* CreatedNode = CreateNode(nullptr, FuncEditorNode, StructDesc, DataHandle, EvaluationScopeStructs);
+	TArray<FInstancedStruct>& InstancedStructContainer = SharedInstanceStructs;
+	FInstancedStruct* CreatedNode = CreateNodeWithSharedInstanceData(nullptr, FuncEditorNode, StructDesc);
 	if (CreatedNode == nullptr)
 	{
 		return false;
@@ -1721,9 +1721,9 @@ bool FStateTreeCompiler::CreatePropertyFunction(const FStateTreeEditorNode& Func
 	check(BindableStruct);
 
 	FStateTreeDataView InstanceView;
-	check(EvaluationScopeStructs.IsValidIndex(Function->InstanceTemplateIndex.Get()));
+	check(InstancedStructContainer.IsValidIndex(Function->InstanceTemplateIndex.Get()));
 
-	FInstancedStruct& Instance = EvaluationScopeStructs[Function->InstanceTemplateIndex.Get()];
+	FInstancedStruct& Instance = InstancedStructContainer[Function->InstanceTemplateIndex.Get()];
 	if (FStateTreeInstanceObjectWrapper* ObjectWrapper = Instance.GetMutablePtr<FStateTreeInstanceObjectWrapper>())
 	{
 		check(FuncEditorNode.InstanceObject->GetClass() == ObjectWrapper->InstanceObject->GetClass());
diff --git a/Engine/Plugins/Runtime/StateTree/Source/StateTreeModule/Private/StateTree.cpp b/Engine/Plugins/Runtime/StateTree/Source/StateTreeModule/Private/StateTree.cpp
--- a/Engine/Plugins/Runtime/StateTree/Source/StateTreeModule/Private/StateTree.cpp
+++ b/Engine/Plugins/Runtime/StateTree/Source/StateTreeModule/Private/StateTree.cpp
@@ -1076,8 +1076,12 @@ bool UStateTree::PatchBindings()
 			for (int32 FuncIndex = FuncsBegin; FuncIndex < FuncsEnd; ++FuncIndex)
 			{
 				const FStateTreePropertyFunctionBase& Func = Nodes[FuncIndex].Get<const FStateTreePropertyFunctionBase>();
-				const FConstStructView DefaultInstanceView = DefaultEvaluationScopeInstanceData.GetStruct(Func.InstanceTemplateIndex.Get());
-				CopyBatchFunctionRequirement.Add(DefaultInstanceView.GetScriptStruct());
+				if (Func.InstanceDataHandle.GetSource() == EStateTreeDataSourceType::EvaluationScopeInstanceData
+					|| Func.InstanceDataHandle.GetSource() == EStateTreeDataSourceType::EvaluationScopeInstanceDataObject)
+				{
+					const FConstStructView DefaultInstanceView = DefaultEvaluationScopeInstanceData.GetStruct(Func.InstanceTemplateIndex.Get());
+					CopyBatchFunctionRequirement.Add(DefaultInstanceView.GetScriptStruct());
+				}
 			}
 		}
 
diff --git a/Engine/Plugins/Runtime/StateTree/Source/StateTreeModule/Private/StateTreeExecutionContext.cpp b/Engine/Plugins/Runtime/StateTree/Source/StateTreeModule/Private/StateTreeExecutionContext.cpp
--- a/Engine/Plugins/Runtime/StateTree/Source/StateTreeModule/Private/StateTreeExecutionContext.cpp
+++ b/Engine/Plugins/Runtime/StateTree/Source/StateTreeModule/Private/StateTreeExecutionContext.cpp
@@ -3017,7 +3017,11 @@ namespace UE::StateTree::ExecutionContext::Private
 		for (int32 FuncIndex = FuncBegin; FuncIndex != FuncEnd; ++FuncIndex)
 		{
 			const FStateTreePropertyFunctionBase& Func = StateTree->GetNodes()[FuncIndex].Get<const FStateTreePropertyFunctionBase>();
-			Container.Add(Func.InstanceDataHandle, StateTree->GetDefaultEvaluationScopeInstanceData().GetStruct(Func.InstanceTemplateIndex.Get()));
+			if (Func.InstanceDataHandle.GetSource() == EStateTreeDataSourceType::EvaluationScopeInstanceData
+				|| Func.InstanceDataHandle.GetSource() == EStateTreeDataSourceType::EvaluationScopeInstanceDataObject)
+			{
+				Container.Add(Func.InstanceDataHandle, StateTree->GetDefaultEvaluationScopeInstanceData().GetStruct(Func.InstanceTemplateIndex.Get()));
+			}
 		}
 	}
 }
@@ -3038,7 +3042,7 @@ bool FStateTreeExecutionContext::CopyBatchInternal(const FStateTreeExecutionFram
 		check(Batch.PropertyFunctionsBegin.IsValid() && Batch.PropertyFunctionsEnd.IsValid());
 
 		const FEvaluationScopeInstanceContainer::FMemoryRequirement& MemoryRequirement = CurrentFrame.StateTree->PropertyFunctionEvaluationScopeMemoryRequirements[BindingsBatch.Get()];
-		if (ensure(MemoryRequirement.Size > 0))
+		if (MemoryRequirement.Size > 0)
 		{
 			void* FunctionBindingMemory = FMemory_Alloca_Aligned(MemoryRequirement.Size, MemoryRequirement.Alignment);
 			EvaluationScopeContainer = FEvaluationScopeInstanceContainer(FunctionBindingMemory, MemoryRequirement);
