[Backout] - CL44450181
#rnx
Original CL Desc
-----------------------------------------------------------------
[StateTree] Add EvaluationScopeInstanceData. A temporary instance data for nodes that are used on the stack. The instance data is created on the stack and destroyed right away after use. PropertyFunction is using EvaluationScopeInstanceData. Condition and Consideration will also use it in the future.
Add EvaluationScopeInstanceContainer, a container to help construct and destroy the evaluation scope instance data.
Add InstanceContainer, a wrapper around FInstancedStructContainer. Use to store the default value instead of using the FStateTreeInstanceStorage, which stores more information. In the future, the storage should be refactored to the InstanceContainer.

diff --git a/Engine/Plugins/Runtime/StateTree/Source/StateTreeEditorModule/Private/StateTreeCompiler.cpp b/Engine/Plugins/Runtime/StateTree/Source/StateTreeEditorModule/Private/StateTreeCompiler.cpp
--- a/Engine/Plugins/Runtime/StateTree/Source/StateTreeEditorModule/Private/StateTreeCompiler.cpp
+++ b/Engine/Plugins/Runtime/StateTree/Source/StateTreeEditorModule/Private/StateTreeCompiler.cpp
@@ -446,7 +446,6 @@ bool FStateTreeCompiler::Compile(UStateTree& InStateTree)
 	StateTree->Nodes = Nodes;
 	StateTree->DefaultInstanceData.Init(*StateTree, InstanceStructs, FStateTreeInstanceData::FAddArgs{ .bDuplicateWrappedObject = false });
 	StateTree->SharedInstanceData.Init(*StateTree, SharedInstanceStructs, FStateTreeInstanceData::FAddArgs{ .bDuplicateWrappedObject = false });
-	StateTree->DefaultEvaluationScopeInstanceData.Init(StateTree, EvaluationScopeStructs, UE::StateTree::InstanceData::FInstanceContainer::FAddArgs{ .bDuplicateWrappedObject = false });
 
 	// Store the new compiled dispatchers.
 	EditorData->CompiledDispatchers = BindingsCompiler.GetCompiledDelegateDispatchers();
@@ -1704,15 +1703,8 @@ bool FStateTreeCompiler::CreatePropertyFunction(const FStateTreeEditorNode& Func
 	StructDesc.Name = FuncEditorNode.GetName();
 	StructDesc.DataSource = EStateTreeBindableStructSource::PropertyFunction;
 
-	TArray<FInstancedStruct>& InstancedStructContainer = SharedInstanceStructs;
-	FInstancedStruct* CreatedNode = CreateNodeWithSharedInstanceData(nullptr, FuncEditorNode, StructDesc);
-	if (CreatedNode == nullptr)
-	{
-		return false;
-	}
-
-	FStateTreeNodeBase* Function = CreatedNode->GetMutablePtr<FStateTreeNodeBase>();
-	if (Function == nullptr)
+	FStateTreeNodeBase* Node = CreateNodeWithSharedInstanceData(nullptr, FuncEditorNode, StructDesc);
+	if(Node == nullptr)
 	{
 		return false;
 	}
@@ -1721,9 +1713,9 @@ bool FStateTreeCompiler::CreatePropertyFunction(const FStateTreeEditorNode& Func
 	check(BindableStruct);
 
 	FStateTreeDataView InstanceView;
-	check(InstancedStructContainer.IsValidIndex(Function->InstanceTemplateIndex.Get()));
+	check(SharedInstanceStructs.IsValidIndex(Node->InstanceTemplateIndex.Get()));
 
-	FInstancedStruct& Instance = InstancedStructContainer[Function->InstanceTemplateIndex.Get()];
+	FInstancedStruct& Instance = SharedInstanceStructs[Node->InstanceTemplateIndex.Get()];
 	if (FStateTreeInstanceObjectWrapper* ObjectWrapper = Instance.GetMutablePtr<FStateTreeInstanceObjectWrapper>())
 	{
 		check(FuncEditorNode.InstanceObject->GetClass() == ObjectWrapper->InstanceObject->GetClass());
@@ -1735,16 +1727,19 @@ bool FStateTreeCompiler::CreatePropertyFunction(const FStateTreeEditorNode& Func
 		InstanceView = FStateTreeDataView(Instance);
 	}
 
-	return CreateBindingsForStruct(*BindableStruct, InstanceView, FStateTreeIndex16::Invalid, FStateTreeIndex16::Invalid, Function->BindingsBatch);
+	return CreateBindingsForStruct(*BindableStruct, InstanceView, FStateTreeIndex16::Invalid, FStateTreeIndex16::Invalid, Node->BindingsBatch);
 }
 
-FInstancedStruct* FStateTreeCompiler::CreateNode(UStateTreeState* State, const FStateTreeEditorNode& EditorNode, FStateTreeBindableStructDesc& StructDesc, const FStateTreeDataHandle DataHandle, TArray<FInstancedStruct>& InstancedStructContainer)
+template<class T>
+T* FStateTreeCompiler::CreateNodeWithSharedInstanceData(UStateTreeState* State, const FStateTreeEditorNode& EditorNode, FStateTreeBindableStructDesc& StructDesc)
 {
 	if (!EditorNode.Node.IsValid())
 	{
 		return nullptr;
 	}
 
+	check(EditorNode.Node.GetScriptStruct()->IsChildOf<T>());
+
 	// Check that item has valid instance initialized.
 	if (!EditorNode.Instance.IsValid() && EditorNode.InstanceObject == nullptr)
 	{
@@ -1760,41 +1755,16 @@ FInstancedStruct* FStateTreeCompiler::CreateNode(UStateTreeState* State, const F
 
 	FStateTreeNodeBase& Node = RawNode.GetMutable<FStateTreeNodeBase>();
 
-	TOptional<FStateTreeDataView> InstanceDataView = CreateNodeInstanceData(EditorNode, Node, StructDesc, DataHandle, InstancedStructContainer);
-	if (!InstanceDataView.IsSet())
-	{
-		return nullptr;
-	}
-
-	if (!CompileAndValidateNode(State, StructDesc, RawNode, InstanceDataView.GetValue()))
-	{
-		return nullptr;
-	}
-
-	// Mark the struct as binding source.
-	BindingsCompiler.AddSourceStruct(StructDesc);
-
-	return &RawNode;
-}
-
-FInstancedStruct* FStateTreeCompiler::CreateNodeWithSharedInstanceData(UStateTreeState* State, const FStateTreeEditorNode& EditorNode, FStateTreeBindableStructDesc& StructDesc)
-{
-	const FStateTreeDataHandle DataHandle = FStateTreeDataHandle(EStateTreeDataSourceType::SharedInstanceData, SharedInstanceStructs.Num());
-	return CreateNode(State, EditorNode, StructDesc, DataHandle, SharedInstanceStructs);
-}
-
-TOptional<FStateTreeDataView> FStateTreeCompiler::CreateNodeInstanceData(const FStateTreeEditorNode& EditorNode, FStateTreeNodeBase& Node, FStateTreeBindableStructDesc& StructDesc, const FStateTreeDataHandle DataHandle, TArray<FInstancedStruct>& InstancedStructContainer)
-{
-	FStateTreeDataView InstanceDataView;
-
-	// Update Node name as description for runtime.
+	// Update node name as description for runtime.
 	Node.Name = EditorNode.GetName();
 
+	FStateTreeDataView InstanceDataView;
+	
 	if (EditorNode.Instance.IsValid())
 	{
-		// Struct Instance
-		const int32 InstanceIndex = InstancedStructContainer.Add(EditorNode.Instance);
-		InstantiateStructSubobjects(InstancedStructContainer[InstanceIndex]);
+		// Struct instance
+		const int32 InstanceIndex = SharedInstanceStructs.Add(EditorNode.Instance);
+		InstantiateStructSubobjects(SharedInstanceStructs[InstanceIndex]);
 
 		// Create binding source struct descriptor.
 		StructDesc.Struct = EditorNode.Instance.GetScriptStruct();
@@ -1802,38 +1772,47 @@ TOptional<FStateTreeDataView> FStateTreeCompiler::CreateNodeInstanceData(const F
 		if (const auto Validation = UE::StateTree::Compiler::IsValidIndex16(InstanceIndex); Validation.DidFail())
 		{
 			Validation.Log(Log, TEXT("InstanceIndex"), StructDesc);
-			return {};
+			return nullptr;
 		}
 		Node.InstanceTemplateIndex = FStateTreeIndex16(InstanceIndex);
-		Node.InstanceDataHandle = DataHandle;
-		InstanceDataView = FStateTreeDataView(InstancedStructContainer[InstanceIndex]);
+		Node.InstanceDataHandle = FStateTreeDataHandle(EStateTreeDataSourceType::SharedInstanceData, InstanceIndex);
+		InstanceDataView = FStateTreeDataView(SharedInstanceStructs[InstanceIndex]);
 	}
 	else
 	{
 		// Object Instance
 		check(EditorNode.InstanceObject != nullptr);
-		UObject* Instance = UE::StateTree::Compiler::DuplicateInstanceObject(Log, StructDesc, EditorNode.ID, EditorNode.InstanceObject, StateTree);
 
+		UObject* Instance = UE::StateTree::Compiler::DuplicateInstanceObject(Log, StructDesc, EditorNode.ID, EditorNode.InstanceObject, StateTree);
 		FInstancedStruct Wrapper;
 		Wrapper.InitializeAs<FStateTreeInstanceObjectWrapper>(Instance);
-		const int32 InstanceIndex = InstancedStructContainer.Add(MoveTemp(Wrapper));
-
+		const int32 InstanceIndex = SharedInstanceStructs.Add(MoveTemp(Wrapper));
+		
 		// Create binding source struct descriptor.
 		StructDesc.Struct = Instance->GetClass();
 
 		if (const auto Validation = UE::StateTree::Compiler::IsValidIndex16(InstanceIndex); Validation.DidFail())
 		{
 			Validation.Log(Log, TEXT("InstanceIndex"), StructDesc);
-			return {};
+			return nullptr;
 		}
 		Node.InstanceTemplateIndex = FStateTreeIndex16(InstanceIndex);
-		Node.InstanceDataHandle = DataHandle.ToObjectSource();
+		Node.InstanceDataHandle = FStateTreeDataHandle(EStateTreeDataSourceType::SharedInstanceDataObject, InstanceIndex);
 		InstanceDataView = FStateTreeDataView(Instance);
 	}
 
 	StructDesc.DataHandle = Node.InstanceDataHandle;
+	
+	if (!CompileAndValidateNode(State, StructDesc, RawNode, InstanceDataView))
+	{
+		return nullptr;
+	}
 
-	return TOptional<FStateTreeDataView>(InstanceDataView);
+	// Mark the struct as binding source.
+	BindingsCompiler.AddSourceStruct(StructDesc);
+	
+	check(RawNode.GetScriptStruct()->IsChildOf<T>());
+	return RawNode.GetMutablePtr<T>();
 }
 
 bool FStateTreeCompiler::ResolveTransitionStateAndFallback(const UStateTreeState* SourceState, const FStateTreeStateLink& Link, FStateTreeStateHandle& OutTransitionHandle, EStateTreeSelectionFallback& OutFallback) const 
@@ -1926,24 +1905,19 @@ bool FStateTreeCompiler::CreateCondition(UStateTreeState& State, const FString&
 	StructDesc.Name = CondNode.GetName();
 	StructDesc.DataSource = EStateTreeBindableStructSource::Condition;
 
-	FInstancedStruct* CreatedNode = CreateNodeWithSharedInstanceData(&State, CondNode, StructDesc);
-	if (CreatedNode)
+	if (FStateTreeConditionBase* Cond = CreateNodeWithSharedInstanceData<FStateTreeConditionBase>(&State, CondNode, StructDesc))
 	{
-		FStateTreeConditionBase* Cond = CreatedNode->GetMutablePtr<FStateTreeConditionBase>();
-		if (ensure(Cond))
+		if (Cond->EvaluationMode == EStateTreeConditionEvaluationMode::ForcedFalse
+			|| Cond->EvaluationMode == EStateTreeConditionEvaluationMode::ForcedTrue)
 		{
-			if (Cond->EvaluationMode == EStateTreeConditionEvaluationMode::ForcedFalse
-				|| Cond->EvaluationMode == EStateTreeConditionEvaluationMode::ForcedTrue)
-			{
-				Log.Reportf(EMessageSeverity::Info, StructDesc,
-					TEXT("The condition result will always be %s."),
-					Cond->EvaluationMode == EStateTreeConditionEvaluationMode::ForcedTrue ? TEXT("True") : TEXT("False"));
-			}
-
-			Cond->Operand = Operand;
-			Cond->DeltaIndent = DeltaIndent;
-			return true;
+			Log.Reportf(EMessageSeverity::Info, StructDesc, 
+						TEXT("The condition result will always be %s."),
+						Cond->EvaluationMode == EStateTreeConditionEvaluationMode::ForcedTrue ? TEXT("True") : TEXT("False"));
 		}
+
+		Cond->Operand = Operand;
+		Cond->DeltaIndent = DeltaIndent;
+		return true;
 	}
 
 	return false;
@@ -1993,16 +1967,11 @@ bool FStateTreeCompiler::CreateConsideration(UStateTreeState& State, const FStri
 	StructDesc.Name = ConsiderationNode.GetName();
 	StructDesc.DataSource = EStateTreeBindableStructSource::Consideration;
 
-	FInstancedStruct* CreatedNode = CreateNodeWithSharedInstanceData(&State, ConsiderationNode, StructDesc);
-	if (CreatedNode)
+	if (FStateTreeConsiderationBase* Consideration = CreateNodeWithSharedInstanceData<FStateTreeConsiderationBase>(&State, ConsiderationNode, StructDesc))
 	{
-		FStateTreeConsiderationBase * Consideration = CreatedNode->GetMutablePtr<FStateTreeConsiderationBase>();
-		if (ensure(Consideration))
-		{
-			Consideration->Operand = Operand;
-			Consideration->DeltaIndent = DeltaIndent;
-			return true;
-		}
+		Consideration->Operand = Operand;
+		Consideration->DeltaIndent = DeltaIndent;
+		return true;
 	}
 
 	return false;
@@ -2105,18 +2074,59 @@ bool FStateTreeCompiler::CreateTask(UStateTreeState* State, const FStateTreeEdit
 	InstantiateStructSubobjects(Node);
 	
 	FStateTreeTaskBase& Task = Node.GetMutable<FStateTreeTaskBase>();
-	TOptional<FStateTreeDataView> InstanceDataView = CreateNodeInstanceData(TaskNode, Task, StructDesc, TaskDataHandle, InstanceStructs);
-	if (!InstanceDataView.IsSet())
+	FStateTreeDataView InstanceDataView;
+
+	// Update task name as description for runtime.
+	Task.Name = TaskNode.GetName();
+	
+	if (TaskNode.Instance.IsValid())
 	{
-		return false;
+		// Struct Instance
+		const int32 InstanceIndex = InstanceStructs.Add(TaskNode.Instance);
+		InstantiateStructSubobjects(InstanceStructs[InstanceIndex]);
+
+		// Create binding source struct descriptor.
+		StructDesc.Struct = TaskNode.Instance.GetScriptStruct();
+
+		if (const auto Validation = UE::StateTree::Compiler::IsValidIndex16(InstanceIndex); Validation.DidFail())
+		{
+			Validation.Log(Log, TEXT("InstanceIndex"), StructDesc);
+			return false;
+		}
+		Task.InstanceTemplateIndex = FStateTreeIndex16(InstanceIndex);
+		Task.InstanceDataHandle = TaskDataHandle;
+		InstanceDataView = FStateTreeDataView(InstanceStructs[InstanceIndex]);
 	}
+	else
+	{
+		// Object Instance
+		check(TaskNode.InstanceObject != nullptr);
+		UObject* Instance = UE::StateTree::Compiler::DuplicateInstanceObject(Log, StructDesc, TaskNode.ID, TaskNode.InstanceObject, StateTree);
+		FInstancedStruct Wrapper;
+		Wrapper.InitializeAs<FStateTreeInstanceObjectWrapper>(Instance);
+		const int32 InstanceIndex = InstanceStructs.Add(MoveTemp(Wrapper));
+
+		// Create binding source struct descriptor.
+		StructDesc.Struct = Instance->GetClass();
+
+		if (const auto Validation = UE::StateTree::Compiler::IsValidIndex16(InstanceIndex); Validation.DidFail())
+		{
+			Validation.Log(Log, TEXT("InstanceIndex"), StructDesc);
+			return false;
+		}
+		Task.InstanceTemplateIndex = FStateTreeIndex16(InstanceIndex);
+		Task.InstanceDataHandle = TaskDataHandle.ToObjectSource();
+		InstanceDataView = FStateTreeDataView(Instance);
+	}
+
+	StructDesc.DataHandle = Task.InstanceDataHandle;
 
 	if (!Task.bTaskEnabled)
 	{
 		Log.Reportf(EMessageSeverity::Info, StructDesc, TEXT("Task is disabled and will have no effect."));
 	}
 
-	if (!CompileAndValidateNode(State, StructDesc, Node,  InstanceDataView.GetValue()))
+	if (!CompileAndValidateNode(State, StructDesc, Node,  InstanceDataView))
 	{
 		return false;
 	}
@@ -2156,14 +2166,55 @@ bool FStateTreeCompiler::CreateEvaluator(const FStateTreeEditorNode& EvalNode, c
 	InstantiateStructSubobjects(Node);
 	
 	FStateTreeEvaluatorBase& Eval = Node.GetMutable<FStateTreeEvaluatorBase>();
-	TOptional<FStateTreeDataView> InstanceDataView = CreateNodeInstanceData(EvalNode, Eval, StructDesc, EvalDataHandle, InstanceStructs);
-	if (!InstanceDataView.IsSet())
+	FStateTreeDataView InstanceDataView;
+
+	// Update eval name as description for runtime.
+	Eval.Name = EvalNode.GetName();
+
+	if (EvalNode.Instance.IsValid())
 	{
-		return false;
+		// Struct Instance
+		const int32 InstanceIndex = InstanceStructs.Add(EvalNode.Instance);
+		InstantiateStructSubobjects(InstanceStructs[InstanceIndex]);
+
+		// Create binding source struct descriptor.
+		StructDesc.Struct = EvalNode.Instance.GetScriptStruct();
+
+		if (const auto Validation = UE::StateTree::Compiler::IsValidIndex16(InstanceIndex); Validation.DidFail())
+		{
+			Validation.Log(Log, TEXT("InstanceIndex"), StructDesc);
+			return false;
+		}
+		Eval.InstanceTemplateIndex = FStateTreeIndex16(InstanceIndex);
+		Eval.InstanceDataHandle = EvalDataHandle;
+		InstanceDataView = FStateTreeDataView(InstanceStructs[InstanceIndex]);
+	}
+	else
+	{
+		// Object Instance
+		check(EvalNode.InstanceObject != nullptr);
+
+		UObject* Instance = UE::StateTree::Compiler::DuplicateInstanceObject(Log, StructDesc, EvalNode.ID, EvalNode.InstanceObject, StateTree);
+		FInstancedStruct Wrapper;
+		Wrapper.InitializeAs<FStateTreeInstanceObjectWrapper>(Instance);
+		const int32 InstanceIndex = InstanceStructs.Add(MoveTemp(Wrapper));
+		
+		// Create binding source struct descriptor.
+		StructDesc.Struct = Instance->GetClass();
+
+		if (const auto Validation = UE::StateTree::Compiler::IsValidIndex16(InstanceIndex); Validation.DidFail())
+		{
+			Validation.Log(Log, TEXT("InstanceIndex"), StructDesc);
+			return false;
+		}
+		Eval.InstanceTemplateIndex = FStateTreeIndex16(InstanceIndex);
+		Eval.InstanceDataHandle = EvalDataHandle.ToObjectSource();
+		InstanceDataView = FStateTreeDataView(Instance);
 	}
 
+	StructDesc.DataHandle = Eval.InstanceDataHandle;
 
-	if (!CompileAndValidateNode(nullptr, StructDesc, Node,  InstanceDataView.GetValue()))
+	if (!CompileAndValidateNode(nullptr, StructDesc, Node,  InstanceDataView))
 	{
 		return false;
 	}
diff --git a/Engine/Plugins/Runtime/StateTree/Source/StateTreeEditorModule/Public/StateTreeCompiler.h b/Engine/Plugins/Runtime/StateTree/Source/StateTreeEditorModule/Public/StateTreeCompiler.h
--- a/Engine/Plugins/Runtime/StateTree/Source/StateTreeEditorModule/Public/StateTreeCompiler.h
+++ b/Engine/Plugins/Runtime/StateTree/Source/StateTreeEditorModule/Public/StateTreeCompiler.h
@@ -86,9 +86,9 @@ private:
 	bool CreateTask(UStateTreeState* State, const FStateTreeEditorNode& TaskNode, const FStateTreeDataHandle TaskDataHandle);
 	bool CreateEvaluator(const FStateTreeEditorNode& EvalNode, const FStateTreeDataHandle EvalDataHandle);
 
-	FInstancedStruct* CreateNode(UStateTreeState* State, const FStateTreeEditorNode& EditorNode, FStateTreeBindableStructDesc& StructDesc, const FStateTreeDataHandle DataHandle, TArray<FInstancedStruct>& InstancedStructContainer);
-	FInstancedStruct* CreateNodeWithSharedInstanceData(UStateTreeState* State, const FStateTreeEditorNode& EditorNode, FStateTreeBindableStructDesc& StructDesc);
-	TOptional<FStateTreeDataView> CreateNodeInstanceData(const FStateTreeEditorNode& EditorNode, FStateTreeNodeBase& Node, FStateTreeBindableStructDesc& StructDesc, const FStateTreeDataHandle DataHandle, TArray<FInstancedStruct>& InstancedStructContainer);
+	template<class T = FStateTreeNodeBase>
+	T* CreateNodeWithSharedInstanceData(UStateTreeState* State, const FStateTreeEditorNode& EditorNode, FStateTreeBindableStructDesc& StructDesc);
+
 
 	struct FValidatedPathBindings
 	{
@@ -122,7 +122,6 @@ private:
 	TArray<FInstancedStruct> Nodes;
 	TArray<FInstancedStruct> InstanceStructs;
 	TArray<FInstancedStruct> SharedInstanceStructs;
-	TArray<FInstancedStruct> EvaluationScopeStructs;
 	
 	/** Cached result of MakeCompletionTasksMask for global tasks. Indicates where state tasks should start. */
 	int32 GlobalTaskEndBit = 0;
diff --git a/Engine/Plugins/Runtime/StateTree/Source/StateTreeModule/Private/StateTree.cpp b/Engine/Plugins/Runtime/StateTree/Source/StateTreeModule/Private/StateTree.cpp
--- a/Engine/Plugins/Runtime/StateTree/Source/StateTreeModule/Private/StateTree.cpp
+++ b/Engine/Plugins/Runtime/StateTree/Source/StateTreeModule/Private/StateTree.cpp
@@ -243,11 +243,9 @@ void UStateTree::ResetCompiled()
 	Transitions.Reset();
 	Nodes.Reset();
 	DefaultInstanceData.Reset();
-	DefaultEvaluationScopeInstanceData.Reset();
 	SharedInstanceData.Reset();
 	ContextDataDescs.Reset();
 	PropertyBindings.Reset();
-	PropertyFunctionEvaluationScopeMemoryRequirements.Reset();
 	Extensions.Reset();
 	Parameters.Reset();
 	ParameterDataType = EStateTreeParameterDataType::GlobalParameterData;
@@ -379,32 +377,17 @@ void UStateTree::PostLoad()
 		{
 			if (Node->InstanceTemplateIndex.IsValid())
 			{
-				const bool bIsUsingSharedInstanceData = Node->InstanceDataHandle.GetSource() == EStateTreeDataSourceType::SharedInstanceData
-					|| Node->InstanceDataHandle.GetSource() == EStateTreeDataSourceType::SharedInstanceDataObject;
-				const bool bIsUsingEvaluationScopeInstanceData = Node->InstanceDataHandle.GetSource() == EStateTreeDataSourceType::EvaluationScopeInstanceData
-					|| Node->InstanceDataHandle.GetSource() == EStateTreeDataSourceType::EvaluationScopeInstanceDataObject;
-				if (bIsUsingEvaluationScopeInstanceData)
+				const bool bIsUsingSharedInstanceData = NodeView.GetScriptStruct()->IsChildOf<FStateTreeConditionBase>()
+														|| NodeView.GetScriptStruct()->IsChildOf<FStateTreeConsiderationBase>()
+														|| NodeView.GetScriptStruct()->IsChildOf<FStateTreePropertyFunctionBase>();
+				FStateTreeInstanceData& SourceInstanceData = bIsUsingSharedInstanceData ? SharedInstanceData : DefaultInstanceData;
+				if (SourceInstanceData.IsObject(Node->InstanceTemplateIndex.Get()))
 				{
-					if (DefaultEvaluationScopeInstanceData.IsObject(Node->InstanceTemplateIndex.Get()))
-					{
-						Node->PostLoad(DefaultEvaluationScopeInstanceData.GetMutableObject(Node->InstanceTemplateIndex.Get()));
-					}
-					else
-					{
-						Node->PostLoad(DefaultEvaluationScopeInstanceData.GetMutableStruct(Node->InstanceTemplateIndex.Get()));
-					}
+					Node->PostLoad(SourceInstanceData.GetMutableObject(Node->InstanceTemplateIndex.Get()));
 				}
 				else
 				{
-					FStateTreeInstanceData& SourceInstanceData = bIsUsingSharedInstanceData ? SharedInstanceData : DefaultInstanceData;
-					if (SourceInstanceData.IsObject(Node->InstanceTemplateIndex.Get()))
-					{
-						Node->PostLoad(SourceInstanceData.GetMutableObject(Node->InstanceTemplateIndex.Get()));
-					}
-					else
-					{
-						Node->PostLoad(SourceInstanceData.GetMutableStruct(Node->InstanceTemplateIndex.Get()));
-					}
+					Node->PostLoad(SourceInstanceData.GetMutableStruct(Node->InstanceTemplateIndex.Get()));
 				}
 			}
 		}
@@ -504,34 +487,16 @@ bool UStateTree::ValidateInstanceData()
 		{
 			const UStruct* CurrentInstanceDataType = nullptr;
 			{
-				 const bool bIsUsingSharedInstanceData = Node->InstanceDataHandle.GetSource() == EStateTreeDataSourceType::SharedInstanceData
-					|| Node->InstanceDataHandle.GetSource() == EStateTreeDataSourceType::SharedInstanceDataObject;
-				 const bool bIsUsingEvaluationScopeInstanceData = Node->InstanceDataHandle.GetSource() == EStateTreeDataSourceType::EvaluationScopeInstanceData
-					 || Node->InstanceDataHandle.GetSource() == EStateTreeDataSourceType::EvaluationScopeInstanceDataObject;
-				if (bIsUsingEvaluationScopeInstanceData)
+				const bool bUseSharedInstanceData = NodeView.GetPtr<const FStateTreeConditionBase>() || NodeView.GetPtr<const FStateTreeConsiderationBase>() || NodeView.GetPtr<const FStateTreePropertyFunctionBase>();
+				const FStateTreeInstanceData& SourceInstanceData = bUseSharedInstanceData ? SharedInstanceData : DefaultInstanceData;
+				if (SourceInstanceData.IsObject(Node->InstanceTemplateIndex.Get()))
 				{
-					if (DefaultEvaluationScopeInstanceData.IsObject(Node->InstanceTemplateIndex.Get()))
-					{
-						const UObject* InstanceObject = DefaultEvaluationScopeInstanceData.GetObject(Node->InstanceTemplateIndex.Get());
-						CurrentInstanceDataType = InstanceObject ? InstanceObject->GetClass() : nullptr;
-					}
-					else
-					{
-						CurrentInstanceDataType = DefaultEvaluationScopeInstanceData.GetStruct(Node->InstanceTemplateIndex.Get()).GetScriptStruct();
-					}
+					const UObject* InstanceObject = SourceInstanceData.GetObject(Node->InstanceTemplateIndex.Get());
+					CurrentInstanceDataType = InstanceObject ? InstanceObject->GetClass() : nullptr;
 				}
 				else
 				{
-					const FStateTreeInstanceData& SourceInstanceData = bIsUsingSharedInstanceData ? SharedInstanceData : DefaultInstanceData;
-					if (SourceInstanceData.IsObject(Node->InstanceTemplateIndex.Get()))
-					{
-						const UObject* InstanceObject = SourceInstanceData.GetObject(Node->InstanceTemplateIndex.Get());
-						CurrentInstanceDataType = InstanceObject ? InstanceObject->GetClass() : nullptr;
-					}
-					else
-					{
-						CurrentInstanceDataType = SourceInstanceData.GetStruct(Node->InstanceTemplateIndex.Get()).GetScriptStruct();
-					}
+					CurrentInstanceDataType = SourceInstanceData.GetStruct(Node->InstanceTemplateIndex.Get()).GetScriptStruct();
 				}
 			}
 
@@ -706,12 +671,6 @@ bool UStateTree::Link()
 		UE_LOG(LogStateTree, Error, TEXT("%s: State Tree asset was not properly loaded (missing shared instance data). See log for loading failures, or recompile the StateTree asset."), *GetPathName());
 		return false;
 	}
-
-	if (!DefaultEvaluationScopeInstanceData.AreAllInstancesValid())
-	{
-		UE_LOG(LogStateTree, Error, TEXT("%s: State Tree asset was not properly loaded (missing evaluation scope instance data). See log for loading failures, or recompile the StateTree asset."), *GetPathName());
-		return false;
-	}
 	
 	if (!PatchBindings())
 	{
@@ -821,7 +780,6 @@ bool UStateTree::PatchBindings()
 			}
 		}
 	}
-
 	for (FPropertyBindingCopyInfoBatch& CopyBatch : CopyBatches)
 	{
 		if (const UClass* TargetClass = Cast<UClass>(CopyBatch.TargetStruct.Get().Struct))
@@ -965,27 +923,19 @@ bool UStateTree::PatchBindings()
 	{
 		const FStateTreeNodeBase& Node = NodeView.Get<const FStateTreeNodeBase>();
 
-		FStateTreeDataView NodeDataView;
-		if (Node.InstanceDataHandle.GetSource() == EStateTreeDataSourceType::SharedInstanceData
-			|| Node.InstanceDataHandle.GetSource() == EStateTreeDataSourceType::SharedInstanceDataObject)
-		{
-			NodeDataView = Node.InstanceDataHandle.IsObjectSource()
-				? FStateTreeDataView(SharedInstanceData.GetMutableObject(Node.InstanceTemplateIndex.Get()))
-				: FStateTreeDataView(SharedInstanceData.GetMutableStruct(Node.InstanceTemplateIndex.Get()));
-		}
-		else if (Node.InstanceDataHandle.GetSource() == EStateTreeDataSourceType::EvaluationScopeInstanceData
-			|| Node.InstanceDataHandle.GetSource() == EStateTreeDataSourceType::EvaluationScopeInstanceDataObject)
-		{
-			NodeDataView = Node.InstanceDataHandle.IsObjectSource()
-				? FStateTreeDataView(DefaultEvaluationScopeInstanceData.GetMutableObject(Node.InstanceTemplateIndex.Get()))
-				: FStateTreeDataView(DefaultEvaluationScopeInstanceData.GetMutableStruct(Node.InstanceTemplateIndex.Get()));
-		}
-		else
+		FStateTreeInstanceData* SourceInstanceData = &DefaultInstanceData;
+		if (NodeView.GetScriptStruct()->IsChildOf<FStateTreeConditionBase>()
+            || NodeView.GetScriptStruct()->IsChildOf<FStateTreeConsiderationBase>()
+            || NodeView.GetScriptStruct()->IsChildOf<FStateTreePropertyFunctionBase>())
 		{
-			NodeDataView = Node.InstanceDataHandle.IsObjectSource()
-				? FStateTreeDataView(DefaultInstanceData.GetMutableObject(Node.InstanceTemplateIndex.Get()))
-				: FStateTreeDataView(DefaultInstanceData.GetMutableStruct(Node.InstanceTemplateIndex.Get()));
+			// Conditions, Considerations, and PropertyFunctions are stored in shared instance data.
+			SourceInstanceData = &SharedInstanceData;
 		}
+
+		FStateTreeDataView NodeDataView = Node.InstanceDataHandle.IsObjectSource()
+						? FStateTreeDataView(SourceInstanceData->GetMutableObject(Node.InstanceTemplateIndex.Get()))
+						: FStateTreeDataView(SourceInstanceData->GetMutableStruct(Node.InstanceTemplateIndex.Get()));
+
 		DataViews.Add(Node.InstanceDataHandle, NodeDataView);
 
 		if (Node.BindingsBatch.IsValid())
@@ -1063,32 +1013,6 @@ bool UStateTree::PatchBindings()
 		}
 	}
 
-	// Cache the amount of memory needed to execute property functions for each batch.
-	PropertyFunctionEvaluationScopeMemoryRequirements.Reset(CopyBatches.Num());
-	for (FPropertyBindingCopyInfoBatch& CopyBatch : CopyBatches)
-	{
-		UE::StateTree::InstanceData::FEvaluationScopeInstanceContainer::FMemoryRequirementBuilder CopyBatchFunctionRequirement;
-
-		if (CopyBatch.PropertyFunctionsBegin != CopyBatch.PropertyFunctionsEnd)
-		{
-			const int32 FuncsBegin = CopyBatch.PropertyFunctionsBegin.Get();
-			const int32 FuncsEnd = CopyBatch.PropertyFunctionsEnd.Get();
-			for (int32 FuncIndex = FuncsBegin; FuncIndex < FuncsEnd; ++FuncIndex)
-			{
-				const FStateTreePropertyFunctionBase& Func = Nodes[FuncIndex].Get<const FStateTreePropertyFunctionBase>();
-				if (Func.InstanceDataHandle.GetSource() == EStateTreeDataSourceType::EvaluationScopeInstanceData
-					|| Func.InstanceDataHandle.GetSource() == EStateTreeDataSourceType::EvaluationScopeInstanceDataObject)
-				{
-					const FConstStructView DefaultInstanceView = DefaultEvaluationScopeInstanceData.GetStruct(Func.InstanceTemplateIndex.Get());
-					CopyBatchFunctionRequirement.Add(DefaultInstanceView.GetScriptStruct());
-				}
-			}
-		}
-
-		PropertyFunctionEvaluationScopeMemoryRequirements.Add(CopyBatchFunctionRequirement.Build());
-	}
-	check(CopyBatches.Num() == PropertyFunctionEvaluationScopeMemoryRequirements.Num());
-
 	return true;
 }
 
@@ -1450,22 +1374,6 @@ FString UStateTree::DebugInternalLayoutAsString() const
 		}
 	}
 
-	// Execution Scope InstanceData data (e.g. tasks)
-	DebugString.Appendf(TEXT("\nExecution Scope(%d)\n"), DefaultEvaluationScopeInstanceData.Num());
-	for (int32 Index = 0; Index < DefaultEvaluationScopeInstanceData.Num(); Index++)
-	{
-		if (DefaultEvaluationScopeInstanceData.IsObject(Index))
-		{
-			const UObject* Data = DefaultEvaluationScopeInstanceData.GetObject(Index);
-			PrintObjectNameSafe(Index, Data);
-		}
-		else
-		{
-			const FConstStructView Data = DefaultEvaluationScopeInstanceData.GetStruct(Index);
-			PrintViewNameSafe(Index, Data);
-		}
-	}
-
 	// External data (e.g. fragments, subsystems)
 	DebugString.Appendf(TEXT("\nExternal Data(%d)\n"), ExternalDataDescs.Num());
 	if (ExternalDataDescs.Num())
diff --git a/Engine/Plugins/Runtime/StateTree/Source/StateTreeModule/Private/StateTreeEvaluationScopeInstanceContainer.cpp b/Engine/Plugins/Runtime/StateTree/Source/StateTreeModule/Private/StateTreeEvaluationScopeInstanceContainer.cpp
deleted file mode 100644
diff --git a/Engine/Plugins/Runtime/StateTree/Source/StateTreeModule/Private/StateTreeExecutionContext.cpp b/Engine/Plugins/Runtime/StateTree/Source/StateTreeModule/Private/StateTreeExecutionContext.cpp
--- a/Engine/Plugins/Runtime/StateTree/Source/StateTreeModule/Private/StateTreeExecutionContext.cpp
+++ b/Engine/Plugins/Runtime/StateTree/Source/StateTreeModule/Private/StateTreeExecutionContext.cpp
@@ -2482,24 +2482,6 @@ FStateTreeDataView FStateTreeExecutionContext::GetDataView(const FStateTreeExecu
 		check(!ContextAndExternalDataViews.IsEmpty());
 		return ContextAndExternalDataViews[Handle.GetIndex()];
 
-	case EStateTreeDataSourceType::EvaluationScopeInstanceData:
-	case EStateTreeDataSourceType::EvaluationScopeInstanceDataObject:
-	{
-		// The data can be accessed in any of the caches (depends on how the binding was constructed), but it is most likely on top of the stack.
-		for (int32 Index = EvaluationScopeInstanceCaches.Num() - 1; Index >= 0; --Index)
-		{
-			if (EvaluationScopeInstanceCaches[Index].StateTree == CurrentFrame.StateTree)
-			{
-				if (FStateTreeDataView* Result = EvaluationScopeInstanceCaches[Index].Container->GetDataViewPtr(Handle))
-				{
-					return *Result;
-				}
-			}
-		}
-		checkf(false, TEXT("The evaluation scope instance data needs to be constructed before you can access it."));
-		return nullptr;
-	}
-
 	case EStateTreeDataSourceType::ExternalData:
 		check(!ContextAndExternalDataViews.IsEmpty());
 		return ContextAndExternalDataViews[CurrentFrame.ExternalDataBaseIndex.Get() + Handle.GetIndex()];
@@ -2796,21 +2778,6 @@ bool FStateTreeExecutionContext::IsHandleSourceValid(const FStateTreeExecutionFr
 	case EStateTreeDataSourceType::ContextData:
 		return true;
 
-	case EStateTreeDataSourceType::EvaluationScopeInstanceData:
-	case EStateTreeDataSourceType::EvaluationScopeInstanceDataObject:
-		// The data can be accessed in any of the caches (depends on how the binding was constructed), but it is most likely on top of the stack.
-		for (int32 Index = EvaluationScopeInstanceCaches.Num() - 1; Index >= 0; --Index)
-		{
-			if (EvaluationScopeInstanceCaches[Index].StateTree == CurrentFrame.StateTree)
-			{
-				if (FStateTreeDataView* Result = EvaluationScopeInstanceCaches[Index].Container->GetDataViewPtr(Handle))
-				{
-					return true;
-				}
-			}
-		}
-		return false;
-
 	case EStateTreeDataSourceType::ExternalData:
 		return CurrentFrame.ExternalDataBaseIndex.IsValid()
 			&& ContextAndExternalDataViews.IsValidIndex(CurrentFrame.ExternalDataBaseIndex.Get() + Handle.GetIndex());
@@ -2876,10 +2843,6 @@ FStateTreeDataView FStateTreeExecutionContext::GetTemporaryDataView(const FState
 	case EStateTreeDataSourceType::ExternalGlobalParameterData:
 		checkf(false, TEXT("External global parameter data currently not supported for linked state-trees"));
 		return {};
-	case EStateTreeDataSourceType::EvaluationScopeInstanceData:
-	case EStateTreeDataSourceType::EvaluationScopeInstanceDataObject:
-		ensureMsgf(false, TEXT("The evaluation scope instance data needs to be constructed before you can access it."));
-		return {};
 
 	default:
 		return UE::StateTree::InstanceData::Private::GetTemporaryDataView(Storage, ParentFrame, CurrentFrame, Handle);
@@ -2896,19 +2859,6 @@ FStateTreeDataView FStateTreeExecutionContext::AddTemporaryInstance(const FState
 	return NewInstance;
 }
 
-void FStateTreeExecutionContext::PushEvaluationScopeInstanceContainer(UE::StateTree::InstanceData::FEvaluationScopeInstanceContainer& Container, const FStateTreeExecutionFrame& Frame)
-{
-	EvaluationScopeInstanceCaches.Emplace(&Container, Frame.StateTree);
-}
-
-void FStateTreeExecutionContext::PopEvaluationScopeInstanceContainer(UE::StateTree::InstanceData::FEvaluationScopeInstanceContainer& Container)
-{
-	if (ensure(EvaluationScopeInstanceCaches.Num() > 0 && EvaluationScopeInstanceCaches.Last().Container == &Container))
-	{
-		EvaluationScopeInstanceCaches.Pop();
-	}
-}
-
 void FStateTreeExecutionContext::CopyAllBindingsOnActiveInstances(const ECopyBindings CopyType)
 {
 	FStateTreeExecutionState& Exec = GetExecState();
@@ -3010,47 +2960,15 @@ bool FStateTreeExecutionContext::CopyBatchWithValidation(const FStateTreeExecuti
 	return CopyBatchInternal<bOnActiveInstances>(ParentFrame, CurrentFrame, TargetView, BindingsBatch);
 }
 
-namespace UE::StateTree::ExecutionContext::Private
-{
-	void InitPropertyFunctions(UE::StateTree::InstanceData::FEvaluationScopeInstanceContainer& Container, TNotNull<const UStateTree*> StateTree, int32 FuncBegin, int32 FuncEnd)
-	{
-		for (int32 FuncIndex = FuncBegin; FuncIndex != FuncEnd; ++FuncIndex)
-		{
-			const FStateTreePropertyFunctionBase& Func = StateTree->GetNodes()[FuncIndex].Get<const FStateTreePropertyFunctionBase>();
-			if (Func.InstanceDataHandle.GetSource() == EStateTreeDataSourceType::EvaluationScopeInstanceData
-				|| Func.InstanceDataHandle.GetSource() == EStateTreeDataSourceType::EvaluationScopeInstanceDataObject)
-			{
-				Container.Add(Func.InstanceDataHandle, StateTree->GetDefaultEvaluationScopeInstanceData().GetStruct(Func.InstanceTemplateIndex.Get()));
-			}
-		}
-	}
-}
-
 template<bool bOnActiveInstances>
 bool FStateTreeExecutionContext::CopyBatchInternal(const FStateTreeExecutionFrame* ParentFrame, const FStateTreeExecutionFrame& CurrentFrame, const FStateTreeDataView TargetView, const FStateTreeIndex16 BindingsBatch)
 {
-	using namespace UE::StateTree::InstanceData;
-	using namespace UE::StateTree::ExecutionContext::Private;
-
 	const FPropertyBindingCopyInfoBatch& Batch = CurrentFrame.StateTree->PropertyBindings.Super::GetBatch(BindingsBatch);
 	check(TargetView.GetStruct() == Batch.TargetStruct.Get().Struct);
 
-	FEvaluationScopeInstanceContainer EvaluationScopeContainer;
-	bool bEvaluationScopeContainerPushed = false;
 	if (Batch.PropertyFunctionsBegin != Batch.PropertyFunctionsEnd)
 	{
 		check(Batch.PropertyFunctionsBegin.IsValid() && Batch.PropertyFunctionsEnd.IsValid());
-
-		const FEvaluationScopeInstanceContainer::FMemoryRequirement& MemoryRequirement = CurrentFrame.StateTree->PropertyFunctionEvaluationScopeMemoryRequirements[BindingsBatch.Get()];
-		if (MemoryRequirement.Size > 0)
-		{
-			void* FunctionBindingMemory = FMemory_Alloca_Aligned(MemoryRequirement.Size, MemoryRequirement.Alignment);
-			EvaluationScopeContainer = FEvaluationScopeInstanceContainer(FunctionBindingMemory, MemoryRequirement);
-			PushEvaluationScopeInstanceContainer(EvaluationScopeContainer, CurrentFrame);
-			InitPropertyFunctions(EvaluationScopeContainer, CurrentFrame.StateTree, Batch.PropertyFunctionsBegin.Get(), Batch.PropertyFunctionsEnd.Get());
-			bEvaluationScopeContainerPushed = true;
-		}
-
 		if constexpr (bOnActiveInstances)
 		{
 			EvaluatePropertyFunctionsOnActiveInstances(ParentFrame, CurrentFrame, FStateTreeIndex16(Batch.PropertyFunctionsBegin), Batch.PropertyFunctionsEnd.Get() - Batch.PropertyFunctionsBegin.Get());
@@ -3077,16 +2995,9 @@ bool FStateTreeExecutionContext::CopyBatchInternal(const FStateTreeExecutionFram
 				bSucceed = false;
 				break;
 			}
-
 			bSucceed &= CurrentFrame.StateTree->PropertyBindings.Super::CopyProperty(Copy, SourceView, TargetView);
 		}
 	}
-
-	if (bEvaluationScopeContainerPushed)
-	{
-		PopEvaluationScopeInstanceContainer(EvaluationScopeContainer);
-	}
-
 	return bSucceed;
 }
 
@@ -4692,10 +4603,10 @@ bool FStateTreeExecutionContext::TestAllConditionsInternal(const FStateTreeExecu
 
 	int32 Level = 0;
 
-	for (int32 Index = 0; Index < ConditionsNum; Index++)
-	{
-		const int32 ConditionIndex = ConditionsOffset + Index;
-		const FStateTreeConditionBase& Cond = CurrentFrame.StateTree->Nodes[ConditionIndex].Get<const FStateTreeConditionBase>();
+		for (int32 Index = 0; Index < ConditionsNum; Index++)
+		{
+			const int32 ConditionIndex = ConditionsOffset + Index;
+			const FStateTreeConditionBase& Cond = CurrentFrame.StateTree->Nodes[ConditionIndex].Get<const FStateTreeConditionBase>();
 		checkf(Cond.InstanceDataHandle.GetSource() == EStateTreeDataSourceType::SharedInstanceData
 			|| Cond.InstanceDataHandle.GetSource() == EStateTreeDataSourceType::SharedInstanceDataObject,
 			TEXT("Condition only support SharedInstanceData"));
@@ -4910,6 +4821,12 @@ void FStateTreeExecutionContext::EvaluatePropertyFunctionsInternal(const FStateT
 		}
 
 		Func.Execute(*this);
+
+		// Reset copied properties that might contain object references.
+		if (Func.BindingsBatch.IsValid())
+		{
+			CurrentFrame.StateTree->PropertyBindings.Super::ResetObjects(Func.BindingsBatch, FuncInstanceView);
+		}
 	}
 }
 
diff --git a/Engine/Plugins/Runtime/StateTree/Source/StateTreeModule/Private/StateTreeInstanceContainer.cpp b/Engine/Plugins/Runtime/StateTree/Source/StateTreeModule/Private/StateTreeInstanceContainer.cpp
deleted file mode 100644
diff --git a/Engine/Plugins/Runtime/StateTree/Source/StateTreeModule/Private/StateTreeInstanceData.cpp b/Engine/Plugins/Runtime/StateTree/Source/StateTreeModule/Private/StateTreeInstanceData.cpp
--- a/Engine/Plugins/Runtime/StateTree/Source/StateTreeModule/Private/StateTreeInstanceData.cpp
+++ b/Engine/Plugins/Runtime/StateTree/Source/StateTreeModule/Private/StateTreeInstanceData.cpp
@@ -19,27 +19,11 @@
 const FGuid FStateTreeInstanceStorageCustomVersion::GUID(0x60C4F0DE, 0x8B264C34, 0xAA937201, 0x5DFF09CC);
 FCustomVersionRegistration GRegisterStateTreeInstanceStorageCustomVersion(FStateTreeInstanceStorageCustomVersion::GUID, FStateTreeInstanceStorageCustomVersion::LatestVersion, TEXT("StateTreeInstanceStorage"));
 
-namespace UE::StateTree::InstanceData::Private
+namespace UE::StateTree
 {
-	bool AreAllInstancesValid(const FInstancedStructContainer& InstanceStructs)
-	{
-		for (FConstStructView Instance : InstanceStructs)
-		{
-			if (!Instance.IsValid())
-			{
-				return false;
-			}
-			if (const FStateTreeInstanceObjectWrapper* Wrapper = Instance.GetPtr<const FStateTreeInstanceObjectWrapper>())
-			{
-				if (!Wrapper->InstanceObject)
-				{
-					return false;
-				}
-			}
-		}
-		return true;
-	}
-
+	/**
+	 * Duplicates object, and tries to covert old BP classes (REINST_*) to their newer version.
+	 */
 	TNotNull<UObject*> CopyNodeInstance(TNotNull<UObject*> Instance, TNotNull<UObject*> InOwner, bool bDuplicate)
 	{
 		const UClass* InstanceClass = Instance->GetClass();
@@ -74,25 +58,6 @@ namespace UE::StateTree::InstanceData::Private
 		return Instance;
 	}
 
-	void PostAppendToInstanceStructContainer(FInstancedStructContainer& InstanceStructs, TNotNull<UObject*> InOwner, bool bDuplicateWrappedObject, int32 StartIndex)
-	{
-		for (int32 Index = StartIndex; Index < InstanceStructs.Num(); ++Index)
-		{
-			if (FStateTreeInstanceObjectWrapper* Wrapper = InstanceStructs[Index].GetPtr<FStateTreeInstanceObjectWrapper>())
-			{
-				if (Wrapper->InstanceObject)
-				{
-					const bool bDuplicate = bDuplicateWrappedObject || InOwner != Wrapper->InstanceObject->GetOuter();
-					Wrapper->InstanceObject = CopyNodeInstance(Wrapper->InstanceObject, InOwner, bDuplicate);
-				}
-			}
-		}
-	}
-}
-
-namespace UE::StateTree
-{
-
 #if WITH_EDITORONLY_DATA
 	void GatherForLocalization(const FString& PathToParent, const UScriptStruct* Struct, const void* StructData, const void* DefaultStructData, FPropertyLocalizationDataGatherer& PropertyLocalizationDataGatherer, const EPropertyLocalizationGathererTextFlags GatherTextFlags)
 	{
@@ -483,7 +448,21 @@ void FStateTreeInstanceStorage::ResetTransitionRequests()
 
 bool FStateTreeInstanceStorage::AreAllInstancesValid() const
 {
-	return UE::StateTree::InstanceData::Private::AreAllInstancesValid(InstanceStructs);
+	for (FConstStructView Instance : InstanceStructs)
+	{
+		if (!Instance.IsValid())
+		{
+			return false;
+		}
+		if (const FStateTreeInstanceObjectWrapper* Wrapper = Instance.GetPtr<const FStateTreeInstanceObjectWrapper>())
+		{
+			if (!Wrapper->InstanceObject)
+			{
+				return false;
+			}
+		}
+	}
+	return true;
 }
 
 FStructView FStateTreeInstanceStorage::AddTemporaryInstance(UObject& InOwner, const FStateTreeExecutionFrame& Frame, const FStateTreeIndex16 OwnerNodeIndex, const FStateTreeDataHandle DataHandle, FConstStructView NewInstanceData)
@@ -517,7 +496,7 @@ FStructView FStateTreeInstanceStorage::AddTemporaryInstance(UObject& InOwner, co
 		if (Wrapper->InstanceObject)
 		{
 			constexpr bool bDuplicate = true;
-			Wrapper->InstanceObject = UE::StateTree::InstanceData::Private::CopyNodeInstance(Wrapper->InstanceObject, &InOwner, bDuplicate);
+			Wrapper->InstanceObject = UE::StateTree::CopyNodeInstance(Wrapper->InstanceObject, &InOwner, bDuplicate);
 		}
 	}
 
@@ -934,7 +913,7 @@ void FStateTreeInstanceData::CopyFrom(UObject& InOwner, const FStateTreeInstance
 			if (Wrapper->InstanceObject)
 			{
 				constexpr bool bDuplicate = true;
-				Wrapper->InstanceObject = UE::StateTree::InstanceData::Private::CopyNodeInstance(Wrapper->InstanceObject, &InOwner, bDuplicate);
+				Wrapper->InstanceObject = UE::StateTree::CopyNodeInstance(Wrapper->InstanceObject, &InOwner, bDuplicate);
 			}
 		}
 	}
@@ -954,12 +933,42 @@ void FStateTreeInstanceData::Init(UObject& InOwner, TConstArrayView<FConstStruct
 
 void FStateTreeInstanceData::Append(UObject& InOwner, TConstArrayView<FInstancedStruct> InStructs, FAddArgs Args)
 {
-	UE::StateTree::InstanceData::Private::AppendToInstanceStructContainer(GetMutableStorage().InstanceStructs, &InOwner, InStructs, Args.bDuplicateWrappedObject);
+	FStateTreeInstanceStorage& Storage = GetMutableStorage();
+
+	const int32 StartIndex = Storage.InstanceStructs.Num();
+	Storage.InstanceStructs.Append(InStructs);
+
+	for (int32 Index = StartIndex; Index < Storage.InstanceStructs.Num(); Index++)
+	{
+		if (FStateTreeInstanceObjectWrapper* Wrapper = Storage.InstanceStructs[Index].GetPtr<FStateTreeInstanceObjectWrapper>())
+		{
+			if (Wrapper->InstanceObject)
+			{
+				const bool bDuplicate = Args.bDuplicateWrappedObject || &InOwner != Wrapper->InstanceObject->GetOuter();
+				Wrapper->InstanceObject = UE::StateTree::CopyNodeInstance(Wrapper->InstanceObject, &InOwner, bDuplicate);
+			}
+		}
+	}
 }
 
 void FStateTreeInstanceData::Append(UObject& InOwner, TConstArrayView<FConstStructView> InStructs, FAddArgs Args)
 {
-	UE::StateTree::InstanceData::Private::AppendToInstanceStructContainer(GetMutableStorage().InstanceStructs, &InOwner, InStructs, Args.bDuplicateWrappedObject);
+	FStateTreeInstanceStorage& Storage = GetMutableStorage();
+
+	const int32 StartIndex = Storage.InstanceStructs.Num();
+	Storage.InstanceStructs.Append(InStructs);
+	
+	for (int32 Index = StartIndex; Index < Storage.InstanceStructs.Num(); Index++)
+	{
+		if (FStateTreeInstanceObjectWrapper* Wrapper = Storage.InstanceStructs[Index].GetPtr<FStateTreeInstanceObjectWrapper>())
+		{
+			if (Wrapper->InstanceObject)
+			{
+				const bool bDuplicate = Args.bDuplicateWrappedObject || &InOwner != Wrapper->InstanceObject->GetOuter();
+				Wrapper->InstanceObject = UE::StateTree::CopyNodeInstance(Wrapper->InstanceObject, &InOwner, bDuplicate);
+			}
+		}
+	}
 }
 
 void FStateTreeInstanceData::Append(UObject& InOwner, TConstArrayView<FConstStructView> InStructs, TConstArrayView<FInstancedStruct*> InInstancesToMove, FAddArgs Args)
@@ -993,7 +1002,7 @@ void FStateTreeInstanceData::Append(UObject& InOwner, TConstArrayView<FConstStru
 			if (Wrapper->InstanceObject)
 			{
 				const bool bDuplicate = Args.bDuplicateWrappedObject || &InOwner != Wrapper->InstanceObject->GetOuter();
-				Wrapper->InstanceObject = UE::StateTree::InstanceData::Private::CopyNodeInstance(Wrapper->InstanceObject, &InOwner, bDuplicate);
+				Wrapper->InstanceObject = UE::StateTree::CopyNodeInstance(Wrapper->InstanceObject, &InOwner, bDuplicate);
 			}
 		}
 	}
diff --git a/Engine/Plugins/Runtime/StateTree/Source/StateTreeModule/Private/StateTreeInstanceDataHelpers.h b/Engine/Plugins/Runtime/StateTree/Source/StateTreeModule/Private/StateTreeInstanceDataHelpers.h
--- a/Engine/Plugins/Runtime/StateTree/Source/StateTreeModule/Private/StateTreeInstanceDataHelpers.h
+++ b/Engine/Plugins/Runtime/StateTree/Source/StateTreeModule/Private/StateTreeInstanceDataHelpers.h
@@ -25,22 +25,5 @@ namespace UE::StateTree::InstanceData::Private
 		const FStateTreeExecutionFrame* ParentFrame,
 		const FStateTreeExecutionFrame& CurrentFrame,
 		const FStateTreeDataHandle& Handle);
-
-	/** @return true if all instances are valid. */
-	[[nodiscard]] bool AreAllInstancesValid(const FInstancedStructContainer& InstanceStructs);
-
-	/** Duplicates object, and tries to covert old BP classes (REINST_*) to their newer version. */
-	[[nodiscard]] TNotNull<UObject*> CopyNodeInstance(TNotNull<UObject*> Instance, TNotNull<UObject*> InOwner, bool bDuplicate);
-
-	/** None generic code for AppendToInstanceStructContainer. Called after AppendToInstanceStructContainer to move code to cpp. */
-	void PostAppendToInstanceStructContainer(FInstancedStructContainer& InstanceStructs, TNotNull<UObject*> InOwner, bool bDuplicateWrappedObject, int32 StartIndex);
-
-	/** Appends new items to the instance, and moves existing data into the allocated instances. */
-	template<typename TOtherType>
-	void AppendToInstanceStructContainer(FInstancedStructContainer& InstanceStructs, TNotNull<UObject*> InOwner, TConstArrayView<TOtherType> InStructs, bool bDuplicateWrappedObject)
-	{
-		const int32 StartIndex = InstanceStructs.Num();
-		InstanceStructs.Append(InStructs);
-		PostAppendToInstanceStructContainer(InstanceStructs, InOwner, bDuplicateWrappedObject, StartIndex);
-	}
 } // namespace
+
diff --git a/Engine/Plugins/Runtime/StateTree/Source/StateTreeModule/Private/StateTreeTypes.cpp b/Engine/Plugins/Runtime/StateTree/Source/StateTreeModule/Private/StateTreeTypes.cpp
--- a/Engine/Plugins/Runtime/StateTree/Source/StateTreeModule/Private/StateTreeTypes.cpp
+++ b/Engine/Plugins/Runtime/StateTree/Source/StateTreeModule/Private/StateTreeTypes.cpp
@@ -144,28 +144,6 @@ PRAGMA_ENABLE_DEPRECATION_WARNINGS
 //////////////////////////////////////////////////////////////////////////
 // FStateTreeDataHandle
 
-bool FStateTreeDataHandle::IsObjectSource() const
-{
-	return Source == EStateTreeDataSourceType::GlobalInstanceDataObject
-		|| Source == EStateTreeDataSourceType::ActiveInstanceDataObject
-		|| Source == EStateTreeDataSourceType::EvaluationScopeInstanceDataObject;
-}
-
-FStateTreeDataHandle FStateTreeDataHandle::ToObjectSource() const
-{
-	switch (Source)
-	{
-	case EStateTreeDataSourceType::GlobalInstanceData:
-		return FStateTreeDataHandle(EStateTreeDataSourceType::GlobalInstanceDataObject, Index, StateHandle);
-	case EStateTreeDataSourceType::ActiveInstanceData:
-		return FStateTreeDataHandle(EStateTreeDataSourceType::ActiveInstanceDataObject, Index, StateHandle);
-	case EStateTreeDataSourceType::EvaluationScopeInstanceData:
-		return FStateTreeDataHandle(EStateTreeDataSourceType::EvaluationScopeInstanceDataObject, Index, StateHandle);
-	default:
-		return *this;
-	}
-}
-
 FString FStateTreeDataHandle::Describe() const
 {
 	switch (Source)
@@ -184,10 +162,6 @@ FString FStateTreeDataHandle::Describe() const
 		return FString::Printf(TEXT("Shared[%d]"), Index);
 	case EStateTreeDataSourceType::SharedInstanceDataObject:
 		return FString::Printf(TEXT("SharedObj[%d]"), Index);
-	case EStateTreeDataSourceType::EvaluationScopeInstanceData:
-		return FString::Printf(TEXT("EvalScope[%d]"), Index);
-	case EStateTreeDataSourceType::EvaluationScopeInstanceDataObject:
-		return FString::Printf(TEXT("EvalScopeObj[%d]"), Index);
 	case EStateTreeDataSourceType::ContextData:
 		return FString::Printf(TEXT("Context[%d]"), Index);
 	case EStateTreeDataSourceType::GlobalParameterData:
diff --git a/Engine/Plugins/Runtime/StateTree/Source/StateTreeModule/Public/StateTree.h b/Engine/Plugins/Runtime/StateTree/Source/StateTreeModule/Public/StateTree.h
--- a/Engine/Plugins/Runtime/StateTree/Source/StateTreeModule/Public/StateTree.h
+++ b/Engine/Plugins/Runtime/StateTree/Source/StateTreeModule/Public/StateTree.h
@@ -3,7 +3,6 @@
 #pragma once
 
 #include "Engine/DataAsset.h"
-#include "StateTreeEvaluationScopeInstanceContainer.h"
 #include "StateTreeSchema.h"
 #include "StateTreePropertyBindings.h"
 #include "StateTreeInstanceData.h"
@@ -130,12 +129,6 @@ public:
 		return NumContextData;
 	}
 
-	/** @return Default evaluation scope instance data. */
-	const UE::StateTree::InstanceData::FInstanceContainer& GetDefaultEvaluationScopeInstanceData() const
-	{
-		return DefaultEvaluationScopeInstanceData;
-	}
-
 	/** @return List of external data required by the state tree */
 	TConstArrayView<FStateTreeExternalDataDesc> GetExternalDataDescs() const
 	{
@@ -396,10 +389,6 @@ private:
 	UPROPERTY()
 	FStateTreeInstanceData DefaultInstanceData;
 
-	/** Default node instance data for evaluation scope (e.g. conditions, considerations, function) */
-	UPROPERTY()
-	UE::StateTree::InstanceData::FInstanceContainer DefaultEvaluationScopeInstanceData;
-
 	/** Shared node instance data (e.g. conditions, considerations). */
 	UPROPERTY()
 	FStateTreeInstanceData SharedInstanceData;
@@ -414,10 +403,6 @@ private:
 	UPROPERTY()
 	FStateTreePropertyBindings PropertyBindings;
 
-	using FMemoryRequirement = UE::StateTree::InstanceData::FEvaluationScopeInstanceContainer::FMemoryRequirement;
-	/** The amount of memory used by property binding copy info batch for property function. */
-	TArray<FMemoryRequirement> PropertyFunctionEvaluationScopeMemoryRequirements;
-
 	/** The asset extensions. A place to add extra information for plugins. */
 	UPROPERTY()
 	TArray<TObjectPtr<UStateTreeExtension>> Extensions;
@@ -458,7 +443,7 @@ private:
 	/** Number of global instance data. */
 	UPROPERTY()
 	uint16 NumGlobalInstanceData = 0;
-
+	
 	/** Index of first evaluator in Nodes. */
 	UPROPERTY()
 	uint16 EvaluatorsBegin = 0;
diff --git a/Engine/Plugins/Runtime/StateTree/Source/StateTreeModule/Public/StateTreeEvaluationScopeInstanceContainer.h b/Engine/Plugins/Runtime/StateTree/Source/StateTreeModule/Public/StateTreeEvaluationScopeInstanceContainer.h
deleted file mode 100644
diff --git a/Engine/Plugins/Runtime/StateTree/Source/StateTreeModule/Public/StateTreeExecutionContext.h b/Engine/Plugins/Runtime/StateTree/Source/StateTreeModule/Public/StateTreeExecutionContext.h
--- a/Engine/Plugins/Runtime/StateTree/Source/StateTreeModule/Public/StateTreeExecutionContext.h
+++ b/Engine/Plugins/Runtime/StateTree/Source/StateTreeModule/Public/StateTreeExecutionContext.h
@@ -31,11 +31,6 @@ struct FStateTreeTransitionRequest;
 struct FStateTreeInstanceDebugId;
 struct FStateTreeWeakExecutionContext;
 
-namespace UE::StateTree::InstanceData
-{
-	struct FEvaluationScopeInstanceContainer;
-}
-
 namespace UE::StateTree::ExecutionContext
 {
 	UE_API bool MarkDelegateAsBroadcasted(FStateTreeDelegateDispatcher Dispatcher, const FStateTreeExecutionFrame& CurrentFrame, FStateTreeInstanceStorage& Storage);
@@ -1452,21 +1447,8 @@ PRAGMA_ENABLE_DEPRECATION_WARNINGS
 		FStateTreeIndex16 BaseIndex;
 	};
 	TArray<FCollectedExternalDataCache> CollectedExternalCache;
-	bool bActiveExternalDataCollected = false;
 
-	/** Holds the container. Multiple instances with the same state tree can occur in a recursive call. */
-	struct FEvaluationScopeDataCache
-	{
-		UE::StateTree::InstanceData::FEvaluationScopeInstanceContainer* Container;
-		TObjectKey<UStateTree> StateTree;
-	};
-
-	static constexpr int32 ExpectedEvaluationScopeCacheLength = 4;
-	/** Data view of the evaluation scope instance data. */
-	TArray<FEvaluationScopeDataCache, TInlineAllocator<ExpectedEvaluationScopeCacheLength>> EvaluationScopeInstanceCaches;
-
-	UE_API void PushEvaluationScopeInstanceContainer(UE::StateTree::InstanceData::FEvaluationScopeInstanceContainer& Container, const FStateTreeExecutionFrame& Frame);
-	UE_API void PopEvaluationScopeInstanceContainer(UE::StateTree::InstanceData::FEvaluationScopeInstanceContainer& Container);
+	bool bActiveExternalDataCollected = false;
 
 #if WITH_EDITORONLY_DATA
 	UE_DEPRECATED(5.6, "CompletedStateRunStatus is not used.")
diff --git a/Engine/Plugins/Runtime/StateTree/Source/StateTreeModule/Public/StateTreeInstanceContainer.h b/Engine/Plugins/Runtime/StateTree/Source/StateTreeModule/Public/StateTreeInstanceContainer.h
deleted file mode 100644
diff --git a/Engine/Plugins/Runtime/StateTree/Source/StateTreeModule/Public/StateTreeInstanceData.h b/Engine/Plugins/Runtime/StateTree/Source/StateTreeModule/Public/StateTreeInstanceData.h
--- a/Engine/Plugins/Runtime/StateTree/Source/StateTreeModule/Public/StateTreeInstanceData.h
+++ b/Engine/Plugins/Runtime/StateTree/Source/StateTreeModule/Public/StateTreeInstanceData.h
@@ -5,7 +5,6 @@
 #include "StructUtils/InstancedStructContainer.h"
 #include "Debugger/StateTreeRuntimeValidation.h"
 #include "StateTreeEvents.h"
-#include "StateTreeInstanceContainer.h"
 #include "StateTreeTypes.h"
 #include "StateTreeExecutionTypes.h"
 #include "Misc/MTAccessDetector.h"
@@ -45,6 +44,21 @@ namespace UE::StateTree::InstanceData
 		const FStateTreeDataHandle& Handle);
 }
 
+/** Wrapper class to store an object amongst the structs. */
+USTRUCT()
+struct FStateTreeInstanceObjectWrapper
+{
+	GENERATED_BODY()
+
+	FStateTreeInstanceObjectWrapper() = default;
+	FStateTreeInstanceObjectWrapper(UObject* Object)
+		: InstanceObject(Object)
+	{}
+	
+	UPROPERTY()
+	TObjectPtr<UObject> InstanceObject = nullptr;
+};
+
 /**
  * Holds temporary instance data created during state selection.
  * The data is identified by Frame (StateTree + RootState) and DataHandle.
diff --git a/Engine/Plugins/Runtime/StateTree/Source/StateTreeModule/Public/StateTreeTypes.h b/Engine/Plugins/Runtime/StateTree/Source/StateTreeModule/Public/StateTreeTypes.h
--- a/Engine/Plugins/Runtime/StateTree/Source/StateTreeModule/Public/StateTreeTypes.h
+++ b/Engine/Plugins/Runtime/StateTree/Source/StateTreeModule/Public/StateTreeTypes.h
@@ -298,18 +298,12 @@ enum class EStateTreeDataSourceType : uint8
 	/** Active State Tasks */
 	ActiveInstanceDataObject,
 
-	/** Conditions, considerations and function bindings. */
+	/** Conditions and Considerations */
 	SharedInstanceData,
 
-	/** Conditions, considerations and function bindings*/
+	/** Conditions and Considerations */
 	SharedInstanceDataObject,
 
-	/** Temporary data constructor, used and destroyed immediately. Used by conditions, considerations and function bindings.*/
-	EvaluationScopeInstanceData,
-
-	/** Temporary data constructor, used and destroyed immediately. Used by conditions, considerations and function bindings.*/
-	EvaluationScopeInstanceDataObject,
-
 	/** Context Data, Tree Parameters */
 	ContextData,
 
@@ -439,9 +433,27 @@ struct FStateTreeDataHandle
 		return StateHandle;
 	}
 
-	STATETREEMODULE_API bool IsObjectSource() const;
+	bool IsObjectSource() const
+	{
+		return Source == EStateTreeDataSourceType::GlobalInstanceDataObject
+			|| Source == EStateTreeDataSourceType::ActiveInstanceDataObject
+			|| Source == EStateTreeDataSourceType::SharedInstanceDataObject;
+	}
 
-	STATETREEMODULE_API FStateTreeDataHandle ToObjectSource() const;
+	FStateTreeDataHandle ToObjectSource() const
+	{
+		switch (Source)
+		{
+		case EStateTreeDataSourceType::GlobalInstanceData:
+			return FStateTreeDataHandle(EStateTreeDataSourceType::GlobalInstanceDataObject, Index, StateHandle);
+		case EStateTreeDataSourceType::ActiveInstanceData:
+			return FStateTreeDataHandle(EStateTreeDataSourceType::ActiveInstanceDataObject, Index, StateHandle);
+		case EStateTreeDataSourceType::SharedInstanceData:
+			return FStateTreeDataHandle(EStateTreeDataSourceType::SharedInstanceDataObject, Index, StateHandle);
+		default:
+			return *this;
+		}
+	}
 
 	STATETREEMODULE_API FString Describe() const;
 
diff --git a/Engine/Plugins/Runtime/StateTree/Source/StateTreeModule/StateTreeModule.Build.cs b/Engine/Plugins/Runtime/StateTree/Source/StateTreeModule/StateTreeModule.Build.cs
--- a/Engine/Plugins/Runtime/StateTree/Source/StateTreeModule/StateTreeModule.Build.cs
+++ b/Engine/Plugins/Runtime/StateTree/Source/StateTreeModule/StateTreeModule.Build.cs
@@ -8,8 +8,6 @@ namespace UnrealBuildTool.Rules
 		{
 			CppCompileWarningSettings.UnsafeTypeCastWarningLevel = WarningLevel.Warning;
 
-			bAllowUETypesInNamespaces = true;
-
 			PublicIncludePaths.AddRange(
 				new string[] {
 				}
diff --git a/Engine/Plugins/Runtime/StateTree/Source/StateTreeTestSuite/Private/StateTreeTestTypes.h b/Engine/Plugins/Runtime/StateTree/Source/StateTreeTestSuite/Private/StateTreeTestTypes.h
--- a/Engine/Plugins/Runtime/StateTree/Source/StateTreeTestSuite/Private/StateTreeTestTypes.h
+++ b/Engine/Plugins/Runtime/StateTree/Source/StateTreeTestSuite/Private/StateTreeTestTypes.h
@@ -649,7 +649,6 @@ struct FStateTreeTestConditionInstanceData
 
 	static std::atomic<int32> GlobalCounter;
 };
-STATETREE_POD_INSTANCEDATA(FStateTreeTestConditionInstanceData);
 
 USTRUCT(meta = (Hidden))
 struct FStateTreeTestCondition : public FStateTreeConditionCommonBase
@@ -867,7 +866,6 @@ struct FTestPropertyFunction_InstanceData
 	UPROPERTY()
 	int32 Result = 0;
 };
-STATETREE_POD_INSTANCEDATA(FTestPropertyFunction_InstanceData);
 
 USTRUCT()
 struct FTestPropertyFunction : public FStateTreePropertyFunctionBase
@@ -931,7 +929,6 @@ struct FStateTreeTestBooleanConditionInstanceData
 	UPROPERTY(EditAnywhere, Category = Parameters)
 	bool bSuccess = true;
 };
-STATETREE_POD_INSTANCEDATA(FStateTreeTestBooleanConditionInstanceData);
 
 USTRUCT(meta = (Hidden))
 struct FStateTreeTestBooleanCondition : public FStateTreeConditionCommonBase
