[StateTree] Missing files and code with backed-out CL 44477348
#rnx

diff --git a/Engine/Plugins/Runtime/StateTree/Source/StateTreeEditorModule/Private/StateTreeCompiler.cpp b/Engine/Plugins/Runtime/StateTree/Source/StateTreeEditorModule/Private/StateTreeCompiler.cpp
--- a/Engine/Plugins/Runtime/StateTree/Source/StateTreeEditorModule/Private/StateTreeCompiler.cpp
+++ b/Engine/Plugins/Runtime/StateTree/Source/StateTreeEditorModule/Private/StateTreeCompiler.cpp
@@ -65,6 +65,13 @@ namespace UE::StateTree::Compiler
 		TEXT("Enable binding from delegate dispatchers that are in the state tree parameters.")
 	);
 
+	FAutoConsoleVariable CVarLogEnablePropertyFunctionWithEvaluationScopeInstanceData(
+		TEXT("StateTree.Compiler.EnablePropertyFunctionWithEvaluationScopeInstanceData"),
+		true,
+		TEXT("Use EvaluationScope data for property functions instead of SharedInstance data.\n")
+		TEXT("SharedInstance is the previous behavior that is deprecated.")
+	);
+
 	FAutoConsoleVariable CVarLogCompiledStateTree(
 		TEXT("StateTree.Compiler.LogResultOnCompilationCompleted"),
 		false,
@@ -1704,8 +1711,12 @@ bool FStateTreeCompiler::CreatePropertyFunction(const FStateTreeEditorNode& Func
 	StructDesc.Name = FuncEditorNode.GetName();
 	StructDesc.DataSource = EStateTreeBindableStructSource::PropertyFunction;
 
-	TArray<FInstancedStruct>& InstancedStructContainer = SharedInstanceStructs;
-	FInstancedStruct* CreatedNode = CreateNodeWithSharedInstanceData(nullptr, FuncEditorNode, StructDesc);
+	const bool bUseEvaluationScopeInstanceData = UE::StateTree::Compiler::CVarLogEnablePropertyFunctionWithEvaluationScopeInstanceData->GetBool();
+	const FStateTreeDataHandle DataHandle = bUseEvaluationScopeInstanceData
+		? FStateTreeDataHandle(EStateTreeDataSourceType::EvaluationScopeInstanceData, EvaluationScopeStructs.Num())
+		: FStateTreeDataHandle(EStateTreeDataSourceType::SharedInstanceData, SharedInstanceStructs.Num());
+	TArray<FInstancedStruct>& InstancedStructContainer = bUseEvaluationScopeInstanceData ? EvaluationScopeStructs : SharedInstanceStructs;
+	FInstancedStruct* CreatedNode = CreateNode(nullptr, FuncEditorNode, StructDesc, DataHandle, InstancedStructContainer);
 	if (CreatedNode == nullptr)
 	{
 		return false;
diff --git a/Engine/Plugins/Runtime/StateTree/Source/StateTreeModule/Private/StateTreeInstanceContainer.cpp b/Engine/Plugins/Runtime/StateTree/Source/StateTreeModule/Private/StateTreeInstanceContainer.cpp
--- a/Engine/Plugins/Runtime/StateTree/Source/StateTreeModule/Private/StateTreeInstanceContainer.cpp
+++ b/Engine/Plugins/Runtime/StateTree/Source/StateTreeModule/Private/StateTreeInstanceContainer.cpp
@@ -10,18 +10,17 @@ namespace UE::StateTree::InstanceData
 
 void FInstanceContainer::Init(TNotNull<UObject*> InOwner, TConstArrayView<FInstancedStruct> InStructs, FAddArgs Args)
 {
-	//Private::AppendToInstanceStructContainer(InstanceStructs, InOwner, InStructs, Args.bDuplicateWrappedObject);
+	Private::AppendToInstanceStructContainer(InstanceStructs, InOwner, InStructs, Args.bDuplicateWrappedObject);
 }
 
 void FInstanceContainer::Init(TNotNull<UObject*> InOwner, TConstArrayView<FConstStructView> InStructs, FAddArgs Args)
 {
-	//Private::AppendToInstanceStructContainer(InstanceStructs, InOwner, InStructs, Args.bDuplicateWrappedObject);
+	Private::AppendToInstanceStructContainer(InstanceStructs, InOwner, InStructs, Args.bDuplicateWrappedObject);
 }
 
 bool FInstanceContainer::AreAllInstancesValid() const
 {
-	//return Private::AreAllInstancesValid(InstanceStructs);
-	return true;
+	return Private::AreAllInstancesValid(InstanceStructs);
 }
 
 } // namespace UE::StateTree::InstanceData
