[StateTree] Fix transition that didn't remove the frame when a sub-frame changes root.

diff --git a/Engine/Plugins/Runtime/StateTree/Source/StateTreeModule/Private/StateTreeExecutionContext.cpp b/Engine/Plugins/Runtime/StateTree/Source/StateTreeModule/Private/StateTreeExecutionContext.cpp
--- a/Engine/Plugins/Runtime/StateTree/Source/StateTreeModule/Private/StateTreeExecutionContext.cpp
+++ b/Engine/Plugins/Runtime/StateTree/Source/StateTreeModule/Private/StateTreeExecutionContext.cpp
@@ -3623,7 +3623,6 @@ void FStateTreeExecutionContext::ExitState(const TSharedPtr<const FSelectStateRe
 	const bool bHasSustained = CommonStates.Contains(ArgsTargetState);
 
 	bool bContinue = true;
-	int32 PreviousFrameIndex = INDEX_NONE;
 	FStateTreeTransitionResult CurrentTransition = Transition;
 
 	auto ExitPreviousFrame = [this, &Exec, &ArgsSelectedFrames, &CurrentTransition](const int32 FrameIndex, const bool bCallExit)
@@ -3647,7 +3646,8 @@ void FStateTreeExecutionContext::ExitState(const TSharedPtr<const FSelectStateRe
 			}
 		};
 
-	for (int32 FrameIndex = Exec.ActiveFrames.Num() - 1; bContinue && FrameIndex >= 0; --FrameIndex)
+	int32 FrameIndex = Exec.ActiveFrames.Num() - 1;
+	for (; bContinue && FrameIndex >= 0; --FrameIndex)
 	{
 		FStateTreeExecutionFrame* CurrentParentFrame = FrameIndex > 0 ? &Exec.ActiveFrames[FrameIndex - 1] : nullptr;
 		FStateTreeExecutionFrame& CurrentFrame = Exec.ActiveFrames[FrameIndex];
@@ -3678,13 +3678,6 @@ void FStateTreeExecutionContext::ExitState(const TSharedPtr<const FSelectStateRe
 				bContinue = false;
 			}
 
-			// If the owning state 
-			if (PreviousFrameIndex != INDEX_NONE)
-			{
-				ExitPreviousFrame(PreviousFrameIndex, true);
-				PreviousFrameIndex = INDEX_NONE;
-			}
-
 			const bool bIsStateSustained = bIsStateCommon;
 
 			CurrentTransition.CurrentState = CurrentHandle;
@@ -3789,12 +3782,8 @@ void FStateTreeExecutionContext::ExitState(const TSharedPtr<const FSelectStateRe
 			UE_STATETREE_DEBUG_STATE_EVENT(this, CurrentHandle, EStateTreeTraceEventType::OnExited);
 		}
 
-		PreviousFrameIndex = FrameIndex;
-	}
-
-	if (PreviousFrameIndex != INDEX_NONE)
-	{
-		ExitPreviousFrame(PreviousFrameIndex, bContinue);
+		// The previous frame is not in the new
+		ExitPreviousFrame(FrameIndex, bContinue);
 	}
 }
 
diff --git a/Engine/Plugins/Runtime/StateTree/Source/StateTreeTestSuite/Private/StateTreeTransitionTest.cpp b/Engine/Plugins/Runtime/StateTree/Source/StateTreeTestSuite/Private/StateTreeTransitionTest.cpp
--- a/Engine/Plugins/Runtime/StateTree/Source/StateTreeTestSuite/Private/StateTreeTransitionTest.cpp
+++ b/Engine/Plugins/Runtime/StateTree/Source/StateTreeTestSuite/Private/StateTreeTransitionTest.cpp
@@ -1859,6 +1859,85 @@ struct FStateTreeTest_TransitionTwoEvents : FStateTreeTestBase
 };
 IMPLEMENT_STATE_TREE_INSTANT_TEST(FStateTreeTest_TransitionTwoEvents, "System.StateTree.Transition.TwoEvents");
 
+struct FStateTreeTest_TransitionLinkedAssetWith2Root : FStateTreeTestBase
+{
+	virtual bool InstantTest() override
+	{
+		/*
+		- Tree1
+			- Root1
+				- State1 : LinkedAsset
+		- Tree2
+			- Root1
+				- State1: Transition To Root3
+			- Root2
+				- State2
+		*/
+
+		UStateTree& StateTree1 = NewStateTree();
+		UStateTree& StateTree2 = NewStateTree();
+		{
+			UStateTreeEditorData& EditorData = *Cast<UStateTreeEditorData>(StateTree2.EditorData);
+			UStateTreeState& Root1 = EditorData.AddSubTree(FName("Root1"));
+			UStateTreeState& State1 = Root1.AddChildState(FName("State1"));
+			UStateTreeState& Root2 = EditorData.AddSubTree(FName("Root2"));
+			UStateTreeState& State2 = Root2.AddChildState(FName("State2"));
+
+			auto& Task1 = State1.AddTask<FTestTask_Stand>(FName("Tree2State1Task1"));
+			Task1.GetNode().TicksToCompletion = 1;
+			State1.AddTransition(EStateTreeTransitionTrigger::OnStateCompleted, EStateTreeTransitionType::GotoState, &State2);
+
+			auto& Task2 = State2.AddTask<FTestTask_Stand>(FName(TEXT("Task1A")));
+			Task2.GetNode().TicksToCompletion = 1;
+			State2.AddTransition(EStateTreeTransitionTrigger::OnStateCompleted, EStateTreeTransitionType::Succeeded);
+
+			FStateTreeCompilerLog Log;
+			FStateTreeCompiler Compiler(Log);
+			const bool bResult = Compiler.Compile(StateTree2);
+			AITEST_TRUE("StateTree should get compiled", bResult);
+		}
+		{
+			UStateTreeEditorData& EditorData = *Cast<UStateTreeEditorData>(StateTree1.EditorData);
+			UStateTreeState& Root1 = EditorData.AddSubTree(FName("Root"));
+			UStateTreeState& State1 = Root1.AddChildState(FName("State1"), EStateTreeStateType::LinkedAsset);
+
+			State1.SetLinkedStateAsset(&StateTree2);
+
+			FStateTreeCompilerLog Log;
+			FStateTreeCompiler Compiler(Log);
+			const bool bResult = Compiler.Compile(StateTree1);
+			AITEST_TRUE("StateTree should get compiled", bResult);
+		}
+
+		FStateTreeInstanceData InstanceData;
+		{
+			FTestStateTreeExecutionContext Exec(StateTree1, StateTree1, InstanceData);
+			const bool bInitSucceeded = Exec.IsValid();
+			AITEST_TRUE("StateTree should init", bInitSucceeded);
+
+			// Start and enter state
+			EStateTreeRunStatus Status = Exec.Start();
+			AITEST_TRUE(TEXT("Valid states"), Exec.ExpectInActiveStates("Root", "State1", "Root1", "State1"));
+			AITEST_TRUE(TEXT("Expected amount of frames."), InstanceData.GetExecutionState()->ActiveFrames.Num() == 2);
+			Exec.LogClear();
+		}
+			{
+			FTestStateTreeExecutionContext Exec(StateTree1, StateTree1, InstanceData);
+			EStateTreeRunStatus Status = Exec.Tick(0.f);
+			AITEST_TRUE(TEXT("Active states should stay the same"), Exec.ExpectInActiveStates("Root", "State1", "Root2", "State2"));
+			AITEST_TRUE(TEXT("Expected amount of frames."), InstanceData.GetExecutionState()->ActiveFrames.Num() == 2);
+			Exec.LogClear();
+		}
+		{
+			FTestStateTreeExecutionContext Exec(StateTree1, StateTree1, InstanceData);
+			Exec.Stop();
+		}
+
+		return true;
+	}
+};
+IMPLEMENT_STATE_TREE_INSTANT_TEST(FStateTreeTest_TransitionLinkedAssetWith2Root, "System.StateTree.Transition.LinkedAssetWith2Root");
+
 } // namespace UE::StateTree::Tests
 
 UE_ENABLE_OPTIMIZATION_SHIP
