[StateTree] Make sure the compiler resets the data on failure. Remove the schema variable to be explicit about which schema we want to use.

diff --git a/Engine/Plugins/Runtime/StateTree/Source/StateTreeEditorModule/Private/StateTreeCompiler.cpp b/Engine/Plugins/Runtime/StateTree/Source/StateTreeEditorModule/Private/StateTreeCompiler.cpp
--- a/Engine/Plugins/Runtime/StateTree/Source/StateTreeEditorModule/Private/StateTreeCompiler.cpp
+++ b/Engine/Plugins/Runtime/StateTree/Source/StateTreeEditorModule/Private/StateTreeCompiler.cpp
@@ -325,6 +325,11 @@ namespace UE::StateTree::Compiler
 }; // UE::StateTree::Compiler
 
 bool FStateTreeCompiler::Compile(UStateTree& InStateTree)
+{
+	return Compile(&InStateTree);
+}
+
+bool FStateTreeCompiler::Compile(TNotNull<UStateTree*> InStateTree)
 {
 	if (bCompiled)
 	{
@@ -333,11 +338,18 @@ bool FStateTreeCompiler::Compile(UStateTree& InStateTree)
 	}
 	bCompiled = true;
 
-	StateTree = &InStateTree;
+	StateTree = InStateTree;
 	EditorData = Cast<UStateTreeEditorData>(StateTree->EditorData);
+
+	auto FailCompilation = [StateTree=StateTree]()
+		{
+			StateTree->ResetCompiled();
+			return false;
+		};
+
 	if (!EditorData)
 	{
-		return false;
+		return FailCompilation();
 	}
 	
 	// Cleanup existing state
@@ -346,14 +358,12 @@ bool FStateTreeCompiler::Compile(UStateTree& InStateTree)
 	if (!EditorData->Schema)
 	{
 		Log.Reportf(EMessageSeverity::Error, TEXT("Missing Schema. Please set valid schema in the State Tree Asset settings."));
-		return false;
+		return FailCompilation();
 	}
-	Schema = EditorData->Schema;
 
 	if (!BindingsCompiler.Init(StateTree->PropertyBindings, Log))
 	{
-		StateTree->ResetCompiled();
-		return false;
+		return FailCompilation();
 	}
 
 	EditorData->GetAllStructValues(IDToStructValue);
@@ -363,8 +373,7 @@ bool FStateTreeCompiler::Compile(UStateTree& InStateTree)
 
 	if (!CreateParameters())
 	{
-		StateTree->ResetCompiled();
-		return false;
+		return FailCompilation();
 	}
 
 	int32 ContextDataIndex = 0;
@@ -389,7 +398,7 @@ bool FStateTreeCompiler::Compile(UStateTree& InStateTree)
 				if (const auto Validation = UE::StateTree::Compiler::IsValidIndex16(ContextDataIndex); Validation.DidFail())
 				{
 					Validation.Log(Log, TEXT("ExternalStructIndex"), ExtDataDesc);
-					return false;
+					return FailCompilation();
 				}
 				Desc.Handle.DataHandle = ExtDataDesc.DataHandle;
 			}
@@ -399,14 +408,13 @@ bool FStateTreeCompiler::Compile(UStateTree& InStateTree)
 	if (const UE::StateTree::Compiler::FValidationResult Validation = UE::StateTree::Compiler::IsValidIndex16(ContextDataIndex); Validation.DidFail())
 	{
 		Validation.Log(Log, TEXT("NumContextData"));
-		return false;
+		return FailCompilation();
 	}
 	StateTree->NumContextData = static_cast<uint16>(ContextDataIndex);
 	
 	if (!CreateStates())
 	{
-		StateTree->ResetCompiled();
-		return false;
+		return FailCompilation();
 	}
 
 	// Eval and Global task methods use InstanceStructs.Num() as ID generator.
@@ -414,40 +422,35 @@ bool FStateTreeCompiler::Compile(UStateTree& InStateTree)
 	
 	if (!CreateEvaluators())
 	{
-		StateTree->ResetCompiled();
-		return false;
+		return FailCompilation();
 	}
 
 	if (!CreateGlobalTasks())
 	{
-		StateTree->ResetCompiled();
-		return false;
+		return FailCompilation();
 	}
 
 	const int32 NumGlobalInstanceData = InstanceStructs.Num();
 	if (const auto Validation = UE::StateTree::Compiler::IsValidIndex16(NumGlobalInstanceData); Validation.DidFail())
 	{
 		Validation.Log(Log, TEXT("NumGlobalInstanceData"));
-		return false;
+		return FailCompilation();
 	}
 	StateTree->NumGlobalInstanceData = uint16(NumGlobalInstanceData);
 
 	if (!CreateStateTasksAndParameters())
 	{
-		StateTree->ResetCompiled();
-		return false;
+		return FailCompilation();
 	}
 
 	if (!CreateStateTransitions())
 	{
-		StateTree->ResetCompiled();
-		return false;
+		return FailCompilation();
 	}
 
 	if (!CreateStateConsiderations())
 	{
-		StateTree->ResetCompiled();
-		return false;
+		return FailCompilation();
 	}
 
 	StateTree->Nodes = Nodes;
@@ -462,9 +465,8 @@ bool FStateTreeCompiler::Compile(UStateTree& InStateTree)
 
 	if (!StateTree->Link())
 	{
-		StateTree->ResetCompiled();
 		Log.Reportf(EMessageSeverity::Error, TEXT("Unexpected failure to link the StateTree asset. See log for more info."));
-		return false;
+		return FailCompilation();
 	}
 
 	// Store mapping between node unique ID and their compiled index. Used for debugging purposes.
@@ -487,8 +489,7 @@ bool FStateTreeCompiler::Compile(UStateTree& InStateTree)
 
 	if (!NotifyInternalPost())
 	{
-		StateTree->ResetCompiled();
-		return false;
+		return FailCompilation();
 	}
 
 	UE::StateTree::Compiler::FCheckOutersArchive CheckOuters(*StateTree, *EditorData, Log);
@@ -619,7 +620,7 @@ bool FStateTreeCompiler::CreateStates()
 bool FStateTreeCompiler::CreateStateRecursive(UStateTreeState& State, const FStateTreeStateHandle Parent)
 {
 	check(StateTree);
-	check(Schema);
+	check(StateTree->Schema);
 
 	FStateTreeCompilerLogStateScope LogStateScope(&State, Log);
 
@@ -640,7 +641,7 @@ bool FStateTreeCompiler::CreateStateRecursive(UStateTreeState& State, const FSta
 	CompactState.bCheckPrerequisitesWhenActivatingChildDirectly = State.bCheckPrerequisitesWhenActivatingChildDirectly;
 	CompactState.Weight = State.Weight;
 
-	CompactState.bHasCustomTickRate = State.bHasCustomTickRate && Schema->IsScheduledTickAllowed();
+	CompactState.bHasCustomTickRate = State.bHasCustomTickRate && StateTree->Schema->IsScheduledTickAllowed();
 	CompactState.CustomTickRate = FMath::Max(State.CustomTickRate, 0.0f);
 	if (CompactState.bHasCustomTickRate && State.CustomTickRate < 0.0f)
 	{
@@ -650,7 +651,7 @@ bool FStateTreeCompiler::CreateStateRecursive(UStateTreeState& State, const FSta
 	CompactState.Type = State.Type;
 	CompactState.SelectionBehavior = State.SelectionBehavior;
 
-	if (!Schema->IsStateTypeAllowed(CompactState.Type))
+	if (!StateTree->Schema->IsStateTypeAllowed(CompactState.Type))
 	{
 		Log.Reportf(EMessageSeverity::Warning,
 			TEXT("The State '%s' has a restricted type for the schema."),
@@ -663,7 +664,7 @@ bool FStateTreeCompiler::CreateStateRecursive(UStateTreeState& State, const FSta
 	{
 		CompactState.SelectionBehavior = EStateTreeStateSelectionBehavior::TryEnterState;
 	}
-	else if (!Schema->IsStateSelectionAllowed(CompactState.SelectionBehavior))
+	else if (!StateTree->Schema->IsStateSelectionAllowed(CompactState.SelectionBehavior))
 	{
 		Log.Reportf(EMessageSeverity::Warning,
 			TEXT("The State '%s' has a restricted selection behavior for the schema."),
@@ -825,7 +826,7 @@ bool FStateTreeCompiler::CreateGlobalTasks()
 	const UE::StateTree::Compiler::FCompletionTasksMaskResult MaskResult = UE::StateTree::Compiler::MakeCompletionTasksMask(CompletionGlobalTaskStartBitIndex, EditorData->GlobalTasks, ValidTaskNodeIndex);
 	StateTree->CompletionGlobalTasksMask = MaskResult.Mask;
 	GlobalTaskEndBit = MaskResult.FullMaskEndTaskBitOffset;
-	StateTree->CompletionGlobalTasksControl = (Schema && Schema->AllowTasksCompletion()) ? EditorData->GlobalTasksCompletion : EStateTreeTaskCompletionType::Any;
+	StateTree->CompletionGlobalTasksControl = StateTree->Schema->AllowTasksCompletion() ? EditorData->GlobalTasksCompletion : EStateTreeTaskCompletionType::Any;
 
 	if (MaskResult.MaskFirstTaskBitOffset != 0)
 	{
@@ -1107,7 +1108,7 @@ bool FStateTreeCompiler::CreateStateTasksAndParameters()
 			}
 
 			CompactState.CompletionTasksMask = MaskResult.Mask;
-			CompactState.CompletionTasksControl = (Schema && Schema->AllowTasksCompletion()) ? State->TasksCompletion : EStateTreeTaskCompletionType::Any;
+			CompactState.CompletionTasksControl = StateTree->Schema->AllowTasksCompletion() ? State->TasksCompletion : EStateTreeTaskCompletionType::Any;
 			CompactState.CompletionTasksMaskBufferIndex = static_cast<uint8>(CompletionTasksMaskBufferIndex);
 			CompactState.CompletionTasksMaskBitsOffset = static_cast<uint8>(CompletionTasksMaskBitsOffset);
 			StateInfos[StateIndex].NextBitIndexForCompletionMask = MaskResult.FullMaskEndTaskBitOffset;
@@ -1291,14 +1292,14 @@ bool FStateTreeCompiler::CreateStateTransitions()
 					continue;
 				}
 			
-				check(Schema);
-				if (LinkedAssetSchema->GetClass() != Schema->GetClass())
+				check(StateTree->Schema);
+				if (LinkedAssetSchema->GetClass() != StateTree->Schema->GetClass())
 				{
 					Log.Reportf(EMessageSeverity::Error,
 						TEXT("Linked State Tree asset '%s' must have same schema class as this asset. Linked asset has '%s', expected '%s'."),
 						*GetFullNameSafe(SourceState->LinkedAsset),
 						*LinkedAssetSchema->GetClass()->GetDisplayNameText().ToString(),
-						*Schema->GetClass()->GetDisplayNameText().ToString()
+						*StateTree->Schema->GetClass()->GetDisplayNameText().ToString()
 					);
 					bSucceeded = false;
 					continue;
diff --git a/Engine/Plugins/Runtime/StateTree/Source/StateTreeEditorModule/Public/StateTreeCompiler.h b/Engine/Plugins/Runtime/StateTree/Source/StateTreeEditorModule/Public/StateTreeCompiler.h
--- a/Engine/Plugins/Runtime/StateTree/Source/StateTreeEditorModule/Public/StateTreeCompiler.h
+++ b/Engine/Plugins/Runtime/StateTree/Source/StateTreeEditorModule/Public/StateTreeCompiler.h
@@ -54,13 +54,13 @@ public:
 		: Log(InLog)
 	{
 	}
-	
+
+	UE_API bool Compile(TNotNull<UStateTree*> StateTree);
 	UE_API bool Compile(UStateTree& InStateTree);
-	
+
 private:
 	/** Resolves the state a transition points to, and the optional fallback for failing to enter the state. SourceState is nullptr for global tasks. */
 	bool ResolveTransitionStateAndFallback(const UStateTreeState* SourceState, const FStateTreeStateLink& Link, FStateTreeStateHandle& OutTransitionHandle, EStateTreeSelectionFallback& OutFallback) const;
-
 	FStateTreeStateHandle GetStateHandle(const FGuid& StateID) const;
 	UStateTreeState* GetState(const FGuid& StateID) const;
 
@@ -110,14 +110,13 @@ private:
 
 private:
 	FStateTreeCompilerLog& Log;
-	UStateTree* StateTree = nullptr;
-	UStateTreeEditorData* EditorData = nullptr;
-	UStateTreeSchema* Schema = nullptr;
+	TObjectPtr<UStateTree> StateTree = nullptr;
+	TObjectPtr<UStateTreeEditorData> EditorData = nullptr;
 	TMap<FGuid, int32> IDToNode;
 	TMap<FGuid, int32> IDToState;
 	TMap<FGuid, int32> IDToTransition;
 	TMap<FGuid, const FStateTreeDataView > IDToStructValue;
-	TArray<UStateTreeState*> SourceStates;
+	TArray<TObjectPtr<UStateTreeState>> SourceStates;
 
 	TArray<FInstancedStruct> Nodes;
 	TArray<FInstancedStruct> InstanceStructs;
