[StateTree] Gather dependencies when StateTree compiles. It will be used to trigger StateTree compilation when needed. The feature can be disabled with cvar StateTree.Compiler.UseDependenciesToTriggerCompilation. For debugging, use the cvar StateTree.Compiler.LogDependenciesOnCompilation for logging the gathered dependencies.

diff --git a/Engine/Plugins/Runtime/StateTree/Source/StateTreeEditorModule/Private/StateTreeCompilerManager.cpp b/Engine/Plugins/Runtime/StateTree/Source/StateTreeEditorModule/Private/StateTreeCompilerManager.cpp
--- a/Engine/Plugins/Runtime/StateTree/Source/StateTreeEditorModule/Private/StateTreeCompilerManager.cpp
+++ b/Engine/Plugins/Runtime/StateTree/Source/StateTreeEditorModule/Private/StateTreeCompilerManager.cpp
@@ -6,6 +6,8 @@
 #include "StateTreeDelegates.h"
 #include "StateTreeEditingSubsystem.h"
 #include "StateTreeEditorData.h"
+#include "StateTreeEditorModule.h"
+#include "StateTreeEditorPropertyBindings.h"
 
 #include "Editor.h"
 #include "StructUtilsDelegates.h"
@@ -16,6 +18,110 @@
 
 namespace UE::StateTree::Compiler::Private
 {
+FAutoConsoleVariable CVarLogStateTreeDependencies(
+	TEXT("StateTree.Compiler.LogDependenciesOnCompilation"),
+	false,
+	TEXT("After a StateTree compiles, log the dependencies that will be required for the asset to recompile.")
+);
+
+bool bUseDependenciesToTriggerCompilation = true;
+FAutoConsoleVariableRef CVarLogStateTreeUseDependenciesToTriggerCompilation(
+	TEXT("StateTree.Compiler.UseDependenciesToTriggerCompilation"),
+	bUseDependenciesToTriggerCompilation,
+	TEXT("Use the build dependencies to detect when a state tree needs to be linked or compiled.")
+);
+
+struct FStateTreeDependencies
+{
+	enum EDependencyType
+	{
+		DT_None = 0,
+		DT_Link = 1 << 0,
+		DT_Internal = 1<< 1,
+		DT_Public = 1 << 2,
+	};
+	struct FItem
+	{
+		FObjectKey Key;
+		EDependencyType Type = EDependencyType::DT_None;
+	};
+	TArray<FItem> Dependencies;
+};
+ENUM_CLASS_FLAGS(FStateTreeDependencies::EDependencyType);
+
+/** Find the references that are needed by the asset. */
+class FArchiveReferencingProperties : public FArchiveUObject
+{
+public:
+	FArchiveReferencingProperties(TNotNull<const UObject*> InReferencingObject)
+		: ReferencingObjectPackage(InReferencingObject->GetPackage())
+		, StateTreeModulePackage(UStateTree::StaticClass()->GetOutermost())
+		, StateTreeEditorModulePackage(UStateTreeEditorData::StaticClass()->GetOutermost())
+		, CoreUObjectModulePackage(UObject::StaticClass()->GetOutermost())
+	{
+		ArIsObjectReferenceCollector = true;
+		ArIgnoreOuterRef = true;
+		ArIgnoreArchetypeRef = true;
+		ArIgnoreClassGeneratedByRef = true;
+		ArIgnoreClassRef = true;
+
+		SetShouldSkipCompilingAssets(false);
+	}
+
+	bool IsSupportedObject(TNotNull<const UStruct*> Struct)
+	{
+		/** As an optimization, do not include basic structures like FVector and state tree internal types. */
+		return !Struct->IsInPackage(StateTreeModulePackage)
+			&& !Struct->IsInPackage(StateTreeEditorModulePackage)
+			&& !Struct->IsInPackage(CoreUObjectModulePackage);
+	}
+
+	virtual FArchive& operator<<(UObject*& InSerializedObject) override
+	{
+		if (InSerializedObject)
+		{
+			if (const UStruct* AsStruct = Cast<const UStruct>(InSerializedObject))
+			{
+				if (IsSupportedObject(AsStruct))
+				{
+					Dependencies.Add(AsStruct);
+				}
+			}
+			else
+			{
+				if (IsSupportedObject(InSerializedObject->GetClass()))
+				{
+					Dependencies.Add(InSerializedObject->GetClass());
+				}
+			}
+
+			// Traversing the asset inner dependencies (instanced type).
+			if (InSerializedObject->IsInPackage(ReferencingObjectPackage))
+			{
+				bool bAlreadyExists;
+				SerializedObjects.Add(InSerializedObject, &bAlreadyExists);
+
+				if (!bAlreadyExists)
+				{
+					InSerializedObject->Serialize(*this);
+				}
+			}
+		}
+
+		return *this;
+	}
+
+	TSet<TNotNull<const UStruct*>> Dependencies;
+
+private:
+	/** Tracks the objects which have been serialized by this archive, to prevent recursion */
+	TSet<UObject*> SerializedObjects;
+
+	TNotNull<const UPackage*> ReferencingObjectPackage;
+	TNotNull<const UPackage*> StateTreeModulePackage;
+	TNotNull<const UPackage*> StateTreeEditorModulePackage;
+	TNotNull<const UPackage*> CoreUObjectModulePackage;
+};
 
 class FCompilerManagerImpl
 {
@@ -25,9 +131,13 @@ public:
 	FCompilerManagerImpl(const FCompilerManagerImpl&) = delete;
 	FCompilerManagerImpl& operator=(const FCompilerManagerImpl&) = delete;
 
+	bool CompileInternalSynchronously(TNotNull<UStateTree*> InStateTree, FStateTreeCompilerLog& InOutLog);
+
 private:
 	bool HandleCompileStateTree(UStateTree& StateTree);
 	void UpdateBindingsInstanceStructsIfNeeded(TSet<const UStruct*>& Structs, TNotNull<UStateTreeEditorData*> EditorData);
+	void GatherDependencies(TNotNull<UStateTree*> StateTree);
+	void LogDependencies(TNotNull<UStateTree*> StateTree) const;
 
 	using FReplacementObjectMap = TMap<UObject*, UObject*>;
 	void HandleObjectsReinstanced(const FReplacementObjectMap& ObjectMap);
@@ -38,6 +148,9 @@ private:
 	FDelegateHandle ObjectsReinstancedHandle;
 	FDelegateHandle UserDefinedStructReinstancedHandle;
 	FDelegateHandle PreBeginPIEHandle;
+
+	TMap<TObjectKey<UStateTree>, TSharedPtr<FStateTreeDependencies>> StateTreeToDependencies;
+	TMap<FObjectKey, TArray<TObjectKey<UStateTree>>> DependenciesToStateTree;
 };
 static TUniquePtr<FCompilerManagerImpl> CompilerManagerImpl;
 
@@ -57,6 +170,41 @@ FCompilerManagerImpl::~FCompilerManagerImpl()
 	UE::StateTree::Delegates::OnRequestCompile.Unbind();
 }
 
+bool FCompilerManagerImpl::CompileInternalSynchronously(TNotNull<UStateTree*> StateTree, FStateTreeCompilerLog& Log)
+{
+	UStateTreeEditingSubsystem::ValidateStateTree(StateTree);
+	FStateTreeCompiler Compiler(Log);
+	const bool bCompilationResult = Compiler.Compile(StateTree);
+	if (bCompilationResult)
+	{
+		const uint32 EditorDataHash = UStateTreeEditingSubsystem::CalculateStateTreeHash(StateTree);
+
+		// Success
+		StateTree->LastCompiledEditorDataHash = EditorDataHash;
+		UE_LOG(LogStateTreeEditor, Log, TEXT("Compile StateTree '%s' succeeded."), *StateTree->GetFullName());
+	}
+	else
+	{
+		// Make sure not to leave stale data on failed compile.
+		StateTree->ResetCompiled();
+		StateTree->LastCompiledEditorDataHash = 0;
+
+		UE_LOG(LogStateTreeEditor, Error, TEXT("Failed to compile '%s', errors follow."), *StateTree->GetFullName());
+		Log.DumpToLog(LogStateTreeEditor);
+	}
+
+	UE::StateTree::Delegates::OnPostCompile.Broadcast(*StateTree);
+
+	GatherDependencies(StateTree);
+
+	if (CVarLogStateTreeDependencies->GetBool())
+	{
+		LogDependencies(StateTree);
+	}
+
+	return bCompilationResult;
+}
+
 void FCompilerManagerImpl::UpdateBindingsInstanceStructsIfNeeded(TSet<const UStruct*>& Structs, TNotNull<UStateTreeEditorData*> EditorData)
 {
 	bool bShouldUpdate = false;
@@ -84,7 +232,7 @@ void FCompilerManagerImpl::UpdateBindingsInstanceStructsIfNeeded(TSet<const UStr
 bool FCompilerManagerImpl::HandleCompileStateTree(UStateTree& StateTree)
 {
 	FStateTreeCompilerLog Log;
-	return UStateTreeEditingSubsystem::CompileStateTree(&StateTree, Log);
+	return CompileInternalSynchronously(&StateTree, Log);
 }
 
 void FCompilerManagerImpl::HandlePreBeginPIE(const bool bIsSimulating)
@@ -113,75 +261,133 @@ void FCompilerManagerImpl::HandleObjectsReinstanced(const FReplacementObjectMap&
 		}
 	}
 
-	TSet<const UStruct*> Structs;
-	Structs.Reserve(ObjectsToBeReplaced.Num());
+	TSet<const UStruct*> StructsToBeReplaced;
+	StructsToBeReplaced.Reserve(ObjectsToBeReplaced.Num());
 	for (const UObject* ObjectToBeReplaced : ObjectsToBeReplaced)
 	{
 		// It's a UClass or a UScriptStruct
-		if (const UStruct* StructToReplaced = Cast<const UStruct>(ObjectToBeReplaced))
+		if (const UStruct* StructToBeReplaced = Cast<const UStruct>(ObjectToBeReplaced))
 		{
-			Structs.Add(StructToReplaced);
+			StructsToBeReplaced.Add(StructToBeReplaced);
 		}
 		else
 		{
-			Structs.Add(ObjectToBeReplaced->GetClass());
+			StructsToBeReplaced.Add(ObjectToBeReplaced->GetClass());
 		}
 	}
 
-	for (TObjectIterator<UStateTreeEditorData> It; It; ++It)
+	if (UE::StateTree::Compiler::Private::bUseDependenciesToTriggerCompilation)
 	{
-		UStateTreeEditorData* StateTreeEditorData = *It;
-		UpdateBindingsInstanceStructsIfNeeded(Structs, StateTreeEditorData);
-	}
+		TArray<TNotNull<UStateTree*>> StateTreeToLink;
+		for (const UStruct* StructToBeReplaced : StructsToBeReplaced)
+		{
+			const FObjectKey StructToReplacedKey = StructToBeReplaced;
+			TArray<TObjectKey<UStateTree>>* Dependencies = DependenciesToStateTree.Find(StructToReplacedKey);
+			if (Dependencies)
+			{
+				for (const TObjectKey<UStateTree>& StateTreeKey : *Dependencies)
+				{
+					if (UStateTree* StateTree = StateTreeKey.ResolveObjectPtr())
+					{
+						StateTreeToLink.AddUnique(StateTree);
+					}
+				}
+			}
+		}
 
-	for (TObjectIterator<UStateTree> It; It; ++It)
+		for (UStateTree* StateTree : StateTreeToLink)
+		{
+			if (!StateTree->Link())
+			{
+				UE_LOG(LogStateTree, Error, TEXT("%s failed to link after Object reinstantiation. Take a look at the asset for any errors. Asset will not be usable at runtime."), *StateTree->GetPathName());
+			}
+		}
+	}
+	else
 	{
-		UStateTree* StateTree = *It;
-		check(!StateTree->HasAnyFlags(RF_ClassDefaultObject));
-		bool bShouldRelink = false;
+		for (TObjectIterator<UStateTreeEditorData> It; It; ++It)
+		{
+			UStateTreeEditorData* StateTreeEditorData = *It;
+			UpdateBindingsInstanceStructsIfNeeded(StructsToBeReplaced, StateTreeEditorData);
+		}
 
-		// Relink if one of the out of date objects got reinstanced.
-		if (StateTree->OutOfDateStructs.Num() > 0)
+		for (TObjectIterator<UStateTree> It; It; ++It)
 		{
-			for (const FObjectKey& OutOfDateObjectKey : StateTree->OutOfDateStructs)
+			UStateTree* StateTree = *It;
+			check(!StateTree->HasAnyFlags(RF_ClassDefaultObject));
+			bool bShouldRelink = false;
+
+			// Relink if one of the out of date objects got reinstanced.
+			if (StateTree->OutOfDateStructs.Num() > 0)
 			{
-				if (const UObject* OutOfDateObject = OutOfDateObjectKey.ResolveObjectPtr())
+				for (const FObjectKey& OutOfDateObjectKey : StateTree->OutOfDateStructs)
 				{
-					if (ObjectMap.Contains(OutOfDateObject))
+					if (const UObject* OutOfDateObject = OutOfDateObjectKey.ResolveObjectPtr())
+					{
+						if (ObjectMap.Contains(OutOfDateObject))
+						{
+							bShouldRelink = true;
+							break;
+						}
+					}
+				}
+			}
+
+			// If the asset is not linked yet (or has failed), no need to link.
+			if (!bShouldRelink && !StateTree->bIsLinked)
+			{
+				continue;
+			}
+
+			// Relink only if the reinstantiated object belongs to this asset,
+			// or anything from the property binding refers to the classes of the reinstantiated object.
+			if (!bShouldRelink)
+			{
+				for (const UObject* ObjectToBeReplaced : ObjectsToBeReplaced)
+				{
+					if (ObjectToBeReplaced->IsInOuter(StateTree))
 					{
 						bShouldRelink = true;
 						break;
 					}
 				}
 			}
-		}
 
-		// If the asset is not linked yet (or has failed), no need to link.
-		if (!bShouldRelink && !StateTree->bIsLinked)
-		{
-			continue;
-		}
+			if (!bShouldRelink)
+			{
+				bShouldRelink |= StateTree->PropertyBindings.ContainsAnyStruct(StructsToBeReplaced);
+			}
 
-		// Relink only if the reinstantiated object belongs to this asset,
-		// or anything from the property binding refers to the classes of the reinstantiated object.
-		if (!bShouldRelink)
-		{
-			for (const UObject* ObjectToBeReplaced : ObjectsToBeReplaced)
+			if (bShouldRelink)
 			{
-				if (ObjectToBeReplaced->IsInOuter(StateTree))
+				if (!StateTree->Link())
 				{
-					bShouldRelink = true;
-					break;
+					UE_LOG(LogStateTree, Error, TEXT("%s failed to link after Object reinstantiation. Take a look at the asset for any errors. Asset will not be usable at runtime."), *StateTree->GetPathName());
 				}
 			}
 		}
+	}
+}
 
-		if (!bShouldRelink)
+void FCompilerManagerImpl::HandleUserDefinedStructReinstanced(const UUserDefinedStruct& UserDefinedStruct)
+{
+	if (UE::StateTree::Compiler::Private::bUseDependenciesToTriggerCompilation)
+	{
+		TSet<TNotNull<UStateTree*>> StateTreeToLink;
+		const FObjectKey StructToReplacedKey = &UserDefinedStruct;
+		TArray<TObjectKey<UStateTree>>* Dependencies = DependenciesToStateTree.Find(StructToReplacedKey);
+		if (Dependencies)
 		{
-			bShouldRelink |= StateTree->PropertyBindings.ContainsAnyStruct(Structs);
+			for (const TObjectKey<UStateTree>& StateTreeKey : *Dependencies)
+			{
+				if (UStateTree* StateTree = StateTreeKey.ResolveObjectPtr())
+				{
+					StateTreeToLink.Add(StateTree);
+				}
+			}
 		}
 
-		if (bShouldRelink)
+		for (UStateTree* StateTree : StateTreeToLink)
 		{
 			if (!StateTree->Link())
 			{
@@ -189,32 +395,201 @@ void FCompilerManagerImpl::HandleObjectsReinstanced(const FReplacementObjectMap&
 			}
 		}
 	}
+	else
+	{
+		TSet<const UStruct*> Structs;
+		Structs.Add(&UserDefinedStruct);
+
+		for (TObjectIterator<UStateTreeEditorData> It; It; ++It)
+		{
+			UStateTreeEditorData* StateTreeEditorData = *It;
+			UpdateBindingsInstanceStructsIfNeeded(Structs, StateTreeEditorData);
+		}
+
+		for (TObjectIterator<UStateTree> It; It; ++It)
+		{
+			UStateTree* StateTree = *It;
+			if (StateTree->PropertyBindings.ContainsAnyStruct(Structs))
+			{
+				if (!StateTree->Link())
+				{
+					UE_LOG(LogStateTree, Error, TEXT("%s failed to link after Struct reinstantiation. Take a look at the asset for any errors. Asset will not be usable at runtime."), *StateTree->GetPathName());
+				}
+			}
+		}
+	}
 }
 
-void FCompilerManagerImpl::HandleUserDefinedStructReinstanced(const UUserDefinedStruct& UserDefinedStruct)
+void FCompilerManagerImpl::GatherDependencies(TNotNull<UStateTree*> StateTree)
 {
-	TSet<const UStruct*> Structs;
-	Structs.Add(&UserDefinedStruct);
+	// Find the tree in the StateTreeToDependencies
+	const TObjectKey<UStateTree> StateTreeKey = StateTree;
+	TSharedPtr<FStateTreeDependencies>& FoundDependencies = StateTreeToDependencies.FindOrAdd(StateTreeKey);
+	
+	// Remove all from DependenciesToStateTree
+	if (FoundDependencies)
+	{
+		for (FStateTreeDependencies::FItem& Item : FoundDependencies->Dependencies)
+		{
+			TArray<TObjectKey<UStateTree>>* FoundKey = DependenciesToStateTree.Find(Item.Key);
+			if (FoundKey)
+			{
+				FoundKey->RemoveSingleSwap(StateTreeKey);
+			}
+		}
+		FoundDependencies->Dependencies.Reset();
+	}
+	else
+	{
+		FoundDependencies = MakeShared<FStateTreeDependencies>();
+	}
+
+	auto AddDependency = [this, StateTreeKey, FoundDependencies](TNotNull<const UStruct*> Object, FStateTreeDependencies::EDependencyType DependencyType)
+		{
+			const FObjectKey ObjectKey = Object;
+			DependenciesToStateTree.FindOrAdd(ObjectKey).AddUnique(StateTreeKey);
 
-	for (TObjectIterator<UStateTreeEditorData> It; It; ++It)
+			if (FStateTreeDependencies::FItem* FoundItem = FoundDependencies->Dependencies.FindByPredicate([ObjectKey](const FStateTreeDependencies::FItem& Other)
+				{
+					return Other.Key == ObjectKey;
+				}))
+			{
+				FoundItem->Type |= DependencyType;
+			}
+			else
+			{
+				FoundDependencies->Dependencies.Add(FStateTreeDependencies::FItem{.Key = ObjectKey, .Type = DependencyType});
+			}
+		};
+
+	// Gather new inner dependencies
+	UStateTreeEditorData* EditorData = CastChecked<UStateTreeEditorData>(StateTree->EditorData);
+	if (EditorData)
 	{
-		UStateTreeEditorData* StateTreeEditorData = *It;
-		UpdateBindingsInstanceStructsIfNeeded(Structs, StateTreeEditorData);
+		// Internal
+		{
+			FArchiveReferencingProperties DependencyArchive(StateTree);
+			EditorData->Serialize(DependencyArchive);
+			for (const UStruct* Dependency : DependencyArchive.Dependencies)
+			{
+				AddDependency(Dependency, FStateTreeDependencies::EDependencyType::DT_Internal);
+			}
+		}
+		// Public
+		{
+			FArchiveReferencingProperties DependencyArchive(StateTree);
+			const_cast<FInstancedPropertyBag&>(EditorData->GetRootParametersPropertyBag()).Serialize(DependencyArchive);
+			for (const UStruct* Dependency : DependencyArchive.Dependencies)
+			{
+				AddDependency(Dependency, FStateTreeDependencies::EDependencyType::DT_Public);
+			}
+			if (EditorData->Schema)
+			{
+				AddDependency(EditorData->Schema->GetClass(), FStateTreeDependencies::EDependencyType::DT_Public);
+			}
+		}
+		// Link
+		{
+			TMap<FGuid, const FPropertyBindingDataView> AllStructValues;
+			EditorData->GetAllStructValues(AllStructValues);
+			auto AddBindingPathDependencies = [&AddDependency , &AllStructValues](const FPropertyBindingPath& PropertyPath)
+				{
+					const FPropertyBindingDataView* FoundStruct = AllStructValues.Find(PropertyPath.GetStructID());
+					if (FoundStruct)
+					{
+						FString Error;
+						TArray<FPropertyBindingPathIndirection> Indirections;
+						if (PropertyPath.ResolveIndirectionsWithValue(*FoundStruct, Indirections, &Error))
+						{
+							for (const FPropertyBindingPathIndirection& Indirection : Indirections)
+							{
+								if (Indirection.GetInstanceStruct())
+								{
+									AddDependency(Indirection.GetInstanceStruct(), FStateTreeDependencies::EDependencyType::DT_Link);
+								}
+								else if (Indirection.GetContainerStruct())
+								{
+									AddDependency(Indirection.GetContainerStruct(), FStateTreeDependencies::EDependencyType::DT_Link);
+								}
+							}
+						}
+					}
+				};
+
+			EditorData->GetEditorPropertyBindings()->VisitBindings([&AddBindingPathDependencies](const FPropertyBindingBinding& Binding)
+				{
+					AddBindingPathDependencies(Binding.GetSourcePath());
+					AddBindingPathDependencies(Binding.GetTargetPath());
+					return FPropertyBindingBindingCollection::EVisitResult::Continue;
+				});
+		}
 	}
+}
 
-	for (TObjectIterator<UStateTree> It; It; ++It)
+void FCompilerManagerImpl::LogDependencies(TNotNull<UStateTree*> StateTree) const
+{
+	FStringBuilderBase LogString;
+	LogString << TEXT("StateTree Dependencies (asset: '");
+	StateTree->GetFullName(LogString);
+	LogString << TEXT("')\n");
+	
+	const TObjectKey<UStateTree> StateTreeKey = StateTree;
+	const TSharedPtr<FStateTreeDependencies>* FoundDependencies = StateTreeToDependencies.Find(StateTreeKey);
+	if (FoundDependencies != nullptr && FoundDependencies->IsValid())
 	{
-		UStateTree* StateTree = *It;
-		if (StateTree->PropertyBindings.ContainsAnyStruct(Structs))
+		auto PrintType = [&LogString](FStateTreeDependencies::EDependencyType Type)
+			{
+				bool bPrinted = false;
+				auto PrintSeparator = [&bPrinted, &LogString]()
+					{
+						if (bPrinted)
+						{
+							LogString << TEXT(" | ");
+						}
+						bPrinted = true;
+					};
+				if (EnumHasAnyFlags(Type, FStateTreeDependencies::EDependencyType::DT_Public))
+				{
+					PrintSeparator();
+					LogString << TEXT("Public");
+				}
+				if (EnumHasAnyFlags(Type, FStateTreeDependencies::EDependencyType::DT_Internal))
+				{
+					PrintSeparator();
+					LogString << TEXT("Internal");
+				}
+				if (EnumHasAnyFlags(Type, FStateTreeDependencies::EDependencyType::DT_Link))
+				{
+					PrintSeparator();
+					LogString << TEXT("Link");
+				}
+			};
+		for (const FStateTreeDependencies::FItem& Item : (*FoundDependencies)->Dependencies)
 		{
-			if (!StateTree->Link())
+			LogString << TEXT("  ");
+			if (const UObject* Object = Item.Key.ResolveObjectPtr())
 			{
-				UE_LOG(LogStateTree, Error, TEXT("%s failed to link after Struct reinstantiation. Take a look at the asset for any errors. Asset will not be usable at runtime."), *StateTree->GetPathName());
+				Object->GetFullName(LogString);
 			}
+			else
+			{
+				LogString << TEXT(" [None]");
+			}
+
+			LogString << TEXT(" [");
+			PrintType(Item.Type);
+			LogString << TEXT("]\n");
 		}
 	}
+	else
+	{
+		LogString << TEXT("  No Dependency");
+	}
+
+	UE_LOG(LogStateTreeEditor, Log, TEXT("%s"), LogString.ToString());
 }
 
+
 } // namespace UE::StateTree::Compiler::Private
 
 
@@ -231,5 +606,20 @@ void FCompilerManager::Shutdown()
 	Private::CompilerManagerImpl.Reset();
 }
 
+bool FCompilerManager::CompileSynchronously(TNotNull<UStateTree*> StateTree)
+{
+	FStateTreeCompilerLog Log;
+	return CompileSynchronously(StateTree, Log);
+}
+
+bool FCompilerManager::CompileSynchronously(TNotNull<UStateTree*> StateTree, FStateTreeCompilerLog& Log)
+{
+	if (ensureMsgf(Private::CompilerManagerImpl.IsValid(), TEXT("Can't compile the asset when the module is not available.")))
+	{
+		return Private::CompilerManagerImpl->CompileInternalSynchronously(StateTree, Log);
+	}
+	return false;
+}
+
 } // UE::StateTree::Compiler
 
diff --git a/Engine/Plugins/Runtime/StateTree/Source/StateTreeEditorModule/Private/StateTreeEditingSubsystem.cpp b/Engine/Plugins/Runtime/StateTree/Source/StateTreeEditorModule/Private/StateTreeEditingSubsystem.cpp
--- a/Engine/Plugins/Runtime/StateTree/Source/StateTreeEditorModule/Private/StateTreeEditingSubsystem.cpp
+++ b/Engine/Plugins/Runtime/StateTree/Source/StateTreeEditorModule/Private/StateTreeEditingSubsystem.cpp
@@ -5,6 +5,7 @@
 #include "SStateTreeView.h"
 #include "StateTreeCompiler.h"
 #include "StateTreeCompilerLog.h"
+#include "StateTreeCompilerManager.h"
 #include "StateTreeDelegates.h"
 #include "StateTreeEditorData.h"
 #include "StateTreeEditorModule.h"
@@ -19,50 +20,35 @@
 UStateTreeEditingSubsystem::UStateTreeEditingSubsystem()
 {
 	PostGarbageCollectHandle = FCoreUObjectDelegates::GetPostGarbageCollect().AddUObject(this, &UStateTreeEditingSubsystem::HandlePostGarbageCollect);
+	PostCompileHandle = UE::StateTree::Delegates::OnPostCompile.AddUObject(this, &UStateTreeEditingSubsystem::HandlePostCompile);
 }
 
 void UStateTreeEditingSubsystem::BeginDestroy()
 {
+	UE::StateTree::Delegates::OnPostCompile.Remove(PostCompileHandle);
 	FCoreUObjectDelegates::GetPostGarbageCollect().Remove(PostGarbageCollectHandle);
 	Super::BeginDestroy();
 }
 
 bool UStateTreeEditingSubsystem::CompileStateTree(TNotNull<UStateTree*> InStateTree, FStateTreeCompilerLog& InOutLog)
 {
-	ValidateStateTree(InStateTree);
-	
-	FStateTreeCompiler Compiler(InOutLog);
-
-	const bool bCompilationResult = Compiler.Compile(*InStateTree);
-	if(bCompilationResult)
-	{
-		const uint32 EditorDataHash = CalculateStateTreeHash(InStateTree);
-
-		// Success
-		InStateTree->LastCompiledEditorDataHash = EditorDataHash;
-		UE::StateTree::Delegates::OnPostCompile.Broadcast(*InStateTree);
-		UE_LOG(LogStateTreeEditor, Log, TEXT("Compile StateTree '%s' succeeded."), *InStateTree->GetFullName());
-	}
-	else
-	{
-		// Make sure not to leave stale data on failed compile.
-		InStateTree->ResetCompiled();
-		InStateTree->LastCompiledEditorDataHash = 0;
-		
-		UE_LOG(LogStateTreeEditor, Error, TEXT("Failed to compile '%s', errors follow."), *InStateTree->GetFullName());
-		InOutLog.DumpToLog(LogStateTreeEditor);
-	}
+	return UE::StateTree::Compiler::FCompilerManager::CompileSynchronously(InStateTree, InOutLog);
+}
 
-	// Notify the UI that something changed. Make sure to not request a new viewmodel. That way, we are not creating new viewmodel when cooking/PIE.
-	if (UStateTreeEditingSubsystem* StateTreeEditingSubsystem = GEditor ? GEditor->GetEditorSubsystem<UStateTreeEditingSubsystem>() : nullptr)
+TSharedPtr<FStateTreeViewModel> UStateTreeEditingSubsystem::FindViewModel(TNotNull<const UStateTree*> InStateTree) const
+{
+	const FObjectKey StateTreeKey = InStateTree;
+	TSharedPtr<FStateTreeViewModel> ViewModelPtr = StateTreeViewModels.FindRef(StateTreeKey);
+	if (ViewModelPtr)
 	{
-		if (TSharedPtr<FStateTreeViewModel> ViewModel = StateTreeEditingSubsystem->StateTreeViewModels.FindRef(FObjectKey(InStateTree)))
+		// The StateTree could be re-instantiated. Can occur when the object is destroyed and recreated in a pool or when reloaded in editor.
+		//The object might have the same pointer value or the same path but it's a new object and all weakptr are now invalid.
+		if (ViewModelPtr->GetStateTree() == InStateTree)
 		{
-			ViewModel->NotifyAssetChangedExternally();
+			return ViewModelPtr.ToSharedRef();
 		}
 	}
-
-	return bCompilationResult;
+	return nullptr;
 }
 
 TSharedRef<FStateTreeViewModel> UStateTreeEditingSubsystem::FindOrAddViewModel(TNotNull<UStateTree*> InStateTree)
@@ -333,3 +319,12 @@ void UStateTreeEditingSubsystem::HandlePostGarbageCollect()
 		}
 	}
 }
+
+void UStateTreeEditingSubsystem::HandlePostCompile(const UStateTree& InStateTree)
+{
+	// Notify the UI that something changed. Make sure to not request a new viewmodel. That way, we are not creating new viewmodel when cooking/PIE.
+	if (TSharedPtr<FStateTreeViewModel> ViewModel = FindViewModel(&InStateTree))
+	{
+		ViewModel->NotifyAssetChangedExternally();
+	}
+}
diff --git a/Engine/Plugins/Runtime/StateTree/Source/StateTreeEditorModule/Public/StateTreeCompilerManager.h b/Engine/Plugins/Runtime/StateTree/Source/StateTreeEditorModule/Public/StateTreeCompilerManager.h
--- a/Engine/Plugins/Runtime/StateTree/Source/StateTreeEditorModule/Public/StateTreeCompilerManager.h
+++ b/Engine/Plugins/Runtime/StateTree/Source/StateTreeEditorModule/Public/StateTreeCompilerManager.h
@@ -17,6 +17,9 @@ class FCompilerManager final
 public:
 	static UE_API void Startup();
 	static UE_API void Shutdown();
+
+	static UE_API bool CompileSynchronously(TNotNull<UStateTree*> StateTree);
+	static UE_API bool CompileSynchronously(TNotNull<UStateTree*> StateTree, FStateTreeCompilerLog& Log);
 	
 private:
 	FCompilerManager() = delete;
diff --git a/Engine/Plugins/Runtime/StateTree/Source/StateTreeEditorModule/Public/StateTreeEditingSubsystem.h b/Engine/Plugins/Runtime/StateTree/Source/StateTreeEditorModule/Public/StateTreeEditingSubsystem.h
--- a/Engine/Plugins/Runtime/StateTree/Source/StateTreeEditorModule/Public/StateTreeEditingSubsystem.h
+++ b/Engine/Plugins/Runtime/StateTree/Source/StateTreeEditorModule/Public/StateTreeEditingSubsystem.h
@@ -24,6 +24,7 @@ public:
 	UE_API UStateTreeEditingSubsystem();
 	UE_API virtual void BeginDestroy() override;
 	
+	UE_API TSharedPtr<FStateTreeViewModel> FindViewModel(TNotNull<const UStateTree*> InStateTree) const;
 	UE_API TSharedRef<FStateTreeViewModel> FindOrAddViewModel(TNotNull<UStateTree*> InStateTree);
 	
 	static UE_API bool CompileStateTree(TNotNull<UStateTree*> InStateTree,  FStateTreeCompilerLog& InOutLog);
@@ -42,11 +43,13 @@ public:
 	static UE_API uint32 CalculateStateTreeHash(TNotNull<const UStateTree*> InStateTree);
 	
 private:
-	UE_API void HandlePostGarbageCollect();
+	void HandlePostGarbageCollect();
+	void HandlePostCompile(const UStateTree& InStateTree);
 
 protected:
 	TMap<FObjectKey, TSharedPtr<FStateTreeViewModel>> StateTreeViewModels;
 	FDelegateHandle PostGarbageCollectHandle;
+	FDelegateHandle PostCompileHandle;
 };
 
 #undef UE_API
diff --git a/Engine/Plugins/Runtime/StateTree/Source/StateTreeModule/Public/StateTreeDelegates.h b/Engine/Plugins/Runtime/StateTree/Source/StateTreeModule/Public/StateTreeDelegates.h
--- a/Engine/Plugins/Runtime/StateTree/Source/StateTreeModule/Public/StateTreeDelegates.h
+++ b/Engine/Plugins/Runtime/StateTree/Source/StateTreeModule/Public/StateTreeDelegates.h
@@ -63,11 +63,11 @@ extern UE_API FOnVisualThemeChanged OnVisualThemeChanged;
 DECLARE_MULTICAST_DELEGATE_OneParam(FOnBreakpointsChanged, const UStateTree& /*StateTree*/);
 extern UE_API FOnBreakpointsChanged OnBreakpointsChanged;
 
-/** Called when compilation succeeds */
+/** Called when the state tree compiles. */
 DECLARE_MULTICAST_DELEGATE_OneParam(FOnPostCompile, const UStateTree& /*StateTree*/);
 extern UE_API FOnPostCompile OnPostCompile;
 
-/** Request StateTree compilation. Works only in editor. */
+/** Request StateTree compilation. */
 DECLARE_DELEGATE_RetVal_OneParam(bool, FOnRequestCompile, UStateTree& /*StateTreeToCompile*/);
 extern UE_API FOnRequestCompile OnRequestCompile;
 
