Fix to state tree search fitlering to use intersection search

diff --git a/Engine/Plugins/Runtime/StateTree/Source/StateTreeEditorModule/Private/Customizations/Widgets/SStateTreeNodeTypePicker.cpp b/Engine/Plugins/Runtime/StateTree/Source/StateTreeEditorModule/Private/Customizations/Widgets/SStateTreeNodeTypePicker.cpp
--- a/Engine/Plugins/Runtime/StateTree/Source/StateTreeEditorModule/Private/Customizations/Widgets/SStateTreeNodeTypePicker.cpp
+++ b/Engine/Plugins/Runtime/StateTree/Source/StateTreeEditorModule/Private/Customizations/Widgets/SStateTreeNodeTypePicker.cpp
@@ -417,7 +417,7 @@ void SStateTreeNodeTypePicker::OnSearchBoxTextChanged(const FText& NewText)
 
 	TArray<FString> FilterStrings;
 	NewText.ToString().ParseIntoArrayWS(FilterStrings);
-	FilterStrings.RemoveAll([](const FString& String) { return String.IsEmpty(); });
+	FilterStrings.RemoveAllSwap([](const FString& String) { return String.IsEmpty(); });
 	
 	if (FilterStrings.IsEmpty())
 	{
@@ -446,12 +446,12 @@ int32 SStateTreeNodeTypePicker::FilterNodeTypesChildren(const TArray<FString>& F
 		const FString ItemName = SourceItem->Struct ? SourceItem->Struct->GetDisplayNameText().ToString() : SourceItem->GetCategoryName();
 		for (const FString& Filter : FilterStrings)
 		{
-			if (ItemName.Contains(Filter))
+			if (!ItemName.Contains(Filter))
 			{
-				return true;
+				return false;
 			}
 		}
-		return false;
+		return true;
 	};
 
 	for (const TSharedPtr<FStateTreeNodeTypeItem>& SourceItem : SourceArray)
