[State Tree] Only set context on Validation Instance data when the tree starts running. Allow changing state tree references while the tree is not run.

When we set context on validation, we check if the new tree differs from the old one. This check is not needed when the tree is not run because tree Instance data is empty.

diff --git a/Engine/Plugins/Runtime/StateTree/Source/StateTreeModule/Private/StateTreeExecutionContext.cpp b/Engine/Plugins/Runtime/StateTree/Source/StateTreeModule/Private/StateTreeExecutionContext.cpp
--- a/Engine/Plugins/Runtime/StateTree/Source/StateTreeModule/Private/StateTreeExecutionContext.cpp
+++ b/Engine/Plugins/Runtime/StateTree/Source/StateTreeModule/Private/StateTreeExecutionContext.cpp
@@ -341,10 +341,6 @@ FStateTreeReadOnlyExecutionContext::FStateTreeReadOnlyExecutionContext(TNotNull<
 	, Storage(InStorage)
 {
 	Storage.AcquireReadAccess();
-	if (IsValid())
-	{
-		Storage.GetRuntimeValidation().SetContext(&Owner, &RootStateTree);
-	}
 }
 
 FStateTreeReadOnlyExecutionContext::~FStateTreeReadOnlyExecutionContext()
