[State Tree] fixed editor perf hitch when spamming undo/redo

FStateTreeStateDetails listens to undo/redo event and refreshes details view which creates new DetailsCustomization. Don't force refresh from a pending delete one.

diff --git a/Engine/Plugins/Runtime/StateTree/Source/StateTreeEditorModule/Private/Customizations/StateTreeStateDetails.cpp b/Engine/Plugins/Runtime/StateTree/Source/StateTreeEditorModule/Private/Customizations/StateTreeStateDetails.cpp
--- a/Engine/Plugins/Runtime/StateTree/Source/StateTreeEditorModule/Private/Customizations/StateTreeStateDetails.cpp
+++ b/Engine/Plugins/Runtime/StateTree/Source/StateTreeEditorModule/Private/Customizations/StateTreeStateDetails.cpp
@@ -29,6 +29,12 @@ TSharedRef<IDetailCustomization> FStateTreeStateDetails::MakeInstance()
 	return MakeShareable(new FStateTreeStateDetails);
 }
 
+void FStateTreeStateDetails::PendingDelete()
+{
+	Super::PendingDelete();
+	bIsPendingDelete = true;
+}
+
 void FStateTreeStateDetails::CustomizeDetails(IDetailLayoutBuilder& DetailBuilder)
 {
 	const TSharedRef<IPropertyUtilities> PropUtils = DetailBuilder.GetPropertyUtilities();
@@ -425,6 +431,12 @@ void FStateTreeStateDetails::CustomizeDetails(IDetailLayoutBuilder& DetailBuilde
 
 void FStateTreeStateDetails::PostUndo(bool bSuccess)
 {
+	// Prevent redundant calls to force refresh which doubles the number of DetailCustomization 
+	if (bIsPendingDelete)
+	{
+		return;
+	}
+
 	// Refresh view on undo or redo so that the customization based on e.g. State type will be reflected correctly.
 	if (const TSharedPtr<IPropertyUtilities> PropertyUtilities = WeakPropertyUtilities.Pin())
 	{
@@ -434,6 +446,12 @@ void FStateTreeStateDetails::PostUndo(bool bSuccess)
 
 void FStateTreeStateDetails::PostRedo(bool bSuccess)
 {
+	// Prevent redundant calls to force refresh which doubles the number of DetailCustomization
+	if (bIsPendingDelete)
+	{
+		return;
+	}
+
 	// Refresh view on undo or redo so that the customization based on e.g. State type will be reflected correctly.
 	if (const TSharedPtr<IPropertyUtilities> PropertyUtilities = WeakPropertyUtilities.Pin())
 	{
diff --git a/Engine/Plugins/Runtime/StateTree/Source/StateTreeEditorModule/Private/Customizations/StateTreeStateDetails.h b/Engine/Plugins/Runtime/StateTree/Source/StateTreeEditorModule/Private/Customizations/StateTreeStateDetails.h
--- a/Engine/Plugins/Runtime/StateTree/Source/StateTreeEditorModule/Private/Customizations/StateTreeStateDetails.h
+++ b/Engine/Plugins/Runtime/StateTree/Source/StateTreeEditorModule/Private/Customizations/StateTreeStateDetails.h
@@ -15,11 +15,15 @@ class SComboButton;
 
 class FStateTreeStateDetails : public IDetailCustomization, FSelfRegisteringEditorUndoClient
 {
+private:
+	using Super = IDetailCustomization;
+
 public:
 	/** Makes a new instance of this detail layout class for a specific detail view requesting it */
 	static TSharedRef<IDetailCustomization> MakeInstance();
 
 	/** IDetailCustomization interface */
+	virtual void PendingDelete() override;
 	virtual void CustomizeDetails(IDetailLayoutBuilder& DetailBuilder) override;
 
 protected:
@@ -29,4 +33,6 @@ protected:
 	virtual void PostRedo(bool bSuccess) override;
 
 	TWeakPtr<IPropertyUtilities> WeakPropertyUtilities;
+
+	bool bIsPendingDelete = false;
 };
