[StateTreeDebugger] make sure the instance descriptor is available before creating the StateTree track.

diff --git a/Engine/Plugins/Runtime/StateTree/Source/StateTreeEditorModule/Private/Debugger/StateTreeRewindDebuggerTrack.cpp b/Engine/Plugins/Runtime/StateTree/Source/StateTreeEditorModule/Private/Debugger/StateTreeRewindDebuggerTrack.cpp
--- a/Engine/Plugins/Runtime/StateTree/Source/StateTreeEditorModule/Private/Debugger/StateTreeRewindDebuggerTrack.cpp
+++ b/Engine/Plugins/Runtime/StateTree/Source/StateTreeEditorModule/Private/Debugger/StateTreeRewindDebuggerTrack.cpp
@@ -33,7 +33,20 @@ TSharedPtr<RewindDebugger::FRewindDebuggerTrack> FRewindDebuggerTrackCreator::Cr
 
 bool FRewindDebuggerTrackCreator::HasDebugInfoInternal(const RewindDebugger::FObjectId& InObjectId) const
 {
-	return true;
+	if (const IRewindDebugger* RewindDebugger = IRewindDebugger::Instance())
+	{
+		if (const TraceServices::IAnalysisSession* Session = RewindDebugger->GetAnalysisSession())
+		{
+			if (const IStateTreeTraceProvider* StateTreeTraceProvider = Session->ReadProvider<IStateTreeTraceProvider>(FStateTreeTraceProvider::ProviderName))
+			{
+				if (StateTreeTraceProvider->GetInstanceDescriptor(FStateTreeInstanceDebugId(InObjectId.GetMainId())).IsValid())
+				{
+					return true;
+				}
+			}
+		}
+	}
+	return false;
 }
 
 bool FRewindDebuggerTrackCreator::IsCreatingPrimaryChildTrackInternal() const
@@ -67,7 +80,11 @@ FRewindDebuggerTrack::FRewindDebuggerTrack(const RewindDebugger::FObjectId& InOb
 		}
 	}
 
-	check(Descriptor.IsValid());
+	if (!ensureMsgf(Descriptor.IsValid(), TEXT("Track should not be created unless the instance is available in the StateTreeTraceProvider")))
+	{
+		return;
+	}
+
 	SAssignNew(CompactTreeView, UE::StateTree::SCompactTreeDebuggerView, Descriptor->StateTree.Get())
 		.ActiveStates_Lambda([WeakTrackHelper = InstanceTrackHelper->AsWeak()]
 			{
@@ -128,6 +145,12 @@ bool FRewindDebuggerTrack::UpdateInternal()
 {
 	TRACE_CPUPROFILER_EVENT_SCOPE(FStateTreeRewindDebuggerTrack::UpdateInternal);
 
+	constexpr bool bChildrenChanged = false;
+	if (!Descriptor.IsValid())
+	{
+		return bChildrenChanged;
+	}
+
 	const IRewindDebugger* RewindDebugger = IRewindDebugger::Instance();
 	if (const TraceServices::IAnalysisSession* Session = RewindDebugger->GetAnalysisSession())
 	{
@@ -172,7 +195,7 @@ bool FRewindDebuggerTrack::UpdateInternal()
 		}
 	}
 
-	return false;
+	return bChildrenChanged;
 }
 
 TSharedPtr<SWidget> FRewindDebuggerTrack::GetTimelineViewInternal()
