[StateTreeDebugger] fixed some code analysis warnings

diff --git a/Engine/Plugins/Runtime/StateTree/Source/StateTreeEditorModule/Private/Debugger/StateTreeRewindDebuggerTrack.cpp b/Engine/Plugins/Runtime/StateTree/Source/StateTreeEditorModule/Private/Debugger/StateTreeRewindDebuggerTrack.cpp
--- a/Engine/Plugins/Runtime/StateTree/Source/StateTreeEditorModule/Private/Debugger/StateTreeRewindDebuggerTrack.cpp
+++ b/Engine/Plugins/Runtime/StateTree/Source/StateTreeEditorModule/Private/Debugger/StateTreeRewindDebuggerTrack.cpp
@@ -152,7 +152,7 @@ bool FRewindDebuggerTrack::UpdateInternal()
 	}
 
 	const IRewindDebugger* RewindDebugger = IRewindDebugger::Instance();
-	if (const TraceServices::IAnalysisSession* Session = RewindDebugger->GetAnalysisSession())
+	if (const TraceServices::IAnalysisSession* Session = RewindDebugger ? RewindDebugger->GetAnalysisSession() : nullptr)
 	{
 		const UStateTree* StateTree = Descriptor->StateTree.Get();
 
@@ -205,6 +205,11 @@ TSharedPtr<SWidget> FRewindDebuggerTrack::GetTimelineViewInternal()
 
 TSharedPtr<SWidget> FRewindDebuggerTrack::GetDetailsViewInternal()
 {
+	if (!Descriptor.IsValid())
+	{
+		return SNullWidget::NullWidget;
+	}
+
 	return SNew(SHorizontalBox)
 		+ SHorizontalBox::Slot()
 		.AutoWidth()
