[StateTreeDebugger] updated debugger window to use two dedicated buttons for 'Open RewindDebugger' and 'Start Recording' actions.
Moved the tab names to the module interface to fix an issue where RewindDebugger couldn't be opened since the singleton was not created to retrieve the names of the tabs.
Also added a warning box to inform users that the visualization of the recorded data is now done inside RewindDebugger.

diff --git a/Engine/Plugins/Animation/GameplayInsights/Source/RewindDebugger/Private/RewindDebugger.cpp b/Engine/Plugins/Animation/GameplayInsights/Source/RewindDebugger/Private/RewindDebugger.cpp
--- a/Engine/Plugins/Animation/GameplayInsights/Source/RewindDebugger/Private/RewindDebugger.cpp
+++ b/Engine/Plugins/Animation/GameplayInsights/Source/RewindDebugger/Private/RewindDebugger.cpp
@@ -152,16 +152,6 @@ void FRewindDebugger::Shutdown()
 	delete InternalInstance;
 }
 
-FName FRewindDebugger::GetMainTabName() const
-{
-	return FRewindDebuggerModule::MainTabName;
-}
-
-FName FRewindDebugger::GetDetailsTabName() const
-{
-	return FRewindDebuggerModule::DetailsTabName;
-}
-
 void FRewindDebugger::SetTrackListChangedDelegate(const FOnTrackListChanged& InTrackListChangedDelegate)
 {
 	TrackListChangedDelegate = InTrackListChangedDelegate;
@@ -385,6 +375,11 @@ void FRewindDebugger::OnConnection()
 	FTraceAuxiliary::OnConnection.RemoveAll(this);
 }
 
+bool FRewindDebugger::CanStartRecording() const
+{
+	return !IsRecording() && bPIESimulating;
+}
+
 void FRewindDebugger::StartRecording() const
 {
 	if (!CanStartRecording())
diff --git a/Engine/Plugins/Animation/GameplayInsights/Source/RewindDebugger/Private/RewindDebugger.h b/Engine/Plugins/Animation/GameplayInsights/Source/RewindDebugger/Private/RewindDebugger.h
--- a/Engine/Plugins/Animation/GameplayInsights/Source/RewindDebugger/Private/RewindDebugger.h
+++ b/Engine/Plugins/Animation/GameplayInsights/Source/RewindDebugger/Private/RewindDebugger.h
@@ -105,18 +105,13 @@ public:
 		return static_cast<FRewindDebugger*>(InternalInstance);
 	}
 
-	virtual FName GetMainTabName() const override;
-	virtual FName GetDetailsTabName() const override;
+	virtual bool CanStartRecording() const override;
 
-	/** Start a new Recording:  Start tracing Object + Animation data, increment the current recording index, and reset the recording elapsed time to 0 */
+	/** Start a new Recording:  Start tracing Objects, increment the current recording index, and reset the recording elapsed time to 0 */
 	virtual void StartRecording() const override;
 	void OnClearRecording();
 	void OnRecordingStarted();
 	void OnRecordingStopped();
-	bool CanStartRecording() const
-	{
-		return !IsRecording() && bPIESimulating;
-	}
 
 	bool CanOpenTrace() const;
 	void OpenTrace(const FString& FilePath);
diff --git a/Engine/Plugins/Animation/GameplayInsights/Source/RewindDebugger/Private/RewindDebuggerModule.cpp b/Engine/Plugins/Animation/GameplayInsights/Source/RewindDebugger/Private/RewindDebuggerModule.cpp
--- a/Engine/Plugins/Animation/GameplayInsights/Source/RewindDebugger/Private/RewindDebuggerModule.cpp
+++ b/Engine/Plugins/Animation/GameplayInsights/Source/RewindDebugger/Private/RewindDebuggerModule.cpp
@@ -27,6 +27,16 @@ const FName FRewindDebuggerModule::DetailsTabName("RewindDebuggerDetails2");
 const FName FRewindDebuggerModule::MainMenuName("RewindDebugger.MainMenu");
 const FName FRewindDebuggerModule::TrackContextMenuName("RewindDebugger.TrackContextMenu");
 
+FName FRewindDebuggerModule::GetMainTabName() const
+{
+	return MainTabName;
+}
+
+FName FRewindDebuggerModule::GetDetailsTabName() const
+{
+	return DetailsTabName;
+}
+
 TSharedRef<SDockTab> FRewindDebuggerModule::SpawnRewindDebuggerDetailsTab(const FSpawnTabArgs& SpawnTabArgs)
 {
 	if (FRewindDebugger::Instance() == nullptr)
@@ -231,7 +241,7 @@ void FRewindDebuggerModule::StartupModule()
 
 	FLevelEditorModule& LevelEditorModule = FModuleManager::GetModuleChecked<FLevelEditorModule>("LevelEditor");
 
-	/* LevelEditorTabManagerChangedHandle = */ LevelEditorModule.OnTabManagerChanged().AddLambda([this]()
+	LevelEditorTabManagerChangedHandle = LevelEditorModule.OnTabManagerChanged().AddLambda([this]()
 		{
 			FLevelEditorModule& LevelEditorModule = FModuleManager::GetModuleChecked<FLevelEditorModule>("LevelEditor");
 			TSharedPtr<FTabManager> LevelEditorTabManager = LevelEditorModule.GetLevelEditorTabManager();
@@ -244,7 +254,7 @@ void FRewindDebuggerModule::StartupModule()
 				.SetGroup(WorkspaceMenu::GetMenuStructure().GetDeveloperToolsDebugCategory());
 
 			LevelEditorTabManager->RegisterTabSpawner(
-				DetailsTabName, FOnSpawnTab::CreateStatic(&FRewindDebuggerModule::SpawnRewindDebuggerDetailsTab))
+				DetailsTabName, FOnSpawnTab::CreateRaw(this, &FRewindDebuggerModule::SpawnRewindDebuggerDetailsTab))
 				.SetDisplayName(LOCTEXT("DetailsTabTitle", "Rewind Debugger Details"))
 				.SetIcon(FSlateIcon("RewindDebuggerStyle", "RewindDebugger.RewindDetailsIcon"))
 				.SetTooltipText(LOCTEXT("DetailsWindowTooltipText", "Opens Rewind Debugger Details Window."))
@@ -252,14 +262,14 @@ void FRewindDebuggerModule::StartupModule()
 			
 		});
 
-	/*LevelEditorLayoutExtensionHandle = */ LevelEditorModule.OnRegisterLayoutExtensions().AddLambda(
+	LevelEditorLayoutExtensionHandle = LevelEditorModule.OnRegisterLayoutExtensions().AddLambda(
 		[this](FLayoutExtender& Extender)
 		{
 			Extender.ExtendLayout(FName("LevelEditorSelectionDetails"), ELayoutExtensionPosition::After, FTabManager::FTab(DetailsTabName, ETabState::ClosedTab));
 			Extender.ExtendLayout(FName("Sequencer"), ELayoutExtensionPosition::After, FTabManager::FTab(MainTabName, ETabState::ClosedTab));
 		}
 	);
-	
+
 	RewindDebuggerCameraExtension.Initialize();
 	RewindDebuggerAnimationExtension.Initialize();
 	IModularFeatures::Get().RegisterModularFeature(IRewindDebuggerExtension::ModularFeatureName, &RewindDebuggerCameraExtension);
@@ -276,7 +286,19 @@ void FRewindDebuggerModule::StartupModule()
 void FRewindDebuggerModule::ShutdownModule()
 {
 	RewindDebuggerAnimationExtension.Shutdown();
-	
+
+	if (FLevelEditorModule* LevelEditorModule = FModuleManager::GetModulePtr<FLevelEditorModule>("LevelEditor"))
+	{
+		LevelEditorModule->OnTabManagerChanged().Remove(LevelEditorTabManagerChangedHandle);
+		LevelEditorModule->OnRegisterLayoutExtensions().Remove(LevelEditorLayoutExtensionHandle);
+
+		if (TSharedPtr<FTabManager> LevelEditorTabManager = LevelEditorModule->GetLevelEditorTabManager())
+		{
+			LevelEditorTabManager->UnregisterTabSpawner(MainTabName);
+			LevelEditorTabManager->UnregisterTabSpawner(DetailsTabName);
+		}
+	}
+
 	IModularFeatures::Get().UnregisterModularFeature(IRewindDebuggerExtension::ModularFeatureName, &RewindDebuggerCameraExtension);
 	IModularFeatures::Get().UnregisterModularFeature(IRewindDebuggerExtension::ModularFeatureName, &RewindDebuggerAnimationExtension);
 	IModularFeatures::Get().UnregisterModularFeature(IRewindDebuggerDoubleClickHandler::ModularFeatureName, &AnimInstanceDoubleClickHandler);
diff --git a/Engine/Plugins/Animation/GameplayInsights/Source/RewindDebugger/Private/RewindDebuggerModule.h b/Engine/Plugins/Animation/GameplayInsights/Source/RewindDebugger/Private/RewindDebuggerModule.h
--- a/Engine/Plugins/Animation/GameplayInsights/Source/RewindDebugger/Private/RewindDebuggerModule.h
+++ b/Engine/Plugins/Animation/GameplayInsights/Source/RewindDebugger/Private/RewindDebuggerModule.h
@@ -2,23 +2,25 @@
 
 #pragma once
 
-#include "Modules/ModuleInterface.h"
-#include "Framework/Docking/TabManager.h"
+#include "IRewindDebuggerModule.h"
 #include "RewindDebuggerCamera.h"
 #include "RewindDebuggerAnimation.h"
 
+class SDockTab;
 class SRewindDebugger;
 class SRewindDebuggerDetails;
 
-class FRewindDebuggerModule : public IModuleInterface
+class FRewindDebuggerModule : public IRewindDebuggerModule
 {
 public:
 	// IModuleInterface interface
 	virtual void StartupModule() override;
 	virtual void ShutdownModule() override;
 
-	TSharedRef<SDockTab> SpawnRewindDebuggerTab(const FSpawnTabArgs& SpawnTabArgs);
-	static TSharedRef<SDockTab> SpawnRewindDebuggerDetailsTab(const FSpawnTabArgs& SpawnTabArgs);
+	virtual FName GetMainTabName() const override;
+	virtual FName GetDetailsTabName() const override;
+	virtual TSharedRef<SDockTab> SpawnRewindDebuggerTab(const FSpawnTabArgs& SpawnTabArgs) override;
+	virtual TSharedRef<SDockTab> SpawnRewindDebuggerDetailsTab(const FSpawnTabArgs& SpawnTabArgs) override;
 
 	static const FName MainTabName;
 	static const FName DetailsTabName;
@@ -31,4 +33,7 @@ private:
 
 	FRewindDebuggerCamera RewindDebuggerCameraExtension;
 	FRewindDebuggerAnimation RewindDebuggerAnimationExtension;
+
+	FDelegateHandle LevelEditorTabManagerChangedHandle;
+	FDelegateHandle LevelEditorLayoutExtensionHandle;
 };
diff --git a/Engine/Plugins/Runtime/StateTree/Source/StateTreeEditorModule/Private/Debugger/SStateTreeDebuggerView.cpp b/Engine/Plugins/Runtime/StateTree/Source/StateTreeEditorModule/Private/Debugger/SStateTreeDebuggerView.cpp
--- a/Engine/Plugins/Runtime/StateTree/Source/StateTreeEditorModule/Private/Debugger/SStateTreeDebuggerView.cpp
+++ b/Engine/Plugins/Runtime/StateTree/Source/StateTreeEditorModule/Private/Debugger/SStateTreeDebuggerView.cpp
@@ -5,10 +5,12 @@
 #include "SStateTreeDebuggerView.h"
 #include "Debugger/SStateTreeFrameEventsView.h"
 #include "Editor.h"
-#include "IGameplayProvider.h"
-#include "LevelEditor.h"
+#include "Framework/Docking/TabManager.h"
 #include "Framework/MultiBox/MultiBoxBuilder.h"
+#include "IGameplayProvider.h"
+#include "IRewindDebuggerModule.h"
 #include "Kismet2/DebuggerCommands.h"
+#include "LevelEditor.h"
 #include "Modules/ModuleManager.h"
 #include "SStateTreeDebuggerInstanceTree.h"
 #include "SStateTreeDebuggerTimelines.h"
@@ -20,7 +22,7 @@
 #include "StateTreeEditorStyle.h"
 #include "StateTreeModule.h"
 #include "StateTreeViewModel.h"
-#include "Framework/Docking/TabManager.h"
+#include "SWarningOrErrorBox.h"
 #include "Widgets/Input/SComboButton.h"
 #include "Widgets/Layout/SScrollBox.h"
 
@@ -537,7 +539,11 @@ void SStateTreeDebuggerView::ConstructView(const TSharedRef<FUICommandList>& InC
 	LeftToolbar.BeginSection(TEXT("Debugging"));
 	{
 		const FStateTreeDebuggerCommands& DebuggerCommands = FStateTreeDebuggerCommands::Get();
-		LeftToolbar.AddToolBarButton(DebuggerCommands.StartRewindDebugger);
+
+		LeftToolbar.SetLabelVisibility(EVisibility::Visible);
+		LeftToolbar.AddToolBarButton(DebuggerCommands.OpenRewindDebugger);
+		LeftToolbar.SetLabelVisibility(EVisibility::Hidden);
+
 		LeftToolbar.AddToolBarButton(DebuggerCommands.ResumeDebuggerAnalysis);
 	}
 	LeftToolbar.EndSection();
@@ -615,32 +621,42 @@ void SStateTreeDebuggerView::ConstructView(const TSharedRef<FUICommandList>& InC
 		];
 
 	ChildSlot
+	[
+		SNew(SBorder)
+		.Padding(4.0f)
+		.BorderImage(FAppStyle::GetBrush("ToolPanel.GroupBorder"))
 		[
-			SNew(SBorder)
-				.Padding(4.0f)
-				.BorderImage(FAppStyle::GetBrush("ToolPanel.GroupBorder"))
-				[
-					SNew(SVerticalBox)
-						+ SVerticalBox::Slot()
-						.AutoHeight()
-						.MinHeight(48.f)
-						[
-							SNew(SHorizontalBox)
-								+ SHorizontalBox::Slot()
-								.AutoWidth()
-								.VAlign(VAlign_Center)
-								[
-									LeftToolbar.MakeWidget()
-								]
-								+ SHorizontalBox::Slot()
-								.AutoWidth()
-								.VAlign(VAlign_Center)
-								[
-									InstanceSelectionBox
-								]
-						]
-				]
-		];
+			SNew(SVerticalBox)
+			+ SVerticalBox::Slot()
+			.AutoHeight()
+			.MinHeight(48.f)
+			[
+				SNew(SHorizontalBox)
+					+ SHorizontalBox::Slot()
+					.AutoWidth()
+					.VAlign(VAlign_Center)
+					[
+						LeftToolbar.MakeWidget()
+					]
+					+ SHorizontalBox::Slot()
+					.AutoWidth()
+					.VAlign(VAlign_Center)
+					[
+						InstanceSelectionBox
+					]
+			]
+			+ SVerticalBox::Slot()
+			.AutoHeight()
+			[
+				SNew(SWarningOrErrorBox)
+				.MessageStyle(EMessageStyle::Warning)
+				.Message(LOCTEXT("MigrationToRewindDebugger",
+					"Note that recording and visualization of the recorded data has been migrated to RewindDebugger.\n"
+					"It is still possible to use the previous visualization by using 'Enable Legacy Debugger Window'"
+					" under the StateTree Editor plugin section of the Editor Preferences."))
+			]
+		]
+	];
 }
 
 FReply SStateTreeDebuggerView::OnPreviewKeyDown(const FGeometry& MyGeometry, const FKeyEvent& InKeyEvent)
@@ -790,9 +806,9 @@ void SStateTreeDebuggerView::BindDebuggingToolbarCommands(const TSharedRef<FUICo
 	const FStateTreeDebuggerCommands& Commands = FStateTreeDebuggerCommands::Get();
 
 	ToolkitCommands->MapAction(
-		Commands.StartRewindDebugger,
-		FExecuteAction::CreateSP(this, &SStateTreeDebuggerView::StartRewindDebugger),
-		FCanExecuteAction::CreateSP(this, &SStateTreeDebuggerView::CanStartRewindDebugger));
+		Commands.OpenRewindDebugger,
+		FExecuteAction::CreateStatic(&SStateTreeDebuggerView::OpenRewindDebugger),
+		FCanExecuteAction::CreateStatic(&SStateTreeDebuggerView::CanOpenRewindDebugger));
 
 	ToolkitCommands->MapAction(
 		Commands.ResumeDebuggerAnalysis,
@@ -800,33 +816,33 @@ void SStateTreeDebuggerView::BindDebuggingToolbarCommands(const TSharedRef<FUICo
 		FCanExecuteAction::CreateSP(this, &SStateTreeDebuggerView::CanResumeDebuggerAnalysis));
 }
 
-bool SStateTreeDebuggerView::CanStartRewindDebugger() const
+bool SStateTreeDebuggerView::CanOpenRewindDebugger()
 {
-	if (!IsRecording())
+	if (const FLevelEditorModule* LevelEditorModule = FModuleManager::GetModulePtr<FLevelEditorModule>("LevelEditor"))
 	{
-		return true;
-	}
-
-	if (const IRewindDebugger* RewindDebugger = IRewindDebugger::Instance())
-	{
-		const FLevelEditorModule& LevelEditorModule = FModuleManager::GetModuleChecked<FLevelEditorModule>("LevelEditor");
-		const TSharedPtr<FTabManager> LevelEditorTabManager = LevelEditorModule.GetLevelEditorTabManager();
-		return !LevelEditorTabManager->FindExistingLiveTab(RewindDebugger->GetMainTabName()).IsValid();
+		if (const TSharedPtr<FTabManager> LevelEditorTabManager = LevelEditorModule->GetLevelEditorTabManager())
+		{
+			if (const IRewindDebuggerModule* RewindDebuggerModule = FModuleManager::GetModulePtr<IRewindDebuggerModule>("RewindDebugger"))
+			{
+				return LevelEditorTabManager->FindTabSpawnerFor(RewindDebuggerModule->GetMainTabName()).IsValid();
+			}
+		}
 	}
-
 	return false;
 }
 
-void SStateTreeDebuggerView::StartRewindDebugger()
+void SStateTreeDebuggerView::OpenRewindDebugger()
 {
-	if (const IRewindDebugger* RewindDebugger = IRewindDebugger::Instance())
+	if (const FLevelEditorModule* LevelEditorModule = FModuleManager::GetModulePtr<FLevelEditorModule>("LevelEditor"))
 	{
-		const FLevelEditorModule& LevelEditorModule = FModuleManager::GetModuleChecked<FLevelEditorModule>("LevelEditor");
-		const TSharedPtr<FTabManager> LevelEditorTabManager = LevelEditorModule.GetLevelEditorTabManager();
-		LevelEditorTabManager->TryInvokeTab(RewindDebugger->GetMainTabName());
-		LevelEditorTabManager->TryInvokeTab(RewindDebugger->GetDetailsTabName());
-
-		RewindDebugger->StartRecording();
+		if (const TSharedPtr<FTabManager> LevelEditorTabManager = LevelEditorModule->GetLevelEditorTabManager())
+		{
+			if (const IRewindDebuggerModule* RewindDebuggerModule = FModuleManager::GetModulePtr<IRewindDebuggerModule>("RewindDebugger"))
+			{
+				LevelEditorTabManager->TryInvokeTab(RewindDebuggerModule->GetMainTabName());
+				LevelEditorTabManager->TryInvokeTab(RewindDebuggerModule->GetDetailsTabName());
+			}
+		}
 	}
 }
 
diff --git a/Engine/Plugins/Runtime/StateTree/Source/StateTreeEditorModule/Private/Debugger/SStateTreeDebuggerView.h b/Engine/Plugins/Runtime/StateTree/Source/StateTreeEditorModule/Private/Debugger/SStateTreeDebuggerView.h
--- a/Engine/Plugins/Runtime/StateTree/Source/StateTreeEditorModule/Private/Debugger/SStateTreeDebuggerView.h
+++ b/Engine/Plugins/Runtime/StateTree/Source/StateTreeEditorModule/Private/Debugger/SStateTreeDebuggerView.h
@@ -86,8 +86,8 @@ private:
 	/** RewindDebugger sync commands */
 	void BindDebuggingToolbarCommands(const TSharedRef<FUICommandList> ToolkitCommands);
 
-	bool CanStartRewindDebugger() const;
-	void StartRewindDebugger();
+	static bool CanOpenRewindDebugger();
+	static void OpenRewindDebugger();
 
 	bool CanUseScrubButtons() const;
 
diff --git a/Engine/Plugins/Runtime/StateTree/Source/StateTreeEditorModule/Private/Debugger/StateTreeDebuggerCommands.cpp b/Engine/Plugins/Runtime/StateTree/Source/StateTreeEditorModule/Private/Debugger/StateTreeDebuggerCommands.cpp
--- a/Engine/Plugins/Runtime/StateTree/Source/StateTreeEditorModule/Private/Debugger/StateTreeDebuggerCommands.cpp
+++ b/Engine/Plugins/Runtime/StateTree/Source/StateTreeEditorModule/Private/Debugger/StateTreeDebuggerCommands.cpp
@@ -31,7 +31,7 @@ void FStateTreeDebuggerCommands::RegisterCommands()
 	UI_COMMAND(NextFrameWithStateChange, "Next State Change", "Jump to next state changed frame", EUserInterfaceActionType::Button, FInputChord(EModifierKey::Control, EKeys::Right));
 	UI_COMMAND(ResumeDebuggerAnalysis, "Resume Analysis", "Resume trace analysis after hitting breakpoint", EUserInterfaceActionType::Button, FInputChord());
 	UI_COMMAND(ResetTracks, "Reset Tracks", "Clear current tracks and associated events.", EUserInterfaceActionType::Button, FInputChord());
-	UI_COMMAND(StartRewindDebugger, "Start RewindDebugger", "Open RewindDebugger window and start recording", EUserInterfaceActionType::Button, FInputChord());
+	UI_COMMAND(OpenRewindDebugger, "Open RewindDebugger", "Open RewindDebugger window and its details tab", EUserInterfaceActionType::Button, FInputChord());
 }
 
 #undef LOCTEXT_NAMESPACE
diff --git a/Engine/Plugins/Runtime/StateTree/Source/StateTreeEditorModule/Private/Debugger/StateTreeDebuggerCommands.h b/Engine/Plugins/Runtime/StateTree/Source/StateTreeEditorModule/Private/Debugger/StateTreeDebuggerCommands.h
--- a/Engine/Plugins/Runtime/StateTree/Source/StateTreeEditorModule/Private/Debugger/StateTreeDebuggerCommands.h
+++ b/Engine/Plugins/Runtime/StateTree/Source/StateTreeEditorModule/Private/Debugger/StateTreeDebuggerCommands.h
@@ -25,7 +25,7 @@ public:
 	TSharedPtr<FUICommandInfo> NextFrameWithStateChange;
 	TSharedPtr<FUICommandInfo> ResumeDebuggerAnalysis;
 	TSharedPtr<FUICommandInfo> ResetTracks;
-	TSharedPtr<FUICommandInfo> StartRewindDebugger;
+	TSharedPtr<FUICommandInfo> OpenRewindDebugger;
 };
 
 #endif // WITH_STATETREE_TRACE_DEBUGGER
\ No newline at end of file
diff --git a/Engine/Plugins/Runtime/StateTree/Source/StateTreeEditorModule/Private/StateTreeEditorStyle.cpp b/Engine/Plugins/Runtime/StateTree/Source/StateTreeEditorModule/Private/StateTreeEditorStyle.cpp
--- a/Engine/Plugins/Runtime/StateTree/Source/StateTreeEditorModule/Private/StateTreeEditorStyle.cpp
+++ b/Engine/Plugins/Runtime/StateTree/Source/StateTreeEditorModule/Private/StateTreeEditorStyle.cpp
@@ -235,11 +235,16 @@ FStateTreeEditorStyle::FStateTreeEditorStyle()
 		Set("Node.Sync", new CORE_IMAGE_BRUSH_SVG("Starship/Common/Update", CoreStyleConstants::Icon16x16));
 	}
 
+	{
+		// From RewindDebugger
+		FContentRootScope Scope(this, FPaths::EngineDir() / TEXT("Plugins/Animation/GameplayInsights/Content"));
+		Set("StateTreeEditor.Debugger.OpenRewindDebugger", new IMAGE_BRUSH("Rewind_24x", CoreStyleConstants::Icon16x16));
+	}
+
 	{
 		// From generic Engine Editor
 		FContentRootScope Scope(this, EngineEditorSlateContentDir);
-
-		Set("StateTreeEditor.Debugger.StartRewindDebugger", new IMAGE_BRUSH("Sequencer/Transport_Bar/Record_24x", CoreStyleConstants::Icon16x16));
+		Set("StateTreeEditor.Debugger.StartRewindDebuggerRecording", new IMAGE_BRUSH("Sequencer/Transport_Bar/Record_24x", CoreStyleConstants::Icon16x16));
 
 		Set("StateTreeEditor.Debugger.StartRecording", new IMAGE_BRUSH("Sequencer/Transport_Bar/Record_24x", CoreStyleConstants::Icon16x16));
 		Set("StateTreeEditor.Debugger.StopRecording", new IMAGE_BRUSH("Sequencer/Transport_Bar/Recording_24x", CoreStyleConstants::Icon16x16));
diff --git a/Engine/Source/Editor/RewindDebuggerInterface/Public/IRewindDebugger.h b/Engine/Source/Editor/RewindDebuggerInterface/Public/IRewindDebugger.h
--- a/Engine/Source/Editor/RewindDebuggerInterface/Public/IRewindDebugger.h
+++ b/Engine/Source/Editor/RewindDebuggerInterface/Public/IRewindDebugger.h
@@ -72,11 +72,8 @@ public:
 
 	REWINDDEBUGGERINTERFACE_API virtual ~IRewindDebugger();
 
-	/** @return the name of the main RewindDebugger window that could be used to invoke the associated tab in the Editor. */
-	virtual FName GetMainTabName() const = 0;
-
-	/** @return the name of the RewindDebugger details window that could be used to invoke the associated tab in the Editor. */
-	virtual FName GetDetailsTabName() const = 0;
+	/** Whether the recording can be started. */
+	virtual bool CanStartRecording() const = 0;
 
 	/** Start traces recording if not already started. */
 	virtual void StartRecording() const = 0;
diff --git a/Engine/Source/Editor/RewindDebuggerInterface/Public/IRewindDebuggerModule.h b/Engine/Source/Editor/RewindDebuggerInterface/Public/IRewindDebuggerModule.h
new file mode 100644
--- /dev/null
+++ b/Engine/Source/Editor/RewindDebuggerInterface/Public/IRewindDebuggerModule.h
@@ -0,0 +1,23 @@
+// Copyright Epic Games, Inc. All Rights Reserved.
+
+#pragma once
+
+#include "Modules/ModuleInterface.h"
+#include "Templates/SharedPointerFwd.h"
+
+class FName;
+class FSpawnTabArgs;
+class SDockTab;
+
+class IRewindDebuggerModule : public IModuleInterface
+{
+public:
+	/** @return the name of the main RewindDebugger window that could be used to invoke the associated tab in the Editor. */
+	virtual FName GetMainTabName() const = 0;
+
+	/** @return the name of the RewindDebugger details window that could be used to invoke the associated tab in the Editor. */
+	virtual FName GetDetailsTabName() const = 0;
+
+	virtual TSharedRef<SDockTab> SpawnRewindDebuggerTab(const FSpawnTabArgs& SpawnTabArgs) = 0;
+	virtual TSharedRef<SDockTab> SpawnRewindDebuggerDetailsTab(const FSpawnTabArgs& SpawnTabArgs) = 0;
+};
\ No newline at end of file
