Copy-pasta renaming fixup

diff --git a/Engine/Plugins/Runtime/StateTree/Source/StateTreeEditorModule/Private/StateTreeEditorUILayer.cpp b/Engine/Plugins/Runtime/StateTree/Source/StateTreeEditorModule/Private/StateTreeEditorUILayer.cpp
--- a/Engine/Plugins/Runtime/StateTree/Source/StateTreeEditorModule/Private/StateTreeEditorUILayer.cpp
+++ b/Engine/Plugins/Runtime/StateTree/Source/StateTreeEditorModule/Private/StateTreeEditorUILayer.cpp
@@ -52,14 +52,14 @@ TSharedPtr<FWorkspaceItem> FStateTreeEditorModeUILayer::GetModeMenuCategory() co
 
 void UStateTreeEditorUISubsystem::Initialize(FSubsystemCollectionBase& Collection)
 {
-	FStateTreeEditorModule& LevelEditorModule = FModuleManager::LoadModuleChecked<FStateTreeEditorModule>("StateTreeEditorModule");
-	LevelEditorModule.OnRegisterLayoutExtensions().AddUObject(this, &UStateTreeEditorUISubsystem::RegisterLayoutExtensions);
+	FStateTreeEditorModule& StateTreeEditorModule = FModuleManager::LoadModuleChecked<FStateTreeEditorModule>("StateTreeEditorModule");
+	StateTreeEditorModule.OnRegisterLayoutExtensions().AddUObject(this, &UStateTreeEditorUISubsystem::RegisterLayoutExtensions);
 }
 
 void UStateTreeEditorUISubsystem::Deinitialize()
 {
-	FStateTreeEditorModule& LevelEditorModule = FModuleManager::LoadModuleChecked<FStateTreeEditorModule>("StateTreeEditorModule");
-	LevelEditorModule.OnRegisterLayoutExtensions().RemoveAll(this);
+	FStateTreeEditorModule& StateTreeEditorModule = FModuleManager::LoadModuleChecked<FStateTreeEditorModule>("StateTreeEditorModule");
+	StateTreeEditorModule.OnRegisterLayoutExtensions().RemoveAll(this);
 }
 
 void UStateTreeEditorUISubsystem::RegisterLayoutExtensions(FLayoutExtender& Extender)
