Fixed typo in WITH_EDITORONLY_DATA.

#rnx

diff --git a/Engine/Plugins/Runtime/StateTree/Source/StateTreeModule/Public/StateTreeExecutionTypes.h b/Engine/Plugins/Runtime/StateTree/Source/StateTreeModule/Public/StateTreeExecutionTypes.h
--- a/Engine/Plugins/Runtime/StateTree/Source/StateTreeModule/Public/StateTreeExecutionTypes.h
+++ b/Engine/Plugins/Runtime/StateTree/Source/StateTreeModule/Public/StateTreeExecutionTypes.h
@@ -958,9 +958,15 @@ struct FStateTreeExecutionFrame
 {
 	GENERATED_BODY()
 
+PRAGMA_DISABLE_DEPRECATION_WARNINGS
 	FStateTreeExecutionFrame() = default;
+	FStateTreeExecutionFrame(const FStateTreeExecutionFrame&) = default;
+	FStateTreeExecutionFrame(FStateTreeExecutionFrame&&) = default;
+	FStateTreeExecutionFrame& operator=(const FStateTreeExecutionFrame&) = default;
+	FStateTreeExecutionFrame& operator=(FStateTreeExecutionFrame&&) = default;
 	UE_DEPRECATED(5.6, "The recorded frame doesn't have all the needed information to properly form an FStateTreeExecutionFrame.")
 	UE_API FStateTreeExecutionFrame(const FRecordedStateTreeExecutionFrame& RecordedExecutionFrame);
+PRAGMA_ENABLE_DEPRECATION_WARNINGS
 
 	UE_DEPRECATED(5.7, "Use the other version of IsSameRoot.")
 	bool IsSameFrame(const FStateTreeExecutionFrame& OtherFrame) const
@@ -1044,7 +1050,7 @@ struct FStateTreeExecutionFrame
 	UPROPERTY()
 	uint8 bHaveEntered : 1 = false;
 
-#if WITH_EDITORDATA_ONLY
+#if WITH_EDITORONLY_DATA
 	UE_DEPRECATED(5.7, "States in ActiveStates are always valid.")
 	uint8 NumCurrentlyActiveStates = 0;
 #endif
