[StateTree] Adding a common condition for comparing FNames.

diff --git a/Engine/Plugins/Runtime/StateTree/Source/StateTreeModule/Private/Conditions/StateTreeCommonConditions.cpp b/Engine/Plugins/Runtime/StateTree/Source/StateTreeModule/Private/Conditions/StateTreeCommonConditions.cpp
--- a/Engine/Plugins/Runtime/StateTree/Source/StateTreeModule/Private/Conditions/StateTreeCommonConditions.cpp
+++ b/Engine/Plugins/Runtime/StateTree/Source/StateTreeModule/Private/Conditions/StateTreeCommonConditions.cpp
@@ -365,6 +365,53 @@ FText FStateTreeCompareDistanceCondition::GetDescription(const FGuid& ID, FState
 }
 #endif
 
+//----------------------------------------------------------------------//
+//  FStateTreeCompareNameCondition
+//----------------------------------------------------------------------//
+
+bool FStateTreeCompareNameCondition::TestCondition(FStateTreeExecutionContext& Context) const
+{
+	const FInstanceDataType& InstanceData = Context.GetInstanceData(*this);
+
+	SET_NODE_CUSTOM_TRACE_TEXT(Context, Override, TEXT("%s%s is %s")
+							   , *UE::StateTree::DescHelpers::GetInvertText(bInvert, EStateTreeNodeFormatting::Text).ToString()
+							   , *InstanceData.Left.ToString()
+							   , *InstanceData.Right.ToString());
+
+	return InstanceData.Left.IsEqual(InstanceData.Right) ^ bInvert;
+}
+
+#if WITH_EDITOR
+FText FStateTreeCompareNameCondition::GetDescription(const FGuid& ID, FStateTreeDataView InstanceDataView, const IStateTreeBindingLookup& BindingLookup, EStateTreeNodeFormatting Formatting) const
+{
+	const FInstanceDataType* InstanceData = InstanceDataView.GetPtr<FInstanceDataType>();
+	check(InstanceData);
+
+	FText LeftValue = BindingLookup.GetBindingSourceDisplayName(FPropertyBindingPath(ID, GET_MEMBER_NAME_CHECKED(FInstanceDataType, Left)), Formatting);
+	if (LeftValue.IsEmpty())
+	{
+		LeftValue = FText::FromName(InstanceData->Left);
+	}
+
+	FText RightValue = BindingLookup.GetBindingSourceDisplayName(FPropertyBindingPath(ID, GET_MEMBER_NAME_CHECKED(FInstanceDataType, Right)), Formatting);
+	if (RightValue.IsEmpty())
+	{
+		RightValue = FText::FromName(InstanceData->Right);
+	}
+
+	const FText InvertText = UE::StateTree::DescHelpers::GetInvertText(bInvert, Formatting);
+
+	const FText Format = (Formatting == EStateTreeNodeFormatting::RichText)
+		? LOCTEXT("CompareNameRich", "{EmptyOrNot}{Left} <s>is</> {Right}")
+	: LOCTEXT("CompareName", "{EmptyOrNot}{Left} is {Right}");
+
+	return FText::FormatNamed(Format,
+							  TEXT("EmptyOrNot"), InvertText,
+							  TEXT("Left"),LeftValue,
+							  TEXT("Right"),RightValue);
+}
+#endif
+
 //----------------------------------------------------------------------//
 //  FStateTreeRandomCondition
 //----------------------------------------------------------------------//
diff --git a/Engine/Plugins/Runtime/StateTree/Source/StateTreeModule/Public/Conditions/StateTreeCommonConditions.h b/Engine/Plugins/Runtime/StateTree/Source/StateTreeModule/Public/Conditions/StateTreeCommonConditions.h
--- a/Engine/Plugins/Runtime/StateTree/Source/StateTreeModule/Public/Conditions/StateTreeCommonConditions.h
+++ b/Engine/Plugins/Runtime/StateTree/Source/StateTreeModule/Public/Conditions/StateTreeCommonConditions.h
@@ -229,6 +229,48 @@ struct FStateTreeCompareDistanceCondition : public FStateTreeConditionCommonBase
 	EGenericAICheck Operator = EGenericAICheck::Equal;
 };
 
+USTRUCT()
+struct FStateTreeCompareNameConditionInstanceData
+{
+	GENERATED_BODY()
+
+	UPROPERTY(EditAnywhere, Category = "Input")
+	FName Left = NAME_None;
+
+	UPROPERTY(EditAnywhere, Category = "Parameter")
+	FName Right = NAME_None;
+};
+STATETREE_POD_INSTANCEDATA(FStateTreeCompareNameConditionInstanceData);
+
+/**
+ * Condition comparing two FNames.
+ */
+USTRUCT(DisplayName = "Name Compare")
+struct FStateTreeCompareNameCondition : public FStateTreeConditionCommonBase
+{
+	GENERATED_BODY()
+
+	using FInstanceDataType = FStateTreeCompareNameConditionInstanceData;
+
+	FStateTreeCompareNameCondition() = default;
+	explicit FStateTreeCompareNameCondition(const EStateTreeCompare InInverts)
+	: bInvert(InInverts == EStateTreeCompare::Invert)
+	{}
+
+	FStateTreeCompareNameCondition(const bool bInInverts)
+	: bInvert(bInInverts)
+	{}
+
+	virtual const UStruct* GetInstanceDataType() const override { return FInstanceDataType::StaticStruct(); }
+	UE_API virtual bool TestCondition(FStateTreeExecutionContext& Context) const override;
+#if WITH_EDITOR
+	UE_API virtual FText GetDescription(const FGuid& ID, FStateTreeDataView InstanceDataView, const IStateTreeBindingLookup& BindingLookup, EStateTreeNodeFormatting Formatting = EStateTreeNodeFormatting::Text) const override;
+#endif
+
+	UPROPERTY(EditAnywhere, Category = "Parameter")
+	bool bInvert = false;
+};
+
 
 USTRUCT()
 struct FStateTreeRandomConditionInstanceData
