[StateTreeDebugger] When stopping the analysic while pause from a breakpoint, we now reset the hit breakpoint before forcing the last update. This is required since it should no longer be in effect and might even be removed by the user while the analysis was paused.
Also cleared the selected instance when stopping the session analysis so the breakpoints can be hit again by any instance on the next PIE session without selecting a new instance.

#rnx

diff --git a/Engine/Plugins/Runtime/StateTree/Source/StateTreeModule/Private/Debugger/StateTreeDebugger.cpp b/Engine/Plugins/Runtime/StateTree/Source/StateTreeModule/Private/Debugger/StateTreeDebugger.cpp
--- a/Engine/Plugins/Runtime/StateTree/Source/StateTreeModule/Private/Debugger/StateTreeDebugger.cpp
+++ b/Engine/Plugins/Runtime/StateTree/Source/StateTreeModule/Private/Debugger/StateTreeDebugger.cpp
@@ -137,6 +137,12 @@ void FStateTreeDebugger::Tick(const float DeltaTime)
 
 void FStateTreeDebugger::StopSessionAnalysis()
 {
+	// HitBreakpoint is normally reset when resuming the session analysis, but it is also
+	// possible to stop the session analysis while it is paused from a breakpoint.
+	// In this case we make sure to reset it before forcing the last update since the breakpoint
+	// should no longer be in effect.
+	HitBreakpoint.Reset();
+
 	if (IsAnalysisSessionActive())
 	{
 		// Force one last update to process events emitted while closing the game session (e.g., EndPlay in PIE)
@@ -151,7 +157,6 @@ void FStateTreeDebugger::StopSessionAnalysis()
 
 	bSessionAnalysisActive = false;
 	bSessionAnalysisPaused = false;
-	HitBreakpoint.Reset();
 	LastProcessedRecordedWorldTime = 0;
 }
 
@@ -273,6 +278,9 @@ void FStateTreeDebugger::RequestAnalysisOfLatestTrace()
 	// Invalidate our current active session
 	ActiveSessionTraceDescriptor = FTraceDescriptor();
 
+	// Invalidate current selected instance so breakpoint can be hit by any instances in the next analysis
+	ClearSelection();
+
 	// Stop current analysis if any
 	StopSessionAnalysis();
 
