[StateTree] Add a rule to remove sustained state. All reselected state are new state.
[StateTree] Add ability to mark a transition for reselection.

---

On Behalf of Patrick.Boutot - Add a rule to remove sustained state. All reselected state are new state.

#tests PIE

diff --git a/Engine/Plugins/Runtime/StateTree/Source/StateTreeEditorModule/Private/Customizations/StateTreeTransitionDetails.cpp b/Engine/Plugins/Runtime/StateTree/Source/StateTreeEditorModule/Private/Customizations/StateTreeTransitionDetails.cpp
--- a/Engine/Plugins/Runtime/StateTree/Source/StateTreeEditorModule/Private/Customizations/StateTreeTransitionDetails.cpp
+++ b/Engine/Plugins/Runtime/StateTree/Source/StateTreeEditorModule/Private/Customizations/StateTreeTransitionDetails.cpp
@@ -66,6 +66,7 @@ void FStateTreeTransitionDetails::CustomizeHeader(TSharedRef<IPropertyHandle> St
 	DelegateListener = StructProperty->GetChildHandle(GET_MEMBER_NAME_CHECKED(FStateTreeTransition, DelegateListener));
 	StateProperty = StructProperty->GetChildHandle(GET_MEMBER_NAME_CHECKED(FStateTreeTransition, State));
 	DelayTransitionProperty = StructProperty->GetChildHandle(GET_MEMBER_NAME_CHECKED(FStateTreeTransition, bDelayTransition));
+	TransitionReactivateProperty = StructProperty->GetChildHandle(GET_MEMBER_NAME_CHECKED(FStateTreeTransition, bReactivateTargetState));
 	DelayDurationProperty = StructProperty->GetChildHandle(GET_MEMBER_NAME_CHECKED(FStateTreeTransition, DelayDuration));
 	DelayRandomVarianceProperty = StructProperty->GetChildHandle(GET_MEMBER_NAME_CHECKED(FStateTreeTransition, DelayRandomVariance));
 	ConditionsProperty = StructProperty->GetChildHandle(GET_MEMBER_NAME_CHECKED(FStateTreeTransition, Conditions));
@@ -160,6 +161,7 @@ void FStateTreeTransitionDetails::CustomizeChildren(TSharedRef<class IPropertyHa
 	check(RequiredEventProperty);
 	check(DelegateListener);
 	check(DelayTransitionProperty);
+	check(TransitionReactivateProperty);
 	check(DelayDurationProperty);
 	check(DelayRandomVarianceProperty);
 	check(StateProperty);
@@ -224,6 +226,9 @@ void FStateTreeTransitionDetails::CustomizeChildren(TSharedRef<class IPropertyHa
 	StructBuilder.AddProperty(DelayRandomVarianceProperty.ToSharedRef())
 		.Visibility(TAttribute<EVisibility>::Create(TAttribute<EVisibility>::FGetter::CreateLambda(IsNotCompletionTransition)));
 
+	// Transition Reselection
+	StructBuilder.AddProperty(TransitionReactivateProperty.ToSharedRef());
+
 	// Show conditions always expanded, with simplified header (remove item count)
 	IDetailPropertyRow& ConditionsRow = StructBuilder.AddProperty(ConditionsProperty.ToSharedRef());
 	ConditionsRow.ShouldAutoExpand(true);
diff --git a/Engine/Plugins/Runtime/StateTree/Source/StateTreeEditorModule/Private/Customizations/StateTreeTransitionDetails.h b/Engine/Plugins/Runtime/StateTree/Source/StateTreeEditorModule/Private/Customizations/StateTreeTransitionDetails.h
--- a/Engine/Plugins/Runtime/StateTree/Source/StateTreeEditorModule/Private/Customizations/StateTreeTransitionDetails.h
+++ b/Engine/Plugins/Runtime/StateTree/Source/StateTreeEditorModule/Private/Customizations/StateTreeTransitionDetails.h
@@ -57,6 +57,7 @@ private:
 	TSharedPtr<IPropertyHandle> DelegateListener;
 	TSharedPtr<IPropertyHandle> StateProperty;
 	TSharedPtr<IPropertyHandle> DelayTransitionProperty;
+	TSharedPtr<IPropertyHandle> TransitionReactivateProperty;
 	TSharedPtr<IPropertyHandle> DelayDurationProperty;
 	TSharedPtr<IPropertyHandle> DelayRandomVarianceProperty;
 	TSharedPtr<IPropertyHandle> ConditionsProperty;
diff --git a/Engine/Plugins/Runtime/StateTree/Source/StateTreeEditorModule/Private/SStateTreeViewRow.cpp b/Engine/Plugins/Runtime/StateTree/Source/StateTreeEditorModule/Private/SStateTreeViewRow.cpp
--- a/Engine/Plugins/Runtime/StateTree/Source/StateTreeEditorModule/Private/SStateTreeViewRow.cpp
+++ b/Engine/Plugins/Runtime/StateTree/Source/StateTreeEditorModule/Private/SStateTreeViewRow.cpp
@@ -995,6 +995,7 @@ TSharedRef<SWidget> SStateTreeViewRow::MakeTransitionWidgetInternal(const EState
 		FText Tooltip;
 		FStateTreeStateLink Link;
 		FGuid NodeID;
+		bool bSkipComma = false;
 	};
 
 	TArray<FItem> DescItems;
@@ -1022,6 +1023,15 @@ TSharedRef<SWidget> SStateTreeViewRow::MakeTransitionWidgetInternal(const EState
 		if (bMatch)
 		{
 			DescItems.Emplace(Transition.State, Transition.ID);
+
+			if (Transition.bReactivateTargetState)
+			{
+				FItem ReselectionNote;
+				ReselectionNote.Desc = LOCTEXT("ReactivateDesc", "(Re-Activate)");
+				ReselectionNote.Tooltip = LOCTEXT("ReactivateDescToolTip", "The target state and all the following states will be reactivated. Their instance data will be reinstantiated and will receive EnterState.");
+				ReselectionNote.bSkipComma = true;
+				DescItems.Emplace(ReselectionNote);
+			}
 		}
 	}
 
@@ -1090,7 +1100,7 @@ TSharedRef<SWidget> SStateTreeViewRow::MakeTransitionWidgetInternal(const EState
 	{
 		const FItem& Item = DescItems[Index];
 
-		if (Index > 0)
+		if (Index > 0 && !Item.bSkipComma)
 		{
 			TransitionContainer->AddSlot()
 			.AutoWidth()
diff --git a/Engine/Plugins/Runtime/StateTree/Source/StateTreeEditorModule/Private/StateTreeCompiler.cpp b/Engine/Plugins/Runtime/StateTree/Source/StateTreeEditorModule/Private/StateTreeCompiler.cpp
--- a/Engine/Plugins/Runtime/StateTree/Source/StateTreeEditorModule/Private/StateTreeCompiler.cpp
+++ b/Engine/Plugins/Runtime/StateTree/Source/StateTreeEditorModule/Private/StateTreeCompiler.cpp
@@ -1354,6 +1354,7 @@ bool FStateTreeCompiler::CreateStateTransitions()
 				}
 			}
 
+			CompactTransition.bReactivateTargetState = Transition.bReactivateTargetState;
 			CompactTransition.bTransitionEnabled = Transition.bTransitionEnabled;
 
 			if (Transition.bDelayTransition)
diff --git a/Engine/Plugins/Runtime/StateTree/Source/StateTreeEditorModule/Private/StateTreeDescriptionHelpers.cpp b/Engine/Plugins/Runtime/StateTree/Source/StateTreeEditorModule/Private/StateTreeDescriptionHelpers.cpp
--- a/Engine/Plugins/Runtime/StateTree/Source/StateTreeEditorModule/Private/StateTreeDescriptionHelpers.cpp
+++ b/Engine/Plugins/Runtime/StateTree/Source/StateTreeEditorModule/Private/StateTreeDescriptionHelpers.cpp
@@ -247,7 +247,16 @@ FText GetTransitionDesc(const UStateTreeEditorData* EditorData, const FStateTree
 			: LOCTEXT("ActionReturn", "return");
 	}
 	
-	return FText::Format(LOCTEXT("TransitionDesc", "{0} {1} {2}"), TriggerText, ActionText, GetStateLinkDesc(EditorData, Transition.State, Formatting, bShowStatePath));
+	FText ReactivateText = Formatting == EStateTreeNodeFormatting::RichText
+		? LOCTEXT("ReselectionTextRich", "<s> (Re-Activate)</>")
+		: LOCTEXT("ReselectionText", " (Re-Activate)");
+
+	if (!Transition.bReactivateTargetState)
+	{
+		ReactivateText = FText::GetEmpty();
+	}
+
+	return FText::Format(LOCTEXT("TransitionDesc", "{0} {1} {2}{3}"), TriggerText, ActionText, GetStateLinkDesc(EditorData, Transition.State, Formatting, bShowStatePath), ReactivateText);
 }
 
 } // UE::StateTree::Editor
diff --git a/Engine/Plugins/Runtime/StateTree/Source/StateTreeEditorModule/Public/StateTreeState.h b/Engine/Plugins/Runtime/StateTree/Source/StateTreeEditorModule/Public/StateTreeState.h
--- a/Engine/Plugins/Runtime/StateTree/Source/StateTreeEditorModule/Public/StateTreeState.h
+++ b/Engine/Plugins/Runtime/StateTree/Source/StateTreeEditorModule/Public/StateTreeState.h
@@ -156,6 +156,10 @@ PRAGMA_ENABLE_DEPRECATION_WARNINGS
 	UPROPERTY(EditDefaultsOnly, Category = "Transition", meta = (BaseStruct = "/Script/StateTreeModule.StateTreeConditionBase", BaseClass = "/Script/StateTreeModule.StateTreeConditionBlueprintBase"))
 	TArray<FStateTreeEditorNode> Conditions;
 
+	/** The target state and all the following states will be reactivated. Their instance data will be reinstantiated and will receive EnterState. */
+	UPROPERTY(EditDefaultsOnly, Category = "Transition")
+	bool bReactivateTargetState = false;
+
 	/** True if the Transition is Enabled (i.e. not explicitly disabled in the asset). */
 	UPROPERTY(EditDefaultsOnly, Category = "Debug")
 	bool bTransitionEnabled = true;
diff --git a/Engine/Plugins/Runtime/StateTree/Source/StateTreeModule/Private/StateTreeExecutionContext.cpp b/Engine/Plugins/Runtime/StateTree/Source/StateTreeModule/Private/StateTreeExecutionContext.cpp
--- a/Engine/Plugins/Runtime/StateTree/Source/StateTreeModule/Private/StateTreeExecutionContext.cpp
+++ b/Engine/Plugins/Runtime/StateTree/Source/StateTreeModule/Private/StateTreeExecutionContext.cpp
@@ -1984,7 +1984,9 @@ void FStateTreeExecutionContext::RequestTransition(const FStateTreeTransitionReq
 
 		const FActiveStateID CurrentlyProcessedStateID = CurrentlyProcessedFrame->ActiveStates.FindStateID(CurrentlyProcessedState);
 		const FStateHandleContext Target = FStateHandleContext(CurrentlyProcessedFrame->StateTree, Request.TargetState);
-		const bool bRequested = RequestTransitionInternal(*CurrentlyProcessedFrame, CurrentlyProcessedStateID, Target, Request.Priority, FStateTreeSharedEvent(), Request.Fallback);
+
+		FTransitionArguments TransitionArgs(*CurrentlyProcessedFrame, CurrentlyProcessedStateID, Target, Request.Priority, FStateTreeSharedEvent(), Request.Fallback);
+		const bool bRequested = RequestTransitionInternal(TransitionArgs);
 		if (bRequested)
 		{
 			check(RequestedTransition.IsValid());
@@ -4958,9 +4960,12 @@ bool FStateTreeExecutionContext::RequestTransition(
 
 	const FStateHandleContext TargetState = FStateHandleContext(CurrentFrame.StateTree, NextState);
 	const FStateTreeSharedEvent Event = TransitionEvent != nullptr ? *TransitionEvent : FStateTreeSharedEvent();
-	return RequestTransitionInternal(CurrentFrame, FActiveStateID(), TargetState, Priority, Event, Fallback);
+
+	FTransitionArguments TransitionArgs(CurrentFrame, FActiveStateID(), TargetState, Priority, Event, Fallback);
+	return RequestTransitionInternal(TransitionArgs);
 }
 
+// Deprecated
 bool FStateTreeExecutionContext::RequestTransitionInternal(
 	const FStateTreeExecutionFrame& SourceFrame,
 	const UE::StateTree::FActiveStateID SourceStateID,
@@ -4968,11 +4973,60 @@ bool FStateTreeExecutionContext::RequestTransitionInternal(
 	const EStateTreeTransitionPriority Priority,
 	const FStateTreeSharedEvent TransitionEvent,
 	const EStateTreeSelectionFallback Fallback)
+{
+	FTransitionArguments TransitionArgs(SourceFrame, SourceStateID, TargetState, Priority, TransitionEvent, Fallback);
+	return RequestTransitionInternal(TransitionArgs);
+}
+
+// Deprecated
+void FStateTreeExecutionContext::SetupNextTransition(const FStateTreeExecutionFrame& CurrentFrame, const FStateTreeStateHandle NextState, const EStateTreeTransitionPriority Priority)
+{
+	using namespace UE::StateTree;
+	using namespace UE::StateTree::ExecutionContext;
+
+	const FStateHandleContext TargetState = FStateHandleContext(CurrentFrame.StateTree, NextState);
+	FTransitionArguments TransitionArgs(CurrentFrame, FActiveStateID(), TargetState, Priority);
+	SetupNextTransition(TransitionArgs);
+}
+
+// Deprecated
+void FStateTreeExecutionContext::SetupNextTransition(
+	const FStateTreeExecutionFrame& SourceFrame,
+	const UE::StateTree::FActiveStateID SourceStateID,
+	const UE::StateTree::ExecutionContext::FStateHandleContext TargetState,
+	const EStateTreeTransitionPriority Priority)
+{
+	FTransitionArguments TransitionArgs(SourceFrame, SourceStateID, TargetState, Priority);
+	RequestedTransition = MakeUnique<FRequestTransitionResult>();
+	SetupNextTransition(TransitionArgs, RequestedTransition->Transition);
+}
+
+// Deprecated
+void FStateTreeExecutionContext::SetupNextTransition(
+	const FStateTreeExecutionFrame& SourceFrame,
+	const UE::StateTree::FActiveStateID SourceStateID,
+	const UE::StateTree::ExecutionContext::FStateHandleContext TargetState,
+	const EStateTreeTransitionPriority Priority,
+	FStateTreeTransitionResult& OutTransitionResult)
+{
+	FTransitionArguments TransitionArgs(SourceFrame, SourceStateID, TargetState, Priority);
+	SetupNextTransition(TransitionArgs, OutTransitionResult);
+}
+
+bool FStateTreeExecutionContext::RequestTransitionInternal(const FTransitionArguments& Args)
 {
 	using namespace UE::StateTree;
 	using namespace UE::StateTree::ExecutionContext;
 	using namespace UE::StateTree::ExecutionContext::Private;
 
+	const FStateTreeExecutionFrame& SourceFrame = Args.SourceFrame;
+	const UE::StateTree::FActiveStateID& SourceStateID = Args.SourceStateID;
+	const UE::StateTree::ExecutionContext::FStateHandleContext& TargetState = Args.TargetState;
+	const EStateTreeTransitionPriority& Priority = Args.Priority;
+	const FStateTreeSharedEvent& TransitionEvent = Args.TransitionEvent;
+	const EStateTreeSelectionFallback& Fallback = Args.Fallback;
+	const EStateTreeStateSelectionRules AdditionalSelectionRules = Args.AdditionalSelectionRules;
+
 	// Skip lower priority transitions.
 	if (RequestedTransition.IsValid() && RequestedTransition->Transition.Priority >= Priority)
 	{
@@ -4981,7 +5035,8 @@ bool FStateTreeExecutionContext::RequestTransitionInternal(
 
 	if (TargetState.StateHandle.IsCompletionState())
 	{
-		SetupNextTransition(SourceFrame, SourceStateID, TargetState, Priority);
+		FTransitionArguments TransitionArgs(SourceFrame, SourceStateID, TargetState, Priority);
+		SetupNextTransition(TransitionArgs);
 		STATETREE_LOG(Verbose, TEXT("Transition on state '%s' -> state '%s'"),
 			*GetSafeStateName(SourceFrame, SourceFrame.ActiveStates.Last()), *TargetState.StateHandle.Describe());
 		return true;
@@ -4989,7 +5044,8 @@ bool FStateTreeExecutionContext::RequestTransitionInternal(
 	if (!TargetState.StateHandle.IsValid() || TargetState.StateTree == nullptr)
 	{
 		// NotSet is no-operation, but can be used to mask a transition at parent state. Returning unset keeps updating current state.
-		SetupNextTransition(SourceFrame, SourceStateID, FStateHandleContext(), Priority);
+		FTransitionArguments TransitionArgs(SourceFrame, SourceStateID, FStateHandleContext(), Priority);
+		SetupNextTransition(TransitionArgs);
 		return true;
 	}
 
@@ -5018,6 +5074,7 @@ bool FStateTreeExecutionContext::RequestTransitionInternal(
 	SelectStateArgs.TransitionEvent = TransitionEvent;
 	SelectStateArgs.Fallback = Fallback;
 	SelectStateArgs.Behavior = ESelectStateBehavior::StateTransition;
+	SelectStateArgs.AdditionalSelectionRules = AdditionalSelectionRules;
 
 	TSharedRef<FSelectStateResult> StateSelectionResult = MakeShared<FSelectStateResult>();
 	TGuardValue<TSharedPtr<FSelectStateResult>> ExecutionFrameHolderScope(CurrentlyProcessedTemporaryStorage, StateSelectionResult);
@@ -5025,7 +5082,8 @@ bool FStateTreeExecutionContext::RequestTransitionInternal(
 	{
 		const FStateTreeExecutionState& Exec = GetExecState();
 
-		SetupNextTransition(SourceFrame, SourceStateID, TargetState, Priority);
+		FTransitionArguments TransitionArgs(SourceFrame, SourceStateID, TargetState, Priority);
+		SetupNextTransition(TransitionArgs);
 		RequestedTransition->Selection = StateSelectionResult; // It will keep the temporary storage alive until the end of the EnterState
 
 		// Fill NextActiveFrames & NextActiveFrameEvents for backward compatibility
@@ -5107,33 +5165,18 @@ PRAGMA_ENABLE_DEPRECATION_WARNINGS
 	return false;
 }
 
-// Deprecated
-void FStateTreeExecutionContext::SetupNextTransition(const FStateTreeExecutionFrame& CurrentFrame, const FStateTreeStateHandle NextState, const EStateTreeTransitionPriority Priority)
-{
-	using namespace UE::StateTree;
-	using namespace UE::StateTree::ExecutionContext;
-
-	const FStateHandleContext TargetState = FStateHandleContext(CurrentFrame.StateTree, NextState);
-	SetupNextTransition(CurrentFrame, FActiveStateID(), TargetState, Priority);
-}
-
-void FStateTreeExecutionContext::SetupNextTransition(
-	const FStateTreeExecutionFrame& SourceFrame,
-	const UE::StateTree::FActiveStateID SourceStateID,
-	const UE::StateTree::ExecutionContext::FStateHandleContext TargetState,
-	const EStateTreeTransitionPriority Priority)
+void FStateTreeExecutionContext::SetupNextTransition(const FTransitionArguments& Args)
 {
 	RequestedTransition = MakeUnique<FRequestTransitionResult>();
-	SetupNextTransition(SourceFrame, SourceStateID, TargetState, Priority, RequestedTransition->Transition);
+	SetupNextTransition(Args, RequestedTransition->Transition);
 }
 
-void FStateTreeExecutionContext::SetupNextTransition(
-	const FStateTreeExecutionFrame& SourceFrame,
-	const UE::StateTree::FActiveStateID SourceStateID,
-	const UE::StateTree::ExecutionContext::FStateHandleContext TargetState,
-	const EStateTreeTransitionPriority Priority,
-	FStateTreeTransitionResult& OutTransitionResult)
+void FStateTreeExecutionContext::SetupNextTransition(const FTransitionArguments& Args, FStateTreeTransitionResult& OutTransitionResult)
 {
+	const FStateTreeExecutionFrame& SourceFrame = Args.SourceFrame;
+	const UE::StateTree::FActiveStateID& SourceStateID = Args.SourceStateID;
+	const UE::StateTree::ExecutionContext::FStateHandleContext& TargetState = Args.TargetState;
+	const EStateTreeTransitionPriority& Priority = Args.Priority;
 	const FStateTreeExecutionState& Exec = GetExecState();
 
 	OutTransitionResult.SourceFrameID = SourceFrame.FrameID;
@@ -5207,7 +5250,8 @@ bool FStateTreeExecutionContext::TriggerTransitions()
 		if (CurrentFrame && CurrentFrame->ActiveStates.Contains(Request.SourceStateID))
 		{
 			UE::StateTree::ExecutionContext::FStateHandleContext TargetState(CurrentFrame->StateTree, Request.TargetState);
-			if (RequestTransitionInternal(*CurrentFrame, Request.SourceStateID, TargetState, Request.Priority, FStateTreeSharedEvent(), Request.Fallback))
+			FTransitionArguments TransitionArgs(*CurrentFrame, Request.SourceStateID, TargetState, Request.Priority, FStateTreeSharedEvent(), Request.Fallback);
+			if (RequestTransitionInternal(TransitionArgs))
 			{
 				check(RequestedTransition.IsValid());
 				RequestedTransition->Source = FStateTreeTransitionSource(CurrentFrame->StateTree, EStateTreeTransitionSourceType::ExternalRequest, Request.TargetState, Request.Priority);
@@ -5435,7 +5479,9 @@ bool FStateTreeExecutionContext::TriggerTransitions()
 
 							// Trigger Delayed Transition when the delay has passed.
 							UE::StateTree::ExecutionContext::FStateHandleContext TargetState(CurrentFrame.StateTree, Transition.State);
-							if (RequestTransitionInternal(CurrentFrame, Handler.StateID, TargetState, Transition.Priority, DelayedTransition.CapturedEvent, Transition.Fallback))
+							EStateTreeStateSelectionRules TransitionRules = Transition.bReactivateTargetState ? EStateTreeStateSelectionRules::ReselectedStateCreatesNewStates : EStateTreeStateSelectionRules::None;
+							FTransitionArguments TransitionArgs(CurrentFrame, Handler.StateID, TargetState, Transition.Priority, DelayedTransition.CapturedEvent, Transition.Fallback, TransitionRules);
+							if (RequestTransitionInternal(TransitionArgs))
 							{
 								// If the transition was successfully requested with a specific event, consume and remove the event, it's been used.
 								if (DelayedTransition.CapturedEvent.IsValid() && Transition.bConsumeEventOnSelect)
@@ -5524,7 +5570,9 @@ PRAGMA_ENABLE_DEPRECATION_WARNINGS
 
 						UE_STATETREE_DEBUG_TRANSITION_EVENT(this, FStateTreeTransitionSource(CurrentFrame.StateTree, FStateTreeIndex16(TransitionIndex), Transition.State, Transition.Priority), EStateTreeTraceEventType::OnRequesting);
 						UE::StateTree::ExecutionContext::FStateHandleContext TargetState(CurrentFrame.StateTree, Transition.State);
-						if (RequestTransitionInternal(CurrentFrame, Handler.StateID, TargetState, Transition.Priority, TransitionEvent, Transition.Fallback))
+						EStateTreeStateSelectionRules TransitionRules = Transition.bReactivateTargetState ? EStateTreeStateSelectionRules::ReselectedStateCreatesNewStates : EStateTreeStateSelectionRules::None;
+						FTransitionArguments TransitionArgs(CurrentFrame, Handler.StateID, TargetState, Transition.Priority, TransitionEvent, Transition.Fallback, TransitionRules);
+						if (RequestTransitionInternal(TransitionArgs))
 						{
 							// If the transition was successfully requested with a specific event, consume and remove the event, it's been used.
 							if (TransitionEvent.IsValid() && Transition.bConsumeEventOnSelect)
@@ -5685,7 +5733,9 @@ PRAGMA_ENABLE_DEPRECATION_WARNINGS
 									// No priority on completion transitions, use the priority to signal that state is selected.
 									UE_STATETREE_DEBUG_TRANSITION_EVENT(this, FStateTreeTransitionSource(CurrentFrame.StateTree, FStateTreeIndex16(TransitionIndex), Transition.State, Transition.Priority), EStateTreeTraceEventType::OnRequesting);
 									UE::StateTree::ExecutionContext::FStateHandleContext TargetState(CurrentFrame.StateTree, Transition.State);
-									if (RequestTransitionInternal(CurrentFrame, CurrentStateID, TargetState, EStateTreeTransitionPriority::Normal, FStateTreeSharedEvent(), Transition.Fallback))
+									EStateTreeStateSelectionRules TransitionRules = Transition.bReactivateTargetState ? EStateTreeStateSelectionRules::ReselectedStateCreatesNewStates : EStateTreeStateSelectionRules::None;
+									FTransitionArguments TransitionArgs(CurrentFrame, CurrentStateID, TargetState, EStateTreeTransitionPriority::Normal, FStateTreeSharedEvent(), Transition.Fallback, TransitionRules);
+									if (RequestTransitionInternal(TransitionArgs))
 									{
 										RequestedTransition->Source = FStateTreeTransitionSource(CurrentFrame.StateTree, FStateTreeIndex16(TransitionIndex), Transition.State, Transition.Priority);
 										break;
@@ -5723,7 +5773,9 @@ PRAGMA_ENABLE_DEPRECATION_WARNINGS
 
 				UE::StateTree::ExecutionContext::FStateHandleContext TargetState(RootFrame.StateTree, FStateTreeStateHandle::Root);
 				UE_STATETREE_DEBUG_TRANSITION_EVENT(this, FStateTreeTransitionSource(RootFrame.StateTree, EStateTreeTransitionSourceType::Internal, FStateTreeStateHandle::Root, EStateTreeTransitionPriority::Normal), EStateTreeTraceEventType::OnRequesting);
-				if (RequestTransitionInternal(RootFrame, RootFrame.ActiveStates.StateIDs[0], TargetState))
+
+				FTransitionArguments TransitionArgs(RootFrame, RootFrame.ActiveStates.StateIDs[0], TargetState);
+				if (RequestTransitionInternal(TransitionArgs))
 				{
 					RequestedTransition->Source = FStateTreeTransitionSource(RootFrame.StateTree, EStateTreeTransitionSourceType::Internal, FStateTreeStateHandle::Root, EStateTreeTransitionPriority::Normal);
 				}
@@ -6579,17 +6631,24 @@ bool FStateTreeExecutionContext::SelectStateInternal(
 		StateSelectionEvents.Add(FStateTreeSharedEvent());
 	}
 
+	EStateTreeStateSelectionRules StateSelectionRules = NextFrame.StateTree->StateSelectionRules | Args.AdditionalSelectionRules;
+
 	bool bShouldCreateNewState = true;
 	const bool bIsNextTargetStateInActiveStates = !bNewFrameCreated
 		&& InternalArgs.MissingActiveStates.Num() > 0
 		&& InternalArgs.MissingActiveStates[0].GetStateHandle() == NextStateHandle
 		&& InternalArgs.MissingActiveStates[0].GetFrameID() == NextFrame.FrameID;
-	if (EnumHasAllFlags(NextFrame.StateTree->StateSelectionRules, EStateTreeStateSelectionRules::CompletedTransitionStatesCreateNewStates))
+	if (EnumHasAnyFlags(StateSelectionRules, EStateTreeStateSelectionRules::ReselectedStateCreatesNewStates | EStateTreeStateSelectionRules::CompletedTransitionStatesCreateNewStates))
 	{
-		// Request a new state if it doesn't match the previous states or if the state is completed.
-		if (bIsNextTargetStateInActiveStates)
+		bShouldCreateNewState = false;
+		if (EnumHasAllFlags(StateSelectionRules, EStateTreeStateSelectionRules::CompletedTransitionStatesCreateNewStates))
 		{
-			bShouldCreateNewState = NextFrame.ActiveTasksStatus.GetStatus(NextState).IsCompleted();
+			// Request a new state if it doesn't match the previous states or if the state is completed.
+			bShouldCreateNewState = !bIsNextTargetStateInActiveStates || NextFrame.ActiveTasksStatus.GetStatus(NextState).IsCompleted();
+		}
+		if (EnumHasAnyFlags(StateSelectionRules, EStateTreeStateSelectionRules::ReselectedStateCreatesNewStates))
+		{
+			bShouldCreateNewState = bShouldCreateNewState || bIsTargetState || !bIsNextTargetStateInActiveStates;
 		}
 	}
 	else
@@ -6780,19 +6839,22 @@ TOptional<bool> FStateTreeExecutionContext::SelectStateFromSourceInternal(
 		const FActiveState& MissingActiveState = InternalArgs.MissingSourceStates[0];
 		bool bContinueWithStateSelection = NextStateHandle != MissingActiveState.GetStateHandle();
 		bool bCompletedState = false;
+
+		EStateTreeStateSelectionRules StateSelectionRules = NextFrame.StateTree->StateSelectionRules | Args.AdditionalSelectionRules;
+
 		if (!bContinueWithStateSelection)
 		{
 			if (InternalArgs.MissingActiveStates.Num() == 1)
 			{
 				// MissingActiveState is the transition source.
-				if (EnumHasAllFlags(NextFrame.StateTree->StateSelectionRules, EStateTreeStateSelectionRules::CompletedTransitionStatesCreateNewStates))
+				if (EnumHasAllFlags(StateSelectionRules, EStateTreeStateSelectionRules::CompletedTransitionStatesCreateNewStates))
 				{
 					bContinueWithStateSelection = NextFrame.ActiveTasksStatus.GetStatus(NextState).IsCompleted();
 				}
 			}
 			else
 			{
-				if (EnumHasAllFlags(NextFrame.StateTree->StateSelectionRules, EStateTreeStateSelectionRules::CompletedTransitionStatesCreateNewStates))
+				if (EnumHasAllFlags(StateSelectionRules, EStateTreeStateSelectionRules::CompletedTransitionStatesCreateNewStates))
 				{
 					bCompletedState = NextFrame.ActiveTasksStatus.GetStatus(NextState).IsCompleted();
 				}
@@ -6821,7 +6883,7 @@ TOptional<bool> FStateTreeExecutionContext::SelectStateFromSourceInternal(
 
 			if (bSelectStateInternalSucceeded && bCompletedState)
 			{
-				check(EnumHasAllFlags(NextFrame.StateTree->StateSelectionRules, EStateTreeStateSelectionRules::CompletedStateBeforeTransitionSourceFailsTransition));
+				check(EnumHasAllFlags(StateSelectionRules, EStateTreeStateSelectionRules::CompletedStateBeforeTransitionSourceFailsTransition));
 				// Before the source, there's a completed state.
 				//The completed state should fail the transition(because of EStateTreeStateSelectionRules::CompletedStateBeforeTargetFailsTransition).
 				//The transition can contain a state with "follow transition" that would remove this state from the transition.
@@ -6864,7 +6926,8 @@ namespace UE::StateTree::ExecutionContext::Private
 	FStateTreeExecutionFrame* SelectedFrameLinkedFrame(FStateTreeExecutionState& Exec,
 		bool bShouldCreateNewState,
 		TArrayView<const UE::StateTree::FActiveState> MatchingActiveStates,
-		const FExecutionFrameKey& LinkStateFrameKey
+		const FExecutionFrameKey& LinkStateFrameKey,
+		EStateTreeStateSelectionRules AdditionalSelectionRules
 	)
 	{
 		FStateTreeExecutionFrame* Result = nullptr;
@@ -6878,7 +6941,8 @@ namespace UE::StateTree::ExecutionContext::Private
 				check(Result);
 				if (Result->HasRoot(LinkStateFrameKey))
 				{
-					if (EnumHasAllFlags(Result->StateTree->GetStateSelectionRules(), EStateTreeStateSelectionRules::CompletedTransitionStatesCreateNewStates)
+					EStateTreeStateSelectionRules StateSelectionRules = Result->StateTree->GetStateSelectionRules() | AdditionalSelectionRules;
+					if (EnumHasAllFlags(StateSelectionRules, EStateTreeStateSelectionRules::CompletedTransitionStatesCreateNewStates)
 						&& Result->ActiveTasksStatus.GetStatus(Result->StateTree).IsCompleted()
 						)
 					{
@@ -6961,7 +7025,7 @@ bool FStateTreeExecutionContext::SelectStateInternal_Linked(
 	}
 
 	// Do we have an existing frame.
-	FStateTreeExecutionFrame* SelectedFrame = SelectedFrameLinkedFrame(Exec, bShouldCreateNewState, InternalArgs.MissingActiveStates, LinkStateFrameKey);
+	FStateTreeExecutionFrame* SelectedFrame = SelectedFrameLinkedFrame(Exec, bShouldCreateNewState, InternalArgs.MissingActiveStates, LinkStateFrameKey, Args.AdditionalSelectionRules);
 	if (SelectedFrame && !SelectedFrame->HasRoot(LinkStateFrameKey))
 	{
 		if (Args.Behavior == ESelectStateBehavior::Forced)
@@ -7118,7 +7182,7 @@ bool FStateTreeExecutionContext::SelectStateInternal_LinkedAsset(
 	}
 
 	// Do we have an existing frame.
-	FStateTreeExecutionFrame* SelectedFrame = SelectedFrameLinkedFrame(Exec, bShouldCreateNewState, InternalArgs.MissingActiveStates, LinkStateFrameKey);
+	FStateTreeExecutionFrame* SelectedFrame = SelectedFrameLinkedFrame(Exec, bShouldCreateNewState, InternalArgs.MissingActiveStates, LinkStateFrameKey, Args.AdditionalSelectionRules);
 	if (SelectedFrame && !SelectedFrame->HasRoot(LinkStateFrameKey))
 	{
 		if (Args.Behavior == ESelectStateBehavior::Forced)
diff --git a/Engine/Plugins/Runtime/StateTree/Source/StateTreeModule/Public/StateTreeExecutionContext.h b/Engine/Plugins/Runtime/StateTree/Source/StateTreeModule/Public/StateTreeExecutionContext.h
--- a/Engine/Plugins/Runtime/StateTree/Source/StateTreeModule/Public/StateTreeExecutionContext.h
+++ b/Engine/Plugins/Runtime/StateTree/Source/StateTreeModule/Public/StateTreeExecutionContext.h
@@ -1109,6 +1109,7 @@ protected:
 		const EStateTreeSelectionFallback Fallback = EStateTreeSelectionFallback::None);
 
 	/** Requests transition to a specified state with specified priority. */
+	UE_DEPRECATED(5.7, "Use the RequestTransitionInternal with TransitionArguments")
 	UE_API bool RequestTransitionInternal(
 		const FStateTreeExecutionFrame& SourceFrame,
 		UE::StateTree::FActiveStateID SourceStateID,
@@ -1117,10 +1118,11 @@ protected:
 		FStateTreeSharedEvent TransitionEvent = FStateTreeSharedEvent(),
 		EStateTreeSelectionFallback Fallback = EStateTreeSelectionFallback::None);
 
-	UE_DEPRECATED(5.7, "Use the SetupNextTransition with the active state ID")
+	UE_DEPRECATED(5.7, "Use the SetupNextTransition with TransitionArguments")
 	UE_API void SetupNextTransition(const FStateTreeExecutionFrame& CurrentFrame, const FStateTreeStateHandle NextState, const EStateTreeTransitionPriority Priority);
 
 	/** Sets up NextTransition based on the provided parameters and the current execution status. */
+	UE_DEPRECATED(5.7, "Use the SetupNextTransition with TransitionArguments")
 	UE_API void SetupNextTransition(
 		const FStateTreeExecutionFrame& SourceFrame,
 		UE::StateTree::FActiveStateID SourceStateID,
@@ -1128,12 +1130,51 @@ protected:
 		EStateTreeTransitionPriority Priority);
 
 	/** Sets up NextTransition based on the provided parameters and the current execution status. */
+	UE_DEPRECATED(5.7, "Use the SetupNextTransition with TransitionArguments")
 	UE_API void SetupNextTransition(
 		const FStateTreeExecutionFrame& SourceFrame,
 		UE::StateTree::FActiveStateID SourceStateID,
 		UE::StateTree::ExecutionContext::FStateHandleContext TargetState,
 		EStateTreeTransitionPriority Priority,
 		FStateTreeTransitionResult& OutTransitionResult);
+	
+	struct FTransitionArguments
+	{
+		FTransitionArguments(const FStateTreeExecutionFrame& InSourceFrame
+			, UE::StateTree::FActiveStateID InSourceStateID
+			, UE::StateTree::ExecutionContext::FStateHandleContext InTargetState
+			, EStateTreeTransitionPriority InPriority = EStateTreeTransitionPriority::Normal
+			, FStateTreeSharedEvent InTransitionEvent = FStateTreeSharedEvent()
+			, EStateTreeSelectionFallback InFallback = EStateTreeSelectionFallback::None
+			, EStateTreeStateSelectionRules InAdditionalSelectionRules = EStateTreeStateSelectionRules::None)
+			:
+			SourceFrame(InSourceFrame)
+			, SourceStateID(InSourceStateID)
+			, TargetState(InTargetState)
+			, Priority(InPriority)
+			, TransitionEvent(InTransitionEvent)
+			, Fallback(InFallback)
+			, AdditionalSelectionRules(InAdditionalSelectionRules)
+		{
+		};
+
+		const FStateTreeExecutionFrame& SourceFrame;
+		UE::StateTree::FActiveStateID SourceStateID;
+		UE::StateTree::ExecutionContext::FStateHandleContext TargetState;
+		EStateTreeTransitionPriority Priority = EStateTreeTransitionPriority::Normal;
+		FStateTreeSharedEvent TransitionEvent = FStateTreeSharedEvent();
+		EStateTreeSelectionFallback Fallback = EStateTreeSelectionFallback::None;
+		EStateTreeStateSelectionRules AdditionalSelectionRules = EStateTreeStateSelectionRules::None;
+	};
+
+	/** Requests transition to a specified state with specified priority. */
+	UE_API bool RequestTransitionInternal(const FTransitionArguments& Args);
+	
+	/** Sets up NextTransition based on the provided parameters and the current execution status. */
+	UE_API void SetupNextTransition(const FTransitionArguments& Args);
+
+	/** Sets up NextTransition based on the provided parameters and the current execution status. */
+	UE_API void SetupNextTransition(const FTransitionArguments& Args, FStateTreeTransitionResult& OutTransitionResult);
 
 	/** Tick the states transitions. */
 	UE_API bool TriggerTransitions();
@@ -1210,6 +1251,8 @@ PRAGMA_ENABLE_DEPRECATION_WARNINGS
 		EStateTreeSelectionFallback Fallback = EStateTreeSelectionFallback::None;
 		/** The type of selection. */
 		ESelectStateBehavior Behavior = ESelectStateBehavior::StateTransition;
+		/** True implies we should always create a new state. */
+		EStateTreeStateSelectionRules AdditionalSelectionRules = EStateTreeStateSelectionRules::None;
 	};
 
 	/**
diff --git a/Engine/Plugins/Runtime/StateTree/Source/StateTreeModule/Public/StateTreeTypes.h b/Engine/Plugins/Runtime/StateTree/Source/StateTreeModule/Public/StateTreeTypes.h
--- a/Engine/Plugins/Runtime/StateTree/Source/StateTreeModule/Public/StateTreeTypes.h
+++ b/Engine/Plugins/Runtime/StateTree/Source/StateTreeModule/Public/StateTreeTypes.h
@@ -102,6 +102,11 @@ enum class EStateTreeStateSelectionRules : uint32
 	 * Without this rule, a transition that triggers on completion might take an extra tick to execute.
 	 */
 	CompletedStateBeforeTransitionSourceFailsTransition = 1 << 1,
+	/**
+	 * When selecting a parent state(from a child), the target child states are new (recreated if needed).
+	 * Without this rule, the "reselected" states are sustained.
+	 */
+	ReselectedStateCreatesNewStates = 1 << 2,
 
 	Default = CompletedTransitionStatesCreateNewStates | CompletedStateBeforeTransitionSourceFailsTransition,
 };
@@ -635,6 +640,10 @@ struct FCompactStateTransition
 	UPROPERTY()
 	uint8 bTransitionEnabled : 1 = true;
 
+	/** The target state and all the following states will be reactivated. Their instance data will be reinstantiated and will receive EnterState. */
+	UPROPERTY()
+	uint8 bReactivateTargetState : 1 = false;
+
 	/** If set to true, the required event is consumed (later state selection cannot react to it) if state selection can be made. */
 	UPROPERTY()
 	uint8 bConsumeEventOnSelect : 1 = true;
diff --git a/Engine/Plugins/Runtime/StateTree/Source/StateTreeTestSuite/Private/StateTreeForceTransitionTest.cpp b/Engine/Plugins/Runtime/StateTree/Source/StateTreeTestSuite/Private/StateTreeForceTransitionTest.cpp
--- a/Engine/Plugins/Runtime/StateTree/Source/StateTreeTestSuite/Private/StateTreeForceTransitionTest.cpp
+++ b/Engine/Plugins/Runtime/StateTree/Source/StateTreeTestSuite/Private/StateTreeForceTransitionTest.cpp
@@ -333,27 +333,69 @@ struct FStateTreeTest_ForceTransition_All : FStateTreeTestBase
 				return true;
 			};
 
+		constexpr int32 MaxRules = 4;
+		auto MakeStateSelectionRule = [](int32 Index)
+			{
+				EStateTreeStateSelectionRules Rule = EStateTreeStateSelectionRules::None;
+				if ((Index % 2) == 1)
+				{
+					Rule |= EStateTreeStateSelectionRules::CompletedTransitionStatesCreateNewStates;
+					Rule |= EStateTreeStateSelectionRules::CompletedStateBeforeTransitionSourceFailsTransition;
+				}
+				if (Index >= 2)
+				{
+					Rule |= EStateTreeStateSelectionRules::ReselectedStateCreatesNewStates;
+				}
+				return Rule;
+			};
+
+
+		for (int32 Index = 0; Index < MaxRules; ++Index)
 		{
-			FTestStateTreeExecutionContext Exec(*StateTree1, *StateTree1, InstanceData);
-			Exec.Start();
-			AITEST_TRUE(TEXT("In correct states"), Exec.ExpectInActiveStates("Tree1StateA"));
-		}
-		{
-			FTestStateTreeExecutionContext Exec(*StateTree1, *StateTree1, OldForceInstanceData);
-			Exec.Start();
-			AITEST_TRUE(TEXT("In correct states"), Exec.ExpectInActiveStates("Tree1StateA"));
-		}
-		bool bResult = TestAllTransitions(TestAllTransitions, 0, TArray<FName>(), true);
-		{
-			FTestStateTreeExecutionContext Exec(*StateTree1, *StateTree1, InstanceData);
-			Exec.Stop();
-		}
-		{
-			FTestStateTreeExecutionContext Exec(*StateTree1, *StateTree1, OldForceInstanceData);
-			Exec.Stop();
+			const EStateTreeStateSelectionRules StateSelectionRules = MakeStateSelectionRule(Index);
+			InstanceData = FStateTreeInstanceData();
+			OldForceInstanceData = FStateTreeInstanceData();
+
+			auto CompileTree = [StateSelectionRules](TNotNull<UStateTree*> StateTree)
+				{
+					StateTree->ResetCompiled();
+					UStateTreeEditorData* EditorData = CastChecked<UStateTreeEditorData>(StateTree->EditorData);
+					UStateTreeTestSchema* Schema = CastChecked<UStateTreeTestSchema>(EditorData->Schema);
+					Schema->SetStateSelectionRules(StateSelectionRules);
+
+					FStateTreeCompilerLog Log;
+					FStateTreeCompiler Compiler(Log);
+					return Compiler.Compile(StateTree);
+				};
+
+			AITEST_TRUE("StateTree3 should get compiled", CompileTree(StateTree3));
+			AITEST_TRUE("StateTree2 should get compiled", CompileTree(StateTree2));
+			AITEST_TRUE("StateTree1 should get compiled", CompileTree(StateTree1));
+
+			{
+				FTestStateTreeExecutionContext Exec(*StateTree1, *StateTree1, InstanceData);
+				Exec.Start();
+				AITEST_TRUE(TEXT("In correct states"), Exec.ExpectInActiveStates("Tree1StateA"));
+			}
+			{
+				FTestStateTreeExecutionContext Exec(*StateTree1, *StateTree1, OldForceInstanceData);
+				Exec.Start();
+				AITEST_TRUE(TEXT("In correct states"), Exec.ExpectInActiveStates("Tree1StateA"));
+			}
+			bool bCurrentResult = TestAllTransitions(TestAllTransitions, 0, TArray<FName>(), true);
+			{
+				FTestStateTreeExecutionContext Exec(*StateTree1, *StateTree1, InstanceData);
+				Exec.Stop();
+			}
+			{
+				FTestStateTreeExecutionContext Exec(*StateTree1, *StateTree1, OldForceInstanceData);
+				Exec.Stop();
+			}
+
+			AITEST_TRUE("Test failed", bCurrentResult);
 		}
 
-		return bResult;
+		return true;
 	}
 };
 IMPLEMENT_STATE_TREE_INSTANT_TEST(FStateTreeTest_ForceTransition_All, "System.StateTree.ForceTransition.All");
diff --git a/Engine/Plugins/Runtime/StateTree/Source/StateTreeTestSuite/Private/StateTreeTransitionTest.cpp b/Engine/Plugins/Runtime/StateTree/Source/StateTreeTestSuite/Private/StateTreeTransitionTest.cpp
--- a/Engine/Plugins/Runtime/StateTree/Source/StateTreeTestSuite/Private/StateTreeTransitionTest.cpp
+++ b/Engine/Plugins/Runtime/StateTree/Source/StateTreeTestSuite/Private/StateTreeTransitionTest.cpp
@@ -1134,7 +1134,7 @@ struct FStateTreeTest_RequestTransition : FStateTreeTestBase
 			AITEST_TRUE("StateTree should init", bInitSucceeded);
 		}
 
-		constexpr int32 MaxRules = 2;
+		constexpr int32 MaxRules = 4;
 		auto MakeStateSelectionRule = [](int32 Index)
 			{
 				EStateTreeStateSelectionRules Rule = EStateTreeStateSelectionRules::None;
@@ -1143,6 +1143,10 @@ struct FStateTreeTest_RequestTransition : FStateTreeTestBase
 					Rule |= EStateTreeStateSelectionRules::CompletedTransitionStatesCreateNewStates;
 					Rule |= EStateTreeStateSelectionRules::CompletedStateBeforeTransitionSourceFailsTransition;
 				}
+				if (Index >= 2)
+				{
+					Rule |= EStateTreeStateSelectionRules::ReselectedStateCreatesNewStates;
+				}
 				return Rule;
 			};
 
@@ -1306,7 +1310,7 @@ struct FStateTreeTest_RequestTransition : FStateTreeTestBase
 				TransitionToExecute = ECustomFunctionToRun::None;
 				AllStateHandle.bFinishTasksD = false;
 			}
-			// Select a new child but B is completed (before the source/target
+			// Select a new child but B is completed (before the source/target)
 			{
 				TransitionToExecute = ECustomFunctionToRun::TransitionD_To_E;
 				AllStateHandle.bFinishTasksB = true;
@@ -1316,6 +1320,7 @@ struct FStateTreeTest_RequestTransition : FStateTreeTestBase
 				}
 
 				const bool bUseCompletedRule = EnumHasAllFlags(StateSelectionRules, EStateTreeStateSelectionRules::CompletedStateBeforeTransitionSourceFailsTransition);
+				const bool bUseReselectedRule = EnumHasAllFlags(StateSelectionRules, EStateTreeStateSelectionRules::ReselectedStateCreatesNewStates);
 
 				FTestStateTreeExecutionContext Exec(StateTree, StateTree, InstanceData);
 				const EStateTreeRunStatus Status = Exec.Tick(0.1f);
@@ -1338,15 +1343,18 @@ struct FStateTreeTest_RequestTransition : FStateTreeTestBase
 				AITEST_TRUE(TEXT("Tick StateD"), LogOrder);
 				if (bUseCompletedRule)
 				{
+					const TCHAR* SchemaEnterStateStr = bUseReselectedRule ? EnterStateChangedStr : EnterStateSustainedStr;
+					const TCHAR* SchemaExitStateStr = bUseReselectedRule ? ExitStateChangedStr : ExitStateSustainedStr;
+
 					LogOrder = LogOrder.Then("StateDTask", ExitStateChangedStr);
 					AITEST_TRUE(TEXT("Enter StateD"), LogOrder);
 					LogOrder = LogOrder.Then("StateCTask", ExitStateChangedStr);
 					AITEST_TRUE(TEXT("Enter StateC"), LogOrder);
 					LogOrder = LogOrder.Then("StateBTask", ExitStateChangedStr);
 					AITEST_TRUE(TEXT("Enter StateB"), LogOrder);
-					LogOrder = LogOrder.Then("StateATask", ExitStateSustainedStr);
+					LogOrder = LogOrder.Then("StateATask", SchemaExitStateStr);
 					AITEST_TRUE(TEXT("Enter StateA"), LogOrder);
-					LogOrder = LogOrder.Then("StateATask", EnterStateSustainedStr);
+					LogOrder = LogOrder.Then("StateATask", SchemaEnterStateStr);
 					AITEST_TRUE(TEXT("Enter StateA"), LogOrder);
 					LogOrder = LogOrder.Then("StateBTask", EnterStateChangedStr);
 					AITEST_TRUE(TEXT("Enter StateB"), LogOrder);
@@ -1424,6 +1432,10 @@ struct FStateTreeTest_RequestTransition : FStateTreeTestBase
 					return false;
 				}
 
+				const bool bUseReselectedRule = EnumHasAllFlags(StateSelectionRules, EStateTreeStateSelectionRules::ReselectedStateCreatesNewStates);
+				const TCHAR* SchemaEnterStateStr = bUseReselectedRule ? EnterStateChangedStr : EnterStateSustainedStr;
+				const TCHAR* SchemaExitStateStr = bUseReselectedRule ? ExitStateChangedStr : ExitStateSustainedStr;
+				 
 				FTestStateTreeExecutionContext Exec(StateTree, StateTree, InstanceData);
 				const EStateTreeRunStatus Status = Exec.Tick(0.1f);
 				AITEST_EQUAL(TEXT("Tick should complete with Running"), Status, EStateTreeRunStatus::Running);
@@ -1436,14 +1448,16 @@ struct FStateTreeTest_RequestTransition : FStateTreeTestBase
 				AITEST_TRUE(TEXT("Tick StateC"), LogOrder);
 				LogOrder = LogOrder.Then("StateDTask", TEXT("Tick0"));
 				AITEST_TRUE(TEXT("Tick StateD"), LogOrder);
-				LogOrder = LogOrder.Then("StateDTask", ExitStateSustainedStr);
+	
+				LogOrder = LogOrder.Then("StateDTask", SchemaExitStateStr);
 				AITEST_TRUE(TEXT("Exit StateD"), LogOrder);
-				LogOrder = LogOrder.Then("StateCTask", ExitStateSustainedStr);
+				LogOrder = LogOrder.Then("StateCTask", SchemaExitStateStr);
 				AITEST_TRUE(TEXT("Exit StateC"), LogOrder);
-				LogOrder = LogOrder.Then("StateCTask", EnterStateSustainedStr);
+				LogOrder = LogOrder.Then("StateCTask", SchemaEnterStateStr);
 				AITEST_TRUE(TEXT("Enter StateC"), LogOrder);
-				LogOrder = LogOrder.Then("StateDTask", EnterStateSustainedStr);
+				LogOrder = LogOrder.Then("StateDTask", SchemaEnterStateStr);
 				AITEST_TRUE(TEXT("Enter StateD"), LogOrder);
+
 				AITEST_FALSE(TEXT("Enter StateA"), Exec.Expect("StateATask", TEXT("EnterState0")));
 				AITEST_FALSE(TEXT("Enter StateB"), Exec.Expect("StateBTask", TEXT("EnterState0")));
 				AITEST_FALSE(TEXT("Exit StateA"), Exec.Expect("StateATask", TEXT("ExitState0")));
@@ -1461,8 +1475,9 @@ struct FStateTreeTest_RequestTransition : FStateTreeTestBase
 				}
 
 				const bool bUseCompletedRule = EnumHasAllFlags(StateSelectionRules, EStateTreeStateSelectionRules::CompletedTransitionStatesCreateNewStates);
-				const TCHAR* SchemaEnterStateStr = bUseCompletedRule ? EnterStateChangedStr : EnterStateSustainedStr;
-				const TCHAR* SchemaExitStateStr = bUseCompletedRule ? ExitStateChangedStr : ExitStateSustainedStr;
+				const bool bUseReselectedRule = EnumHasAllFlags(StateSelectionRules, EStateTreeStateSelectionRules::ReselectedStateCreatesNewStates);
+				const TCHAR* SchemaEnterStateStr = bUseCompletedRule || bUseReselectedRule ? EnterStateChangedStr : EnterStateSustainedStr;
+				const TCHAR* SchemaExitStateStr = bUseCompletedRule || bUseReselectedRule ? ExitStateChangedStr : ExitStateSustainedStr;
 
 				FTestStateTreeExecutionContext Exec(StateTree, StateTree, InstanceData);
 				const EStateTreeRunStatus Status = Exec.Tick(0.1f);
@@ -1500,6 +1515,10 @@ struct FStateTreeTest_RequestTransition : FStateTreeTestBase
 					return false;
 				}
 
+				const bool bUseReselectedRule = EnumHasAllFlags(StateSelectionRules, EStateTreeStateSelectionRules::ReselectedStateCreatesNewStates);
+				const TCHAR* SchemaEnterStateStr = bUseReselectedRule ? EnterStateChangedStr : EnterStateSustainedStr;
+				const TCHAR* SchemaExitStateStr = bUseReselectedRule ? ExitStateChangedStr : ExitStateSustainedStr;
+
 				FTestStateTreeExecutionContext Exec(StateTree, StateTree, InstanceData);
 				const EStateTreeRunStatus Status = Exec.Tick(0.1f);
 				AITEST_EQUAL(TEXT("Tick should complete with Running"), Status, EStateTreeRunStatus::Running);
@@ -1512,17 +1531,17 @@ struct FStateTreeTest_RequestTransition : FStateTreeTestBase
 				AITEST_TRUE(TEXT("Tick StateC"), LogOrder);
 				LogOrder = LogOrder.Then("StateDTask", TEXT("Tick0"));
 				AITEST_TRUE(TEXT("Tick StateD"), LogOrder);
-				LogOrder = LogOrder.Then("StateDTask", ExitStateSustainedStr);
+				LogOrder = LogOrder.Then("StateDTask", SchemaExitStateStr);
 				AITEST_TRUE(TEXT("Exit StateD"), LogOrder);
-				LogOrder = LogOrder.Then("StateCTask", ExitStateSustainedStr);
+				LogOrder = LogOrder.Then("StateCTask", SchemaExitStateStr);
 				AITEST_TRUE(TEXT("Exit StateC"), LogOrder);
-				LogOrder = LogOrder.Then("StateBTask", ExitStateSustainedStr);
+				LogOrder = LogOrder.Then("StateBTask", SchemaExitStateStr);
 				AITEST_TRUE(TEXT("Exit StateB"), LogOrder);
-				LogOrder = LogOrder.Then("StateBTask", EnterStateSustainedStr);
+				LogOrder = LogOrder.Then("StateBTask", SchemaEnterStateStr);
 				AITEST_TRUE(TEXT("Enter StateB"), LogOrder);
-				LogOrder = LogOrder.Then("StateCTask", EnterStateSustainedStr);
+				LogOrder = LogOrder.Then("StateCTask", SchemaEnterStateStr);
 				AITEST_TRUE(TEXT("Enter StateC"), LogOrder);
-				LogOrder = LogOrder.Then("StateDTask", EnterStateSustainedStr);
+				LogOrder = LogOrder.Then("StateDTask", SchemaEnterStateStr);
 				AITEST_TRUE(TEXT("Enter StateD"), LogOrder);
 				AITEST_FALSE(TEXT("Enter StateA"), Exec.Expect("StateATask", TEXT("EnterState0")));
 				AITEST_FALSE(TEXT("Exit StateA"), Exec.Expect("StateATask", TEXT("ExitState0")));
@@ -1540,8 +1559,9 @@ struct FStateTreeTest_RequestTransition : FStateTreeTestBase
 				}
 
 				const bool bUseCompletedRule = EnumHasAllFlags(StateSelectionRules, EStateTreeStateSelectionRules::CompletedTransitionStatesCreateNewStates);
-				const TCHAR* SchemaEnterStateStr = bUseCompletedRule ? EnterStateChangedStr : EnterStateSustainedStr;
-				const TCHAR* SchemaExitStateStr = bUseCompletedRule ? ExitStateChangedStr : ExitStateSustainedStr;
+				const bool bUseReselectedRule = EnumHasAllFlags(StateSelectionRules, EStateTreeStateSelectionRules::ReselectedStateCreatesNewStates);
+				const TCHAR* SchemaEnterStateStr = bUseCompletedRule || bUseReselectedRule ? EnterStateChangedStr : EnterStateSustainedStr;
+				const TCHAR* SchemaExitStateStr = bUseCompletedRule || bUseReselectedRule ? ExitStateChangedStr : ExitStateSustainedStr;
 
 				FTestStateTreeExecutionContext Exec(StateTree, StateTree, InstanceData);
 				const EStateTreeRunStatus Status = Exec.Tick(0.1f);
@@ -1558,10 +1578,20 @@ struct FStateTreeTest_RequestTransition : FStateTreeTestBase
 				LogOrder = LogOrder.Then("StateDTask", SchemaExitStateStr);
 				AITEST_TRUE(TEXT("Exit StateD"), LogOrder);
 				LogOrder = LogOrder.Then("StateCTask", SchemaExitStateStr);
-				AITEST_TRUE(TEXT("Exit StateC"), LogOrder);
-				LogOrder = LogOrder.Then("StateBTask", ExitStateSustainedStr);
-				AITEST_TRUE(TEXT("Exit StateB"), LogOrder);
-				LogOrder = LogOrder.Then("StateBTask", EnterStateSustainedStr);
+				if (bUseReselectedRule)
+				{
+					AITEST_TRUE(TEXT("Exit StateC"), LogOrder);
+					LogOrder = LogOrder.Then("StateBTask", ExitStateChangedStr);
+					AITEST_TRUE(TEXT("Exit StateB"), LogOrder);
+					LogOrder = LogOrder.Then("StateBTask", EnterStateChangedStr);
+				}
+				else
+				{
+					AITEST_TRUE(TEXT("Exit StateC"), LogOrder);
+					LogOrder = LogOrder.Then("StateBTask", ExitStateSustainedStr);
+					AITEST_TRUE(TEXT("Exit StateB"), LogOrder);
+					LogOrder = LogOrder.Then("StateBTask", EnterStateSustainedStr);
+				}
 				AITEST_TRUE(TEXT("Enter StateB"), LogOrder);
 				LogOrder = LogOrder.Then("StateCTask", SchemaEnterStateStr);
 				AITEST_TRUE(TEXT("Enter StateC"), LogOrder);
@@ -1616,6 +1646,9 @@ struct FStateTreeTest_RequestTransition : FStateTreeTestBase
 				{
 					return false;
 				}
+				const bool bUseReselectedRule = EnumHasAllFlags(StateSelectionRules, EStateTreeStateSelectionRules::ReselectedStateCreatesNewStates);
+				const TCHAR* SchemaEnterStateStr = bUseReselectedRule ? EnterStateChangedStr : EnterStateSustainedStr;
+				const TCHAR* SchemaExitStateStr = bUseReselectedRule ? ExitStateChangedStr : ExitStateSustainedStr;
 
 				FTestStateTreeExecutionContext Exec(StateTree, StateTree, InstanceData);
 				const EStateTreeRunStatus Status = Exec.Tick(0.1f);
@@ -1629,17 +1662,17 @@ struct FStateTreeTest_RequestTransition : FStateTreeTestBase
 				AITEST_TRUE(TEXT("Tick StateC"), LogOrder);
 				LogOrder = LogOrder.Then("StateDTask", TEXT("Tick0"));
 				AITEST_TRUE(TEXT("Tick StateD"), LogOrder);
-				LogOrder = LogOrder.Then("StateDTask", ExitStateSustainedStr);
+				LogOrder = LogOrder.Then("StateDTask", SchemaExitStateStr);
 				AITEST_TRUE(TEXT("Exit StateD"), LogOrder);
-				LogOrder = LogOrder.Then("StateCTask", ExitStateSustainedStr);
+				LogOrder = LogOrder.Then("StateCTask", SchemaExitStateStr);
 				AITEST_TRUE(TEXT("Exit StateC"), LogOrder);
-				LogOrder = LogOrder.Then("StateBTask", ExitStateSustainedStr);
+				LogOrder = LogOrder.Then("StateBTask", SchemaExitStateStr);
 				AITEST_TRUE(TEXT("Exit StateB"), LogOrder);
-				LogOrder = LogOrder.Then("StateBTask", EnterStateSustainedStr);
+				LogOrder = LogOrder.Then("StateBTask", SchemaEnterStateStr);
 				AITEST_TRUE(TEXT("Enter StateB"), LogOrder);
-				LogOrder = LogOrder.Then("StateCTask", EnterStateSustainedStr);
+				LogOrder = LogOrder.Then("StateCTask", SchemaEnterStateStr);
 				AITEST_TRUE(TEXT("Enter StateC"), LogOrder);
-				LogOrder = LogOrder.Then("StateDTask", EnterStateSustainedStr);
+				LogOrder = LogOrder.Then("StateDTask", SchemaEnterStateStr);
 				AITEST_TRUE(TEXT("Enter StateD"), LogOrder);
 				AITEST_FALSE(TEXT("Enter StateA"), Exec.Expect("StateATask", TEXT("EnterState0")));
 				AITEST_FALSE(TEXT("Exit StateA"), Exec.Expect("StateATask", TEXT("ExitState0")));
