[State Tree] fixed temporary frames not cleaned up when we run multiple transitions and did State selection more than once
minor fix to comments and code

diff --git a/Engine/Plugins/Runtime/StateTree/Source/StateTreeModule/Private/StateTreeExecutionContext.cpp b/Engine/Plugins/Runtime/StateTree/Source/StateTreeModule/Private/StateTreeExecutionContext.cpp
--- a/Engine/Plugins/Runtime/StateTree/Source/StateTreeModule/Private/StateTreeExecutionContext.cpp
+++ b/Engine/Plugins/Runtime/StateTree/Source/StateTreeModule/Private/StateTreeExecutionContext.cpp
@@ -5077,9 +5077,43 @@ bool FStateTreeExecutionContext::RequestTransitionInternal(const FTransitionArgu
 	SelectStateArgs.AdditionalSelectionRules = AdditionalSelectionRules;
 
 	TSharedRef<FSelectStateResult> StateSelectionResult = MakeShared<FSelectStateResult>();
-	TGuardValue<TSharedPtr<FSelectStateResult>> ExecutionFrameHolderScope(CurrentlyProcessedTemporaryStorage, StateSelectionResult);
+	TGuardValue<TSharedPtr<FSelectStateResult>> TemporaryStorageScope(CurrentlyProcessedTemporaryStorage, StateSelectionResult);
 	if (SelectState(SelectStateArgs, StateSelectionResult))
 	{
+		if (RequestedTransition)
+		{
+			// If we have a previous selection(i.e. we succeeded a selection already from a previous transition), we need to clean up temporary frames if any.
+			if (FSelectStateResult* PrevSelectStateResult = RequestedTransition->Selection.Get())
+			{
+				TGuardValue<TSharedPtr<FSelectStateResult>> PrevTemporaryStorageScope(CurrentlyProcessedTemporaryStorage, RequestedTransition->Selection);
+
+				for (int32 Idx = PrevSelectStateResult->SelectedFrames.Num() - 1; Idx >= 0; --Idx)
+				{
+					const FActiveFrameID TemporaryFrameID = PrevSelectStateResult->SelectedFrames[Idx];
+
+					FSelectStateResult::FFrameAndParent TemporaryFrameAndParent = PrevSelectStateResult->GetExecutionFrame(TemporaryFrameID);
+					FStateTreeExecutionFrame* TemporaryFrame = TemporaryFrameAndParent.Frame;
+
+					// stop as soon as the frame is active already: no more temporary frames forward
+					if (!TemporaryFrame)
+					{
+						break;
+					}
+
+					FStateTreeExecutionState& Exec = Storage.GetMutableExecutionState();
+					const FStateTreeExecutionFrame* ParentFrame = FindExecutionFrame(TemporaryFrameAndParent.ParentFrameID, MakeConstArrayView(Exec.ActiveFrames), MakeConstArrayView(PrevSelectStateResult->TemporaryFrames));
+
+					if (TemporaryFrame->bIsGlobalFrame && TemporaryFrame->bHaveEntered)
+					{
+						constexpr EStateTreeRunStatus Status = EStateTreeRunStatus::Stopped;
+						StopTemporaryEvaluatorsAndGlobalTasks(ParentFrame, *TemporaryFrame, Status);
+					}
+
+					CleanFrame(Exec, TemporaryFrameID);
+				}
+			}
+		}
+
 		const FStateTreeExecutionState& Exec = GetExecState();
 
 		FTransitionArguments TransitionArgs(SourceFrame, SourceStateID, TargetState, Priority);
@@ -5221,7 +5255,7 @@ bool FStateTreeExecutionContext::TriggerTransitions()
 	// If delayed,
 	//	If delayed completed, then process.
 	//	Else add them to the delayed transition list.
-	//3. If no transition, Process completion transitions, from top to bottom.
+	//3. If no transition, Process completion transitions, from bottom to top.
 	//4. If transition occurs, check if there are any frame (sub-tree) that completed.
 
 	CSV_SCOPED_TIMING_STAT_EXCLUSIVE(StateTree_TriggerTransition);
@@ -6104,7 +6138,7 @@ bool FStateTreeExecutionContext::SelectState(const FSelectStateArguments& Args,
 	// 
 	//           ABCD   | AB     | AX     | AXY
 	//  Source is not in target. New XY.
-	//  ExitState: Sustained -. Changed B. EnterState: Sustained -. Changed X
+	//  ExitState: Sustained -. Changed BCD. EnterState: Sustained -. Changed XY
 	//
 	//           ABCD   | AB     | AB     | ABCD, ABCD', ABC'D', AB'C'D'
 	//  Source is target. New D, CD, BCD if completed.
@@ -6389,7 +6423,7 @@ bool FStateTreeExecutionContext::SelectStateInternal(
 		return false;
 	}
 
-	const FCompactStateTreeState& NextState = InternalArgs.MissingStatesToReachTarget[0].StateTree->States[NextStateHandle.Index];
+	const FCompactStateTreeState& NextState = NextStateTree->States[NextStateHandle.Index];
 	if (!NextState.bEnabled)
 	{
 		STATETREE_LOG(VeryVerbose, TEXT("%hs: Ignoring disabled state '%s'.  '%s' using StateTree '%s'."),
diff --git a/Engine/Plugins/Runtime/StateTree/Source/StateTreeTestSuite/Private/StateTreeTransitionTest.cpp b/Engine/Plugins/Runtime/StateTree/Source/StateTreeTestSuite/Private/StateTreeTransitionTest.cpp
--- a/Engine/Plugins/Runtime/StateTree/Source/StateTreeTestSuite/Private/StateTreeTransitionTest.cpp
+++ b/Engine/Plugins/Runtime/StateTree/Source/StateTreeTestSuite/Private/StateTreeTransitionTest.cpp
@@ -1971,6 +1971,110 @@ struct FStateTreeTest_TransitionLinkedAssetWith2Root : FStateTreeTestBase
 };
 IMPLEMENT_STATE_TREE_INSTANT_TEST(FStateTreeTest_TransitionLinkedAssetWith2Root, "System.StateTree.Transition.LinkedAssetWith2Root");
 
+struct FStateTreeTest_MultipleTransition_TemporaryFrame_GlobalTask : FStateTreeTestBase
+{
+	virtual bool InstantTest() override
+	{
+		/*
+		- Tree1
+			- Root
+				- State1 Transition1 -> Goto State2; Transition2 -> Goto State3
+				- State2 : Tree2
+				- State3
+
+		- Tree2 : Global Tasks
+			- Root : Tree3
+
+		- Tree3 : Global Tasks
+		*/
+
+		UStateTree& StateTree1 = NewStateTree();
+		UStateTree& StateTree2 = NewStateTree();
+		UStateTree& StateTree3 = NewStateTree();
+
+		{
+			UStateTreeEditorData& Tree1EditorData = *Cast<UStateTreeEditorData>(StateTree1.EditorData);
+			UStateTreeState& Root = Tree1EditorData.AddSubTree(FName("Tree1Root"));
+			UStateTreeState& State1 = Root.AddChildState(FName("Tree1State1"));
+			UStateTreeState& State2 = Root.AddChildState(FName("Tree1State2"));
+			UStateTreeState& State3 = Root.AddChildState(FName("Tree1State3"));
+
+			FStateTreeTransition& LowPriorityTrans = State1.AddTransition(EStateTreeTransitionTrigger::OnTick, EStateTreeTransitionType::GotoState, &State2);
+			LowPriorityTrans.Priority = EStateTreeTransitionPriority::Low;
+			FStateTreeTransition& HighPriorityTrans = State1.AddTransition(EStateTreeTransitionTrigger::OnTick, EStateTreeTransitionType::GotoState, &State3);
+			HighPriorityTrans.Priority = EStateTreeTransitionPriority::High;
+
+			State2.Type = EStateTreeStateType::LinkedAsset;
+			State2.SetLinkedStateAsset(&StateTree2);
+		}
+
+		{
+			UStateTreeEditorData& Tree2EditorData = *Cast<UStateTreeEditorData>(StateTree2.EditorData);
+			TStateTreeEditorNode<FTestTask_StandNoTick>& TaskEditorNode = Tree2EditorData.AddGlobalTask<FTestTask_StandNoTick>(FName("Tree2Stand"));
+			UStateTreeState& Root = Tree2EditorData.AddSubTree(FName("Tree2Root"));
+			Root.Type = EStateTreeStateType::LinkedAsset;
+			Root.SetLinkedStateAsset(&StateTree3);
+		}
+
+		{
+			UStateTreeEditorData& Tree3EditorData = *Cast<UStateTreeEditorData>(StateTree3.EditorData);
+			Tree3EditorData.AddGlobalTask<FTestTask_StandNoTick>(FName("Tree3Stand"));
+
+			UStateTreeState& Root = Tree3EditorData.AddSubTree(FName("Tree3Root"));
+		}
+
+		{
+			FStateTreeCompilerLog Log;
+			FStateTreeCompiler Compiler(Log);
+			const bool bResult = Compiler.Compile(StateTree3);
+			AITEST_TRUE("StateTree should get compiled", bResult);
+		}
+		{
+			FStateTreeCompilerLog Log;
+			FStateTreeCompiler Compiler(Log);
+			const bool bResult = Compiler.Compile(StateTree2);
+			AITEST_TRUE("StateTree should get compiled", bResult);
+		}
+		{
+			FStateTreeCompilerLog Log;
+			FStateTreeCompiler Compiler(Log);
+			const bool bResult = Compiler.Compile(StateTree1);
+			AITEST_TRUE("StateTree should get compiled", bResult)
+		}
+
+		{
+			FStateTreeInstanceData InstanceData;
+			FTestStateTreeExecutionContext Exec(StateTree1, StateTree1, InstanceData);
+			const bool bInitSucceeded = Exec.IsValid();
+			AITEST_TRUE("StateTree should init", bInitSucceeded);
+
+			EStateTreeRunStatus Status = Exec.Start();
+			AITEST_TRUE(TEXT("Valid states"), Exec.ExpectInActiveStates("Tree1Root", "Tree1State1"));
+			Exec.LogClear();
+
+			Exec.Tick(0.1f);
+
+			FTestStateTreeExecutionContext::FLogOrder LogOrder = FTestStateTreeExecutionContext::FLogOrder(Exec, 0);
+			LogOrder = LogOrder.Then("Tree2Stand", TEXT("EnterState=Changed")).Then("Tree3Stand", TEXT("EnterState=Changed"));
+			AITEST_TRUE(TEXT("Enter Global tasks on temporary frames correctly"), LogOrder);
+
+			LogOrder = LogOrder.Then("Tree3Stand", TEXT("ExitStopped"))
+				.Then("Tree3Stand", TEXT("ExitState=Changed"))
+				.Then("Tree2Stand", TEXT("ExitStopped"))
+				.Then("Tree2Stand", TEXT("ExitState=Changed"));
+
+			AITEST_TRUE(TEXT("Exit Global tasks on temporary frames correctly"), LogOrder);
+
+			AITEST_TRUE(TEXT("Valid states"), Exec.ExpectInActiveStates("Tree1Root", "Tree1State3"));
+
+			Exec.Stop();
+		}
+
+		return true;
+	}
+};
+IMPLEMENT_STATE_TREE_INSTANT_TEST(FStateTreeTest_MultipleTransition_TemporaryFrame_GlobalTask, "System.StateTree.Transition.MultipleTransition.TemporaryFrame.GlobalTask");
+
 } // namespace UE::StateTree::Tests
 
 UE_ENABLE_OPTIMIZATION_SHIP
