[State Tree] remove redundant struct declaration

diff --git a/Engine/Plugins/Runtime/StateTree/Source/StateTreeModule/Public/StateTree.h b/Engine/Plugins/Runtime/StateTree/Source/StateTreeModule/Public/StateTree.h
--- a/Engine/Plugins/Runtime/StateTree/Source/StateTreeModule/Public/StateTree.h
+++ b/Engine/Plugins/Runtime/StateTree/Source/StateTreeModule/Public/StateTree.h
@@ -558,7 +558,7 @@ private:
 	friend struct FStateTreeMinimalExecutionContext;
 	friend struct FStateTreeReadOnlyExecutionContext;
 	friend struct FStateTreeWeakExecutionContext;
-	friend struct TStateTreeStrongExecutionContext<true>;
+	friend TStateTreeStrongExecutionContext<true>;
 
 #if WITH_EDITOR
 	friend struct FStateTreeCompiler;
