[StateTree] removed the struct redirect for FStateTreePropertyAccess since that struct was not moved to PropertyBindingUtilsModule

diff --git a/Engine/Plugins/Runtime/StateTree/Config/BaseStateTree.ini b/Engine/Plugins/Runtime/StateTree/Config/BaseStateTree.ini
--- a/Engine/Plugins/Runtime/StateTree/Config/BaseStateTree.ini
+++ b/Engine/Plugins/Runtime/StateTree/Config/BaseStateTree.ini
@@ -79,5 +79,4 @@
 +StructRedirects=(OldName="StateTreePropertyPath",NewName="/script/PropertyBindingUtilsModule.PropertyBindingPath")
 
 +StructRedirects=(OldName="StateTreePropertyIndirection",NewName="/script/PropertyBindingUtilsModule.PropertyBindingPropertyIndirection")
-+StructRedirects=(OldName="StateTreePropertyAccess",NewName="/script/PropertyBindingUtilsModule.PropertyBindingPropertyAccess")
 +StructRedirects=(OldName="StateTreePropertyCopy",NewName="/script/PropertyBindingUtilsModule.PropertyBindingCopyInfo")
