[StateTree] fixed unit test setting the same property twice with different container type

diff --git a/Engine/Plugins/Runtime/StateTree/Source/StateTreeTestSuite/Private/StateTreeBindingTest.cpp b/Engine/Plugins/Runtime/StateTree/Source/StateTreeTestSuite/Private/StateTreeBindingTest.cpp
--- a/Engine/Plugins/Runtime/StateTree/Source/StateTreeTestSuite/Private/StateTreeBindingTest.cpp
+++ b/Engine/Plugins/Runtime/StateTree/Source/StateTreeTestSuite/Private/StateTreeBindingTest.cpp
@@ -537,7 +537,6 @@ struct FStateTreeTest_MutableArray : FStateTreeTestBase
 				RootPropertyBag.AddProperty("Value", EPropertyBagPropertyType::Int32);
 				RootPropertyBag.SetValueInt32("Value", -111);
 				RootPropertyBag.AddContainerProperty("ArrayValue", FPropertyBagContainerTypes(EPropertyBagContainerType::Array), EPropertyBagPropertyType::Int32, nullptr);
-				RootPropertyBag.AddProperty("ArrayValue", EPropertyBagPropertyType::Int32);
 				FPropertyBagArrayRef ValueArrayRef = RootPropertyBag.GetMutableArrayRef("ArrayValue").GetValue();
 				ValueArrayRef.EmptyAndAddValues(4);
 				ValueArrayRef.SetValueInt32(0, -11);
