[State Tree]
Attempt to sync StateTreeReference parameters from StateTree on loading, when there is no circular dependency between assets.
The user would have to manually fix up STRefernce in the Outer UObject's PostLoad if there is.

diff --git a/Engine/Plugins/Runtime/StateTree/Source/StateTreeEditorModule/Private/Customizations/StateTreeReferenceDetails.cpp b/Engine/Plugins/Runtime/StateTree/Source/StateTreeEditorModule/Private/Customizations/StateTreeReferenceDetails.cpp
--- a/Engine/Plugins/Runtime/StateTree/Source/StateTreeEditorModule/Private/Customizations/StateTreeReferenceDetails.cpp
+++ b/Engine/Plugins/Runtime/StateTree/Source/StateTreeEditorModule/Private/Customizations/StateTreeReferenceDetails.cpp
@@ -144,7 +144,7 @@ void FStateTreeReferenceDetails::CustomizeHeader(const TSharedRef<IPropertyHandl
 
 	// Registers a delegate to be notified when the associated StateTree asset get successfully recompiled
 	// to make sure that the parameters in the StateTreeReference are still valid.
-	UE::StateTree::Delegates::OnPostCompile.AddSP(this, &FStateTreeReferenceDetails::OnTreeCompiled);
+	OnTreeCompiledHandle = UE::StateTree::Delegates::OnPostCompile.AddSP(this, &FStateTreeReferenceDetails::OnTreeCompiled);
 }
 
 void FStateTreeReferenceDetails::CustomizeChildren(const TSharedRef<IPropertyHandle> InStructPropertyHandle, IDetailChildrenBuilder& InChildrenBuilder, IPropertyTypeCustomizationUtils& InCustomizationUtils)
@@ -154,6 +154,11 @@ void FStateTreeReferenceDetails::CustomizeChildren(const TSharedRef<IPropertyHan
 	InChildrenBuilder.AddCustomBuilder(ParameterInstanceDetails);
 }
 
+FStateTreeReferenceDetails::~FStateTreeReferenceDetails()
+{
+	UE::StateTree::Delegates::OnPostCompile.Remove(OnTreeCompiledHandle);
+}
+
 void FStateTreeReferenceDetails::OnTreeCompiled(const UStateTree& StateTree) const
 {
 	SyncParameters(&StateTree);
diff --git a/Engine/Plugins/Runtime/StateTree/Source/StateTreeEditorModule/Private/Customizations/StateTreeReferenceDetails.h b/Engine/Plugins/Runtime/StateTree/Source/StateTreeEditorModule/Private/Customizations/StateTreeReferenceDetails.h
--- a/Engine/Plugins/Runtime/StateTree/Source/StateTreeEditorModule/Private/Customizations/StateTreeReferenceDetails.h
+++ b/Engine/Plugins/Runtime/StateTree/Source/StateTreeEditorModule/Private/Customizations/StateTreeReferenceDetails.h
@@ -21,6 +21,8 @@ public:
 	virtual void CustomizeHeader(TSharedRef<IPropertyHandle> InStructPropertyHandle, FDetailWidgetRow& InHeaderRow, IPropertyTypeCustomizationUtils& InCustomizationUtils) override;
 	virtual void CustomizeChildren(TSharedRef<IPropertyHandle> InStructPropertyHandle, IDetailChildrenBuilder& InChildrenBuilder, IPropertyTypeCustomizationUtils& InCustomizationUtils) override;
 
+	virtual ~FStateTreeReferenceDetails() override;
+
 private:
 	/**
 	 * Callback registered to StateTree post compilation delegate used to synchronize parameters.
@@ -39,4 +41,6 @@ private:
 	TSharedPtr<IPropertyHandle> StructPropertyHandle;
 	TSharedPtr<IPropertyHandle> ParametersHandle;
 	TSharedPtr<IPropertyUtilities> PropUtils;
+
+	FDelegateHandle OnTreeCompiledHandle;
 };
diff --git a/Engine/Plugins/Runtime/StateTree/Source/StateTreeModule/Private/StateTreeReference.cpp b/Engine/Plugins/Runtime/StateTree/Source/StateTreeModule/Private/StateTreeReference.cpp
--- a/Engine/Plugins/Runtime/StateTree/Source/StateTreeModule/Private/StateTreeReference.cpp
+++ b/Engine/Plugins/Runtime/StateTree/Source/StateTreeModule/Private/StateTreeReference.cpp
@@ -90,7 +90,11 @@ void FStateTreeReference::ConditionallySyncParameters() const
 	{
 		FStateTreeReference* NonConstThis = const_cast<FStateTreeReference*>(this);
 		NonConstThis->SyncParameters();
+
+#if !WITH_EDITORONLY_DATA
+		// Parameters might get fixed up on PIE, but shouldn't at runtime.
 		UE_LOG(LogStateTree, Warning, TEXT("Parameters for '%s' stored in StateTreeReference were auto-fixed to be usable at runtime."), *GetNameSafe(StateTree));	
+#endif //!WITH_EDITORONLY_DATA
 	}
 }
 
@@ -115,16 +119,33 @@ bool FStateTreeReference::Serialize(FStructuredArchive::FSlot Slot)
 
 void FStateTreeReference::PostSerialize(const FArchive& Ar)
 {
-	const int32 CurrentVersion = Ar.CustomVer(FStateTreeCustomVersion::GUID);
-	if (CurrentVersion < FStateTreeCustomVersion::OverridableParameters)
+	if (Ar.IsLoading())
 	{
-		// In earlier versions, all parameters were overwritten.
-		if (const UPropertyBag* Bag = Parameters.GetPropertyBagStruct())
+		const int32 CurrentVersion = Ar.CustomVer(FStateTreeCustomVersion::GUID);
+		if (CurrentVersion < FStateTreeCustomVersion::OverridableParameters)
 		{
-			for (const FPropertyBagPropertyDesc& Desc : Bag->GetPropertyDescs())
+			// In earlier versions, all parameters were overwritten.
+			if (const UPropertyBag* Bag = Parameters.GetPropertyBagStruct())
 			{
-				PropertyOverrides.Add(Desc.ID);
+				for (const FPropertyBagPropertyDesc& Desc : Bag->GetPropertyDescs())
+				{
+					PropertyOverrides.Add(Desc.ID);
+				}
 			}
 		}
+
+#if WITH_EDITORONLY_DATA
+		// This is not a perfect solution. In case there is a circular dependency between two assets,
+		// StateTree might not have been loaded or was only partially loaded, in which case the tree data
+		// won't be ready for parameters sync. We would need to sync parameters in the outer UObject's PostLoad.
+		// We also can't always preload StateTree here because if there is a circular dependency, that would deadlock.
+		if (StateTree && StateTree->HasAnyFlags(RF_LoadCompleted))
+		{
+			StateTree->ConditionalPostLoad();
+
+			ConditionallySyncParameters();
+		}
+#endif //WITH_EDITORONLY_DATA
 	}
+
 }
