[StateTree] Extend the state tree editor data extension allowing it to customize the state tree state details panel

diff --git a/Engine/Plugins/Runtime/StateTree/Source/StateTreeEditorModule/Private/Customizations/StateTreeStateDetails.cpp b/Engine/Plugins/Runtime/StateTree/Source/StateTreeEditorModule/Private/Customizations/StateTreeStateDetails.cpp
--- a/Engine/Plugins/Runtime/StateTree/Source/StateTreeEditorModule/Private/Customizations/StateTreeStateDetails.cpp
+++ b/Engine/Plugins/Runtime/StateTree/Source/StateTreeEditorModule/Private/Customizations/StateTreeStateDetails.cpp
@@ -12,6 +12,7 @@
 #include "StateTreeEditorData.h"
 #include "StateTreeEditorNodeUtils.h"
 #include "StateTreeEditorStyle.h"
+#include "StateTreeEditorDataExtension.h"
 #include "StateTreePropertyHelpers.h"
 #include "StateTreeSchema.h"
 #include "StateTreeStateParametersDetails.h"
@@ -427,6 +428,14 @@ void FStateTreeStateDetails::CustomizeDetails(IDetailLayoutBuilder& DetailBuilde
 				PropertyUtilities->ForceRefresh();
 			}
 		}));
+
+	for (UStateTreeEditorDataExtension* Extension : EditorData->Extensions)
+	{
+		if (Extension)
+		{
+			Extension->CustomizeDetails(WeakState.Get(), DetailBuilder);
+		}
+	}
 }
 
 void FStateTreeStateDetails::PostUndo(bool bSuccess)
diff --git a/Engine/Plugins/Runtime/StateTree/Source/StateTreeEditorModule/Public/StateTreeEditorDataExtension.h b/Engine/Plugins/Runtime/StateTree/Source/StateTreeEditorModule/Public/StateTreeEditorDataExtension.h
--- a/Engine/Plugins/Runtime/StateTree/Source/StateTreeEditorModule/Public/StateTreeEditorDataExtension.h
+++ b/Engine/Plugins/Runtime/StateTree/Source/StateTreeEditorModule/Public/StateTreeEditorDataExtension.h
@@ -10,7 +10,9 @@ namespace UE::StateTree::Compiler
 	struct FPostInternalContext;
 }
 
+class IDetailLayoutBuilder;
 class UStateTreeEditorData;
+class UStateTreeState;
 
 /**
  * Extension for the editor data of the state tree asset.
@@ -26,6 +28,10 @@ public:
 		return true;
 	}
 
+	virtual void CustomizeDetails(TNonNullPtr<UStateTreeState> State, IDetailLayoutBuilder& DetailBuilder)
+	{
+	}
+
 protected:
 	UStateTreeEditorData* GetStateTreeEditorData() const
 	{
