[StateTree] force activate level editor TabManager owner tab and draw attention to the RewindDebugger main tab and details tab when using the "Open RewindDebugger" button

diff --git a/Engine/Plugins/Runtime/StateTree/Source/StateTreeEditorModule/Private/Debugger/SStateTreeDebuggerView.cpp b/Engine/Plugins/Runtime/StateTree/Source/StateTreeEditorModule/Private/Debugger/SStateTreeDebuggerView.cpp
--- a/Engine/Plugins/Runtime/StateTree/Source/StateTreeEditorModule/Private/Debugger/SStateTreeDebuggerView.cpp
+++ b/Engine/Plugins/Runtime/StateTree/Source/StateTreeEditorModule/Private/Debugger/SStateTreeDebuggerView.cpp
@@ -23,6 +23,7 @@
 #include "StateTreeModule.h"
 #include "StateTreeViewModel.h"
 #include "SWarningOrErrorBox.h"
+#include "Widgets/Docking/SDockTab.h"
 #include "Widgets/Input/SComboButton.h"
 #include "Widgets/Layout/SScrollBox.h"
 
@@ -839,8 +840,20 @@ void SStateTreeDebuggerView::OpenRewindDebugger()
 		{
 			if (const IRewindDebuggerModule* RewindDebuggerModule = FModuleManager::GetModulePtr<IRewindDebuggerModule>("RewindDebugger"))
 			{
-				LevelEditorTabManager->TryInvokeTab(RewindDebuggerModule->GetMainTabName());
-				LevelEditorTabManager->TryInvokeTab(RewindDebuggerModule->GetDetailsTabName());
+				if (const TSharedPtr<SDockTab> Tab = LevelEditorTabManager->GetOwnerTab())
+				{
+					Tab->ActivateInParent(SetDirectly);
+				}
+
+				if (const TSharedPtr<SDockTab> Tab = LevelEditorTabManager->TryInvokeTab(RewindDebuggerModule->GetMainTabName()))
+				{
+					LevelEditorTabManager->DrawAttention(Tab.ToSharedRef());
+				}
+
+				if (const TSharedPtr<SDockTab> Tab = LevelEditorTabManager->TryInvokeTab(RewindDebuggerModule->GetDetailsTabName()))
+				{
+					LevelEditorTabManager->DrawAttention(Tab.ToSharedRef());
+				}
 			}
 		}
 	}
