[StateTree] allow FStateTreePropertyRef to handle property mismatches when replacing FStateTreeStructRef properties (will be deprecated soon). Those are used as inputs so we can replace the latter and the existing bindings will still be able to do the link.

diff --git a/Engine/Plugins/Runtime/StateTree/Source/StateTreeModule/Public/StateTreePropertyRef.h b/Engine/Plugins/Runtime/StateTree/Source/StateTreeModule/Public/StateTreePropertyRef.h
--- a/Engine/Plugins/Runtime/StateTree/Source/StateTreeModule/Public/StateTreePropertyRef.h
+++ b/Engine/Plugins/Runtime/StateTree/Source/StateTreeModule/Public/StateTreePropertyRef.h
@@ -193,6 +193,23 @@ struct FStateTreePropertyRef
 		return RefAccessIndex;
 	}
 
+	//~ Begin TStructOpsTypeTraits interface
+	bool SerializeFromMismatchedTag(const FPropertyTag& Tag, FStructuredArchive::FSlot Slot)
+	{
+		static const FName StateTreeStructRefName("StateTreeStructRef");
+		if (Tag.GetType().IsStruct(StateTreeStructRefName))
+		{
+			// Serialize the data, but we don't have anything to do with it.
+			// StructRef and PropertyRef are used for input and existing bindings will set them if needed.
+			FStateTreeStructRef TempStructRef;
+			FStateTreeStructRef::StaticStruct()->SerializeItem(Slot, &TempStructRef, nullptr);
+			return true;
+		}
+
+		return false;
+	}
+	//~ End TStructOpsTypeTraits interface
+
 private:
 	UPROPERTY()
 	FStateTreeIndex16 RefAccessIndex;
@@ -200,6 +217,15 @@ private:
 	friend FStateTreePropertyBindingCompiler;
 };
 
+template <>
+struct TStructOpsTypeTraits<FStateTreePropertyRef> : public TStructOpsTypeTraitsBase2<FStateTreePropertyRef>
+{
+	enum
+	{
+		WithStructuredSerializeFromMismatchedTag = true,
+	};
+};
+
 /**
  * TStateTreePropertyRef is a type-safe FStateTreePropertyRef wrapper against a single given type.
  * @note When used as a property, this automatically defines PropertyRef property meta-data.
