[State Tree] Binding on Node Struct. A node struct is considered for binding only
when it has DelegateDispatcher, DelegateListener or PropertyRef.

- Added GetNodeID function based on ID. StateTreeEditorData now iterates through NodeStruct also to produce descriptors for binding.
- UX Support and minor fixes for UX.
- some refactorings and minor fixes to StateTreeCompiler.
- Unit tests for binding on node struct.

diff --git a/Engine/Plugins/Runtime/PropertyBindingUtils/Source/PropertyBindingUtils/Public/PropertyBindingDataView.h b/Engine/Plugins/Runtime/PropertyBindingUtils/Source/PropertyBindingUtils/Public/PropertyBindingDataView.h
--- a/Engine/Plugins/Runtime/PropertyBindingUtils/Source/PropertyBindingUtils/Public/PropertyBindingDataView.h
+++ b/Engine/Plugins/Runtime/PropertyBindingUtils/Source/PropertyBindingUtils/Public/PropertyBindingDataView.h
@@ -43,6 +43,13 @@ struct FPropertyBindingDataView
 		check(!Memory || (Memory && Struct));
 	}
 
+	/** Struct from a InstancedStruct. */
+	FPropertyBindingDataView(FInstancedStruct& InstancedStruct) : Struct(InstancedStruct.GetScriptStruct()), Memory(InstancedStruct.GetMutableMemory())
+	{
+		// Must have type with valid pointer.
+		check(!Memory || (Memory && Struct));
+	}
+
 	/**
 	 * Check is the view is valid (both pointer and type are set). On valid views it is safe to call the Get<>() methods returning a reference.
 	 * @return True if the view is valid.
diff --git a/Engine/Plugins/Runtime/StateTree/Source/StateTreeEditorModule/Private/Customizations/StateTreeBindingExtension.cpp b/Engine/Plugins/Runtime/StateTree/Source/StateTreeEditorModule/Private/Customizations/StateTreeBindingExtension.cpp
--- a/Engine/Plugins/Runtime/StateTree/Source/StateTreeEditorModule/Private/Customizations/StateTreeBindingExtension.cpp
+++ b/Engine/Plugins/Runtime/StateTree/Source/StateTreeEditorModule/Private/Customizations/StateTreeBindingExtension.cpp
@@ -58,9 +58,9 @@ UObject* FindEditorBindingsOwner(UObject* InObject)
 }
 
 
-static bool IsDelegateDispatcherProperty(const FProperty* Property)
+static bool IsDelegateDispatcherProperty(const FProperty* InProperty)
 {
-	if (const FStructProperty* StructProperty = CastField<FStructProperty>(Property))
+	if (const FStructProperty* StructProperty = CastField<FStructProperty>(InProperty))
 	{
 		return StructProperty->Struct == FStateTreeDelegateDispatcher::StaticStruct();
 	}
@@ -68,9 +68,9 @@ static bool IsDelegateDispatcherProperty(const FProperty* Property)
 	return false;
 }
 
-static bool IsDelegateListenerProperty(const FProperty* Property)
+static bool IsDelegateListenerProperty(const FProperty* InProperty)
 {
-	if (const FStructProperty* StructProperty = CastField<FStructProperty>(Property))
+	if (const FStructProperty* StructProperty = CastField<FStructProperty>(InProperty))
 	{
 		return StructProperty->Struct == FStateTreeDelegateListener::StaticStruct() || StructProperty->Struct == FStateTreeTransitionDelegateListener::StaticStruct();
 	}
@@ -78,6 +78,39 @@ static bool IsDelegateListenerProperty(const FProperty* Property)
 	return false;
 }
 
+static bool IsNodeStructProperty(const TSharedPtr<const IPropertyHandle>& InPropertyHandle)
+{
+	void* Address = nullptr;
+	if (InPropertyHandle && InPropertyHandle->GetValueData(Address) == FPropertyAccess::Success)
+	{
+		const FStructProperty* StructProperty = CastField<FStructProperty>(InPropertyHandle->GetProperty());
+		if (Address && StructProperty)
+		{
+			const UScriptStruct* Struct = nullptr;
+
+			// @todo: refactor this into UE::Object::GetInnerStruct(const UScriptStruct*)
+			if (StructProperty->Struct == TBaseStructure<FInstancedStruct>::Get())
+			{
+				Struct = (static_cast<FInstancedStruct*>(Address))->GetScriptStruct();
+			}
+			else if (StructProperty->Struct == TBaseStructure<FSharedStruct>::Get())
+			{
+				Struct = (static_cast<FSharedStruct*>(Address))->GetScriptStruct();
+			}
+			else
+			{
+				Struct = StructProperty->Struct;
+			}
+
+			if (Struct)
+			{
+				return Struct->IsChildOf<FStateTreeNodeBase>();
+			}
+		}
+	}
+
+	return false;
+}
 
 FOnStateTreePropertyBindingChanged STATETREEEDITORMODULE_API OnStateTreePropertyBindingChanged;
 
@@ -290,12 +323,22 @@ struct FStateTreeCachedBindingData : UE::PropertyBinding::FCachedBindingData, TS
 		{
 			const IStateTreeEditorPropertyBindingsOwner* BindingOwner = Cast<IStateTreeEditorPropertyBindingsOwner>(GetOwner());
 			FStateTreeDataView DataView;
+
 			// If DataView exists, struct is an instance of already bound function.
 			if (BindingOwner == nullptr || BindingOwner->GetBindingDataViewByID(StructDesc.ID, DataView))
 			{
 				return false;
 			}
 
+			if (const FStructProperty* StructProperty = CastField<FStructProperty>(GetPropertyHandle()->GetProperty()))
+			{
+				// Disallow PropertyRef bound to property functions
+				if (StructProperty->Struct && StructProperty->Struct->IsChildOf<FStateTreePropertyRef>())
+				{
+					return false;
+				}
+			}
+
 			if (const FProperty* SingleOutputProperty = GetStructSingleOutputProperty(*StructDesc.Struct))
 			{
 				return CanBindToProperty(SingleOutputProperty, {FBindingChainElement(nullptr, InStructIndex), FBindingChainElement(const_cast<FProperty*>(SingleOutputProperty))});
@@ -338,16 +381,13 @@ struct FStateTreeCachedBindingData : UE::PropertyBinding::FCachedBindingData, TS
 
 		if (PropertyRefHelpers::IsPropertyRef(*GetPropertyHandle()->GetProperty()) && !PropertyRefHelpers::IsPropertyAccessibleForPropertyRef(SourceProperty, InBindingChain, StructDesc))
 		{
-			if (!PropertyRefHelpers::IsPropertyAccessibleForPropertyRef(SourceProperty, InBindingChain, StructDesc))
-			{
-				return false;
-			}
+			return false;
 		}
 
 		// Listener can only bind to dispatcher (prevents listener to listener)
-		if (IsDelegateListenerProperty(GetPropertyHandle()->GetProperty()))
+		if (IsDelegateListenerProperty(GetPropertyHandle()->GetProperty()) && !IsDelegateDispatcherProperty(&SourceProperty))
 		{
-			return IsDelegateDispatcherProperty(&SourceProperty);
+			return false;
 		}
 
 		return true;
@@ -567,6 +607,30 @@ TSharedPtr<UE::PropertyBinding::FCachedBindingData> FStateTreeBindingExtension::
 
 bool FStateTreeBindingExtension::CanBindToProperty(const FPropertyBindingPath& InTargetPath, const IPropertyHandle& InPropertyHandle) const
 {
+	const FProperty* PropertyToBind = InPropertyHandle.GetProperty();
+	if (!PropertyToBind)
+	{
+		return false;
+	}
+
+	const bool bIsPropertyRef = UE::StateTree::PropertyRefHelpers::IsPropertyRef(*PropertyToBind);
+	const bool bIsDelegateListener = UE::StateTree::PropertyBinding::IsDelegateListenerProperty(PropertyToBind);
+
+	// Node struct properties are considered as binding target only if their value stay constant at runtime
+	if (!bIsPropertyRef && !bIsDelegateListener)
+	{
+		TSharedPtr<const IPropertyHandle> CurrentPropertyHandle = InPropertyHandle.AsShared();
+		while (CurrentPropertyHandle)
+		{
+			if (UE::StateTree::PropertyBinding::IsNodeStructProperty(CurrentPropertyHandle))
+			{
+				return false;
+			}
+
+			CurrentPropertyHandle = CurrentPropertyHandle->GetParentHandle();
+		}
+	}
+
 	const EStateTreePropertyUsage Usage = UE::StateTree::GetUsageFromMetaData(InPropertyHandle.GetProperty());
 	if (Usage == EStateTreePropertyUsage::Input || Usage == EStateTreePropertyUsage::Context)
 	{
diff --git a/Engine/Plugins/Runtime/StateTree/Source/StateTreeEditorModule/Private/Customizations/StateTreeEditorNodeDetails.cpp b/Engine/Plugins/Runtime/StateTree/Source/StateTreeEditorModule/Private/Customizations/StateTreeEditorNodeDetails.cpp
--- a/Engine/Plugins/Runtime/StateTree/Source/StateTreeEditorModule/Private/Customizations/StateTreeEditorNodeDetails.cpp
+++ b/Engine/Plugins/Runtime/StateTree/Source/StateTreeEditorModule/Private/Customizations/StateTreeEditorNodeDetails.cpp
@@ -321,23 +321,20 @@ class FBindableNodeInstanceDetails : public FInstancedStructDataDetails
 {
 public:
 
-	FBindableNodeInstanceDetails(TSharedPtr<IPropertyHandle> InStructProperty, TSharedPtr<IPropertyHandle> InIDProperty, UStateTreeEditorData* InEditorData)
+	FBindableNodeInstanceDetails(TSharedPtr<IPropertyHandle> InStructProperty, const FGuid& InID, UStateTreeEditorData* InEditorData)
 		: FInstancedStructDataDetails(InStructProperty)
+		, ID(InID)
 		, EditorData(InEditorData)
-		, IDProperty(InIDProperty)
 	{
 	}
 
 	virtual void OnChildRowAdded(IDetailPropertyRow& ChildRow) override
 	{
-		FGuid ID;
-		UE::StateTree::PropertyHelpers::GetStructValue<FGuid>(IDProperty, ID);
-		
 		UE::StateTreeEditor::Internal::ModifyRow(ChildRow, ID, EditorData.Get());
 	}
 
+	FGuid ID;
 	TWeakObjectPtr<UStateTreeEditorData> EditorData;
-	TSharedPtr<IPropertyHandle> IDProperty;
 };
 
 ////////////////////////////////////
@@ -969,8 +966,12 @@ void FStateTreeEditorNodeDetails::ResetToDefault(TSharedPtr<IPropertyHandle> Pro
 
 void FStateTreeEditorNodeDetails::CustomizeChildren(TSharedRef<class IPropertyHandle> StructPropertyHandle, class IDetailChildrenBuilder& StructBuilder, IPropertyTypeCustomizationUtils& StructCustomizationUtils)
 {
-	FGuid ID;
-	UE::StateTree::PropertyHelpers::GetStructValue<FGuid>(IDProperty, ID);
+	const FStateTreeEditorNode* EditorNodePtr = UE::StateTree::PropertyHelpers::GetStructPtr<FStateTreeEditorNode>(StructProperty);
+	if (EditorNodePtr == nullptr)
+	{
+		// in case StructPropertyHandle is not valid
+		return;
+	}
 
 	// ID
 	if (UE::StateTree::Editor::GbDisplayItemIds)
@@ -982,11 +983,11 @@ void FStateTreeEditorNodeDetails::CustomizeChildren(TSharedRef<class IPropertyHa
 	UStateTreeEditorData* EditorDataPtr = EditorData.Get();
 
 	// Node
-	TSharedRef<FBindableNodeInstanceDetails> NodeDetails = MakeShareable(new FBindableNodeInstanceDetails(NodeProperty, {}, EditorDataPtr));
+	TSharedRef<FBindableNodeInstanceDetails> NodeDetails = MakeShareable(new FBindableNodeInstanceDetails(NodeProperty, EditorNodePtr->GetNodeID(), EditorDataPtr));
 	StructBuilder.AddCustomBuilder(NodeDetails);
 
 	// Instance
-	TSharedRef<FBindableNodeInstanceDetails> InstanceDetails = MakeShareable(new FBindableNodeInstanceDetails(InstanceProperty, IDProperty, EditorDataPtr));
+	TSharedRef<FBindableNodeInstanceDetails> InstanceDetails = MakeShareable(new FBindableNodeInstanceDetails(InstanceProperty, EditorNodePtr->ID, EditorDataPtr));
 	StructBuilder.AddCustomBuilder(InstanceDetails);
 
 	// InstanceObject
@@ -1030,7 +1031,7 @@ void FStateTreeEditorNodeDetails::CustomizeChildren(TSharedRef<class IPropertyHa
 		for (FSortedChild& Child : SortedChildren)
 		{
 			IDetailPropertyRow& ChildRow = StructBuilder.AddProperty(Child.PropertyHandle.ToSharedRef());
-			UE::StateTreeEditor::Internal::ModifyRow(ChildRow, ID, EditorDataPtr);
+			UE::StateTreeEditor::Internal::ModifyRow(ChildRow, EditorNodePtr->ID, EditorDataPtr);
 		}
 	}
 }
diff --git a/Engine/Plugins/Runtime/StateTree/Source/StateTreeEditorModule/Private/StateTreeCompiler.cpp b/Engine/Plugins/Runtime/StateTree/Source/StateTreeEditorModule/Private/StateTreeCompiler.cpp
--- a/Engine/Plugins/Runtime/StateTree/Source/StateTreeEditorModule/Private/StateTreeCompiler.cpp
+++ b/Engine/Plugins/Runtime/StateTree/Source/StateTreeEditorModule/Private/StateTreeCompiler.cpp
@@ -20,6 +20,7 @@
 #include "StateTreePropertyHelpers.h"
 #include "StateTreeDelegate.h"
 #include "Customizations/StateTreeEditorNodeUtils.h"
+#include "StateTreePropertyFunctionBase.h"
 
 namespace UE::StateTree
 {
@@ -135,67 +136,86 @@ namespace UE::StateTree::Compiler
 		}
 	};
 
-	/** Scans Data for actors that are tied to some level and returns them. */
-	void ScanLevelActorReferences(FStateTreeDataView Data, TSet<const UObject*>& Visited, TArray<const AActor*>& OutActors)
+	enum class EPropertyVisitorResult : uint8
 	{
-		if (!Data.IsValid())
-		{
-			return;
-		}
-		
-		for (TPropertyValueIterator<FProperty> It(Data.GetStruct(), Data.GetMemory()); It; ++It)
-		{
-			const FProperty* Property = It->Key;
-			const void* ValuePtr = It->Value;
+		Continue,
+		Break
+	};
 
-			if (!ValuePtr)
-			{
-				continue;
-			}
-			
-			if (const FStructProperty* StructProperty = CastField<FStructProperty>(Property))
+	void ScanProperties(FStateTreeDataView Data, TFunctionRef<EPropertyVisitorResult(const FProperty* InProperty, const void* InAddress)> InFunc)
+	{
+		TSet<const UObject*> Visited;
+		auto ScanPropertiesRecursive = [&Visited, InFunc](auto&& Self, FStateTreeDataView CurrentData)
 			{
-				if (StructProperty->Struct == TBaseStructure<FInstancedStruct>::Get())
+				if (!CurrentData.IsValid())
 				{
-					const FInstancedStruct& InstancedStruct = *static_cast<const FInstancedStruct*>(ValuePtr);
-					if (InstancedStruct.IsValid())
+					return;
+				}
+
+				for (TPropertyValueIterator<FProperty> It(CurrentData.GetStruct(), CurrentData.GetMemory()); It; ++It)
+				{
+					const FProperty* Property = It->Key;
+					const void* Address = It->Value;
+
+					if (!Address)
 					{
-						ScanLevelActorReferences(FStateTreeDataView(const_cast<FInstancedStruct&>(InstancedStruct)), Visited, OutActors);
+						continue;
+					}
+
+					if (InFunc(Property, Address) == EPropertyVisitorResult::Break)
+					{
+						break;
+					}
+
+					if (const FObjectPropertyBase* ObjectProperty = CastField<FObjectPropertyBase>(Property))
+					{
+						if (const UObject* Object = ObjectProperty->GetObjectPropertyValue(Address))
+						{
+							// Recurse into instanced object
+							if (Property->HasAnyPropertyFlags(CPF_InstancedReference))
+							{
+								if (!Visited.Contains(Object))
+								{
+									Visited.Add(Object);
+									Self(Self, FStateTreeDataView(const_cast<UObject*>(Object)));
+								}
+							}
+						}
 					}
 				}
-			}
-			else if (const FObjectPropertyBase* ObjectProperty = CastField<FObjectPropertyBase>(Property))
+			};
+
+		ScanPropertiesRecursive(ScanPropertiesRecursive, Data);
+	}
+
+	/** Scans Data for actors that are tied to some level and returns them. */
+	void ScanLevelActorReferences(FStateTreeDataView Data, TArray<const AActor*>& OutActors)
+	{
+		ScanProperties(Data, [&OutActors](const FProperty* InProperty, const void* InAddress)
+		{
+			if (const FObjectPropertyBase* ObjectProperty = CastField<FObjectPropertyBase>(InProperty))
 			{
-				if (const UObject* Object = ObjectProperty->GetObjectPropertyValue(ValuePtr))
+				if (const UObject* Object = ObjectProperty->GetObjectPropertyValue(InAddress))
 				{
 					if (const AActor* Actor = Cast<AActor>(Object))
 					{
-						const ULevel* Level = Actor->GetLevel();
-						if (Level != nullptr)
+						if (const ULevel* Level = Actor->GetLevel())
 						{
 							OutActors.Add(Actor);
 						}
 					}
-					// Recurse into instanced object
-					if (Property->HasAnyPropertyFlags(CPF_InstancedReference))
-					{
-						if (!Visited.Contains(Object))
-						{
-							Visited.Add(Object);
-							ScanLevelActorReferences(FStateTreeDataView(const_cast<UObject*>(Object)), Visited, OutActors);
-						}
-					}
 				}
 			}
-		}
+
+			return EPropertyVisitorResult::Continue;
+			});
 	}
 
 	bool ValidateNoLevelActorReferences(FStateTreeCompilerLog& Log, const FStateTreeBindableStructDesc& NodeDesc, const FStateTreeDataView NodeView, const FStateTreeDataView InstanceView)
 	{
-		TSet<const UObject*> Visited;
 		TArray<const AActor*> LevelActors;
-		UE::StateTree::Compiler::ScanLevelActorReferences(NodeView, Visited, LevelActors);
-		UE::StateTree::Compiler::ScanLevelActorReferences(InstanceView, Visited, LevelActors);
+		ScanLevelActorReferences(NodeView, LevelActors);
+		ScanLevelActorReferences(InstanceView, LevelActors);
 		if (!LevelActors.IsEmpty())
 		{
 			FStringBuilderBase AllActorsString;
@@ -216,6 +236,54 @@ namespace UE::StateTree::Compiler
 		return true;
 	}
 
+	template<typename... T>
+	bool StructHasAnyStructProperties(FConstStructView StructView)
+	{
+		bool bResult = false;
+		ScanProperties(FStateTreeDataView(StructView.GetScriptStruct(), const_cast<uint8*>(StructView.GetMemory())), 
+			[&bResult](const FProperty* InProperty, const void* InAddress)
+			{
+				if (const FStructProperty* StructProperty = CastField<FStructProperty>(InProperty))
+				{
+					if ((... || StructProperty->Struct->IsChildOf<T>()))
+					{
+						bResult = true;
+						return EPropertyVisitorResult::Break;
+					}
+				}
+
+				return EPropertyVisitorResult::Continue;
+			});
+
+		return bResult;
+	}
+
+	template<typename... T>
+	bool IsPropertyChildOfAnyStruct(const FStateTreeBindableStructDesc& Struct, FPropertyBindingPath Path)
+	{
+		TArray<FPropertyBindingPathIndirection> Indirection;
+		const bool bResolved = Path.ResolveIndirections(Struct.Struct, Indirection);
+
+		if (bResolved && Indirection.Num() > 0)
+		{
+			const FProperty* Property = Indirection.Last().GetProperty();
+			check(Property);
+
+			if (const FStructProperty* StructProperty = CastField<FStructProperty>(Property))
+			{
+				return StructProperty->Struct && (... || StructProperty->Struct->IsChildOf<T>());
+			}
+		}
+
+		return false;
+	}
+
+	bool IsNodeStructEligibleForBinding(FConstStructView NodeView)
+	{
+		check(NodeView.GetScriptStruct() && NodeView.GetScriptStruct()->IsChildOf<FStateTreeNodeBase>());
+
+		return StructHasAnyStructProperties<FStateTreePropertyRef, FStateTreeDelegateDispatcher, FStateTreeDelegateListener>(NodeView);
+	}
 
 	void FValidationResult::Log(FStateTreeCompilerLog& Log, const TCHAR* ContextText, const FStateTreeBindableStructDesc& ContextStruct) const
 	{
@@ -1580,7 +1648,8 @@ bool FStateTreeCompiler::CreateBindingsForNodes(TConstArrayView<FStateTreeEditor
 			continue;
 		}
 
-		FStateTreeNodeBase& Node = Nodes[NodeIndex++].GetMutable<FStateTreeNodeBase>();
+		FInstancedStruct& NodeInstancedStruct = Nodes[NodeIndex++];
+		FStateTreeNodeBase& Node = NodeInstancedStruct.GetMutable<FStateTreeNodeBase>();
 
 		FStateTreeIndex16 PropertyFunctionsBegin(Nodes.Num());
 		if (!CreatePropertyFunctionsForStruct(EditorNode.ID))
@@ -1612,13 +1681,25 @@ bool FStateTreeCompiler::CreateBindingsForNodes(TConstArrayView<FStateTreeEditor
 		}
 
 		{
-			const FStateTreeBindableStructDesc* BindableStruct = BindingsCompiler.GetSourceStructDescByID(EditorNode.ID);
-			check(BindableStruct);
-			if (!CreateBindingsForStruct(*BindableStruct, InstanceView, PropertyFunctionsBegin, PropertyFunctionsEnd, Node.BindingsBatch))
+			const FStateTreeBindableStructDesc* BindableInstanceStruct = BindingsCompiler.GetSourceStructDescByID(EditorNode.ID);
+			check(BindableInstanceStruct);
+			if (!CreateBindingsForStruct(*BindableInstanceStruct, InstanceView, PropertyFunctionsBegin, PropertyFunctionsEnd, Node.BindingsBatch))
 			{
 				bSucceeded = false;
 				continue;
 			}
+
+			if (const FStateTreeBindableStructDesc* BindableNodeStruct = BindingsCompiler.GetSourceStructDescByID(EditorNode.GetNodeID()))
+			{
+				if (UE::StateTree::Compiler::IsNodeStructEligibleForBinding(NodeInstancedStruct))
+				{
+					if (!CreateBindingsForStruct(*BindableNodeStruct, static_cast<FStructView>(NodeInstancedStruct), PropertyFunctionsBegin, PropertyFunctionsEnd, Node.BindingsBatch))
+					{
+						bSucceeded = false;
+						continue;
+					}
+				}
+			}
 		}
 	}
 
@@ -1751,7 +1832,7 @@ bool FStateTreeCompiler::CreatePropertyFunction(const FStateTreeEditorNode& Func
 	return CreateBindingsForStruct(*BindableStruct, InstanceView, FStateTreeIndex16::Invalid, FStateTreeIndex16::Invalid, Function->BindingsBatch);
 }
 
-FInstancedStruct* FStateTreeCompiler::CreateNode(UStateTreeState* State, const FStateTreeEditorNode& EditorNode, FStateTreeBindableStructDesc& StructDesc, const FStateTreeDataHandle DataHandle, TArray<FInstancedStruct>& InstancedStructContainer)
+FInstancedStruct* FStateTreeCompiler::CreateNode(UStateTreeState* State, const FStateTreeEditorNode& EditorNode, FStateTreeBindableStructDesc& InstanceDesc, const FStateTreeDataHandle DataHandle, TArray<FInstancedStruct>& InstancedStructContainer)
 {
 	if (!EditorNode.Node.IsValid())
 	{
@@ -1761,7 +1842,7 @@ FInstancedStruct* FStateTreeCompiler::CreateNode(UStateTreeState* State, const F
 	// Check that item has valid instance initialized.
 	if (!EditorNode.Instance.IsValid() && EditorNode.InstanceObject == nullptr)
 	{
-		Log.Reportf(EMessageSeverity::Error, StructDesc,
+		Log.Reportf(EMessageSeverity::Error, InstanceDesc,
 			TEXT("Malformed node, missing instance value."));
 		return nullptr;
 	}
@@ -1773,30 +1854,29 @@ FInstancedStruct* FStateTreeCompiler::CreateNode(UStateTreeState* State, const F
 
 	FStateTreeNodeBase& Node = RawNode.GetMutable<FStateTreeNodeBase>();
 
-	TOptional<FStateTreeDataView> InstanceDataView = CreateNodeInstanceData(EditorNode, Node, StructDesc, DataHandle, InstancedStructContainer);
+	TOptional<FStateTreeDataView> InstanceDataView = CreateNodeInstanceData(EditorNode, Node, InstanceDesc, DataHandle, InstancedStructContainer);
 	if (!InstanceDataView.IsSet())
 	{
 		return nullptr;
 	}
 
-	if (!CompileAndValidateNode(State, StructDesc, RawNode, InstanceDataView.GetValue()))
+	if (!CompileAndValidateNode(State, InstanceDesc, RawNode, InstanceDataView.GetValue()))
 	{
 		return nullptr;
 	}
 
-	// Mark the struct as binding source.
-	BindingsCompiler.AddSourceStruct(StructDesc);
+	CreateBindingSourceStructsForNode(EditorNode, InstanceDesc);
 
 	return &RawNode;
 }
 
-FInstancedStruct* FStateTreeCompiler::CreateNodeWithSharedInstanceData(UStateTreeState* State, const FStateTreeEditorNode& EditorNode, FStateTreeBindableStructDesc& StructDesc)
+FInstancedStruct* FStateTreeCompiler::CreateNodeWithSharedInstanceData(UStateTreeState* State, const FStateTreeEditorNode& EditorNode, FStateTreeBindableStructDesc& InstanceDesc)
 {
 	const FStateTreeDataHandle DataHandle = FStateTreeDataHandle(EStateTreeDataSourceType::SharedInstanceData, SharedInstanceStructs.Num());
-	return CreateNode(State, EditorNode, StructDesc, DataHandle, SharedInstanceStructs);
+	return CreateNode(State, EditorNode, InstanceDesc, DataHandle, SharedInstanceStructs);
 }
 
-TOptional<FStateTreeDataView> FStateTreeCompiler::CreateNodeInstanceData(const FStateTreeEditorNode& EditorNode, FStateTreeNodeBase& Node, FStateTreeBindableStructDesc& StructDesc, const FStateTreeDataHandle DataHandle, TArray<FInstancedStruct>& InstancedStructContainer)
+TOptional<FStateTreeDataView> FStateTreeCompiler::CreateNodeInstanceData(const FStateTreeEditorNode& EditorNode, FStateTreeNodeBase& Node, FStateTreeBindableStructDesc& InstanceDesc, const FStateTreeDataHandle DataHandle, TArray<FInstancedStruct>& InstancedStructContainer)
 {
 	FStateTreeDataView InstanceDataView;
 
@@ -1810,11 +1890,11 @@ TOptional<FStateTreeDataView> FStateTreeCompiler::CreateNodeInstanceData(const F
 		InstantiateStructSubobjects(InstancedStructContainer[InstanceIndex]);
 
 		// Create binding source struct descriptor.
-		StructDesc.Struct = EditorNode.Instance.GetScriptStruct();
+		InstanceDesc.Struct = EditorNode.Instance.GetScriptStruct();
 
 		if (const auto Validation = UE::StateTree::Compiler::IsValidIndex16(InstanceIndex); Validation.DidFail())
 		{
-			Validation.Log(Log, TEXT("InstanceIndex"), StructDesc);
+			Validation.Log(Log, TEXT("InstanceIndex"), InstanceDesc);
 			return {};
 		}
 		Node.InstanceTemplateIndex = FStateTreeIndex16(InstanceIndex);
@@ -1825,18 +1905,18 @@ TOptional<FStateTreeDataView> FStateTreeCompiler::CreateNodeInstanceData(const F
 	{
 		// Object Instance
 		check(EditorNode.InstanceObject != nullptr);
-		UObject* Instance = UE::StateTree::Compiler::DuplicateInstanceObject(Log, StructDesc, EditorNode.ID, EditorNode.InstanceObject, StateTree);
+		UObject* Instance = UE::StateTree::Compiler::DuplicateInstanceObject(Log, InstanceDesc, EditorNode.ID, EditorNode.InstanceObject, StateTree);
 
 		FInstancedStruct Wrapper;
 		Wrapper.InitializeAs<FStateTreeInstanceObjectWrapper>(Instance);
 		const int32 InstanceIndex = InstancedStructContainer.Add(MoveTemp(Wrapper));
 
 		// Create binding source struct descriptor.
-		StructDesc.Struct = Instance->GetClass();
+		InstanceDesc.Struct = Instance->GetClass();
 
 		if (const auto Validation = UE::StateTree::Compiler::IsValidIndex16(InstanceIndex); Validation.DidFail())
 		{
-			Validation.Log(Log, TEXT("InstanceIndex"), StructDesc);
+			Validation.Log(Log, TEXT("InstanceIndex"), InstanceDesc);
 			return {};
 		}
 		Node.InstanceTemplateIndex = FStateTreeIndex16(InstanceIndex);
@@ -1844,8 +1924,8 @@ TOptional<FStateTreeDataView> FStateTreeCompiler::CreateNodeInstanceData(const F
 		InstanceDataView = FStateTreeDataView(Instance);
 	}
 
-	StructDesc.DataHandle = Node.InstanceDataHandle;
-
+	InstanceDesc.DataHandle = Node.InstanceDataHandle;
+	
 	return TOptional<FStateTreeDataView>(InstanceDataView);
 }
 
@@ -1933,13 +2013,13 @@ PRAGMA_ENABLE_DEPRECATION_WARNINGS
 
 bool FStateTreeCompiler::CreateCondition(UStateTreeState& State, const FString& StatePath, const FStateTreeEditorNode& CondNode, const EStateTreeExpressionOperand Operand, const int8 DeltaIndent)
 {
-	FStateTreeBindableStructDesc StructDesc;
-	StructDesc.StatePath = StatePath;
-	StructDesc.ID = CondNode.ID;
-	StructDesc.Name = CondNode.GetName();
-	StructDesc.DataSource = EStateTreeBindableStructSource::Condition;
+	FStateTreeBindableStructDesc InstanceDesc;
+	InstanceDesc.StatePath = StatePath;
+	InstanceDesc.ID = CondNode.ID;
+	InstanceDesc.Name = CondNode.GetName();
+	InstanceDesc.DataSource = EStateTreeBindableStructSource::Condition;
 
-	FInstancedStruct* CreatedNode = CreateNodeWithSharedInstanceData(&State, CondNode, StructDesc);
+	FInstancedStruct* CreatedNode = CreateNodeWithSharedInstanceData(&State, CondNode, InstanceDesc);
 	if (CreatedNode)
 	{
 		FStateTreeConditionBase* Cond = CreatedNode->GetMutablePtr<FStateTreeConditionBase>();
@@ -1948,9 +2028,9 @@ bool FStateTreeCompiler::CreateCondition(UStateTreeState& State, const FString&
 			if (Cond->EvaluationMode == EStateTreeConditionEvaluationMode::ForcedFalse
 				|| Cond->EvaluationMode == EStateTreeConditionEvaluationMode::ForcedTrue)
 			{
-				Log.Reportf(EMessageSeverity::Info, StructDesc,
-					TEXT("The condition result will always be %s."),
-					Cond->EvaluationMode == EStateTreeConditionEvaluationMode::ForcedTrue ? TEXT("True") : TEXT("False"));
+				Log.Reportf(EMessageSeverity::Info, InstanceDesc,
+						TEXT("The condition result will always be %s."),
+						Cond->EvaluationMode == EStateTreeConditionEvaluationMode::ForcedTrue ? TEXT("True") : TEXT("False"));
 			}
 
 			Cond->Operand = Operand;
@@ -2000,16 +2080,16 @@ bool FStateTreeCompiler::CreateConsiderations(UStateTreeState& State, const FStr
 
 bool FStateTreeCompiler::CreateConsideration(UStateTreeState& State, const FString& StatePath, const FStateTreeEditorNode& ConsiderationNode, const EStateTreeExpressionOperand Operand, const int8 DeltaIndent)
 {
-	FStateTreeBindableStructDesc StructDesc;
-	StructDesc.StatePath = StatePath;
-	StructDesc.ID = ConsiderationNode.ID;
-	StructDesc.Name = ConsiderationNode.GetName();
-	StructDesc.DataSource = EStateTreeBindableStructSource::Consideration;
+	FStateTreeBindableStructDesc InstanceDesc;
+	InstanceDesc.StatePath = StatePath;
+	InstanceDesc.ID = ConsiderationNode.ID;
+	InstanceDesc.Name = ConsiderationNode.GetName();
+	InstanceDesc.DataSource = EStateTreeBindableStructSource::Consideration;
 
-	FInstancedStruct* CreatedNode = CreateNodeWithSharedInstanceData(&State, ConsiderationNode, StructDesc);
+	FInstancedStruct* CreatedNode = CreateNodeWithSharedInstanceData(&State, ConsiderationNode, InstanceDesc);
 	if (CreatedNode)
 	{
-		FStateTreeConsiderationBase * Consideration = CreatedNode->GetMutablePtr<FStateTreeConsiderationBase>();
+		FStateTreeConsiderationBase* Consideration = CreatedNode->GetMutablePtr<FStateTreeConsiderationBase>();
 		if (ensure(Consideration))
 		{
 			Consideration->Operand = Operand;
@@ -2021,7 +2101,7 @@ bool FStateTreeCompiler::CreateConsideration(UStateTreeState& State, const FStri
 	return false;
 }
 
-bool FStateTreeCompiler::CompileAndValidateNode(const UStateTreeState* SourceState, const FStateTreeBindableStructDesc& NodeDesc, FStructView NodeView, const FStateTreeDataView InstanceData)
+bool FStateTreeCompiler::CompileAndValidateNode(const UStateTreeState* SourceState, const FStateTreeBindableStructDesc& InstanceDesc, FStructView NodeView, const FStateTreeDataView InstanceData)
 {
 	if (!NodeView.IsValid())
 	{
@@ -2060,31 +2140,31 @@ bool FStateTreeCompiler::CompileAndValidateNode(const UStateTreeState* SourceSta
 	}
 
 	const FStateTreeBindingLookup& BindingLookup = FStateTreeBindingLookup(EditorData);
-	UE::StateTree::FCompileNodeContext CompileContext(InstanceData, NodeDesc, BindingLookup);
+	UE::StateTree::FCompileNodeContext CompileContext(InstanceData, InstanceDesc, BindingLookup);
 	const EDataValidationResult Result = Node.Compile(CompileContext);
 
 	if (Result == EDataValidationResult::Invalid && CompileContext.ValidationErrors.IsEmpty())
 	{
-		Log.Report(EMessageSeverity::Error, NodeDesc, TEXT("Node validation failed."));
+		Log.Report(EMessageSeverity::Error, InstanceDesc, TEXT("Node validation failed."));
 	}
 	else
 	{
 		const EMessageSeverity::Type Severity = Result == EDataValidationResult::Invalid ? EMessageSeverity::Error : EMessageSeverity::Warning;
 		for (const FText& Error : CompileContext.ValidationErrors)
 		{
-			Log.Report(Severity, NodeDesc, Error.ToString());
+			Log.Report(Severity, InstanceDesc, Error.ToString());
 		}
 	}
 
 	// Make sure there's no level actor references in the data.
-	if (!UE::StateTree::Compiler::ValidateNoLevelActorReferences(Log, NodeDesc, NodeView, InstanceData))
+	if (!UE::StateTree::Compiler::ValidateNoLevelActorReferences(Log, InstanceDesc, NodeView, InstanceData))
 	{
 		return false;
 	}
 	
 	if (const UClass* InstanceDataClass = Cast<UClass>(InstanceData.GetStruct()); InstanceDataClass && InstanceDataClass->HasAllClassFlags(CLASS_Deprecated))
 	{
-		Log.Report(EMessageSeverity::Warning, NodeDesc, TEXT("The instance data class is deprecated. It won't work in a cooked build."));
+		Log.Report(EMessageSeverity::Warning, InstanceDesc, TEXT("The instance data class is deprecated. It won't work in a cooked build."));
 	}
 
 	return Result != EDataValidationResult::Invalid;
@@ -2098,16 +2178,16 @@ bool FStateTreeCompiler::CreateTask(UStateTreeState* State, const FStateTreeEdit
 	}
 	
 	// Create binding source struct descriptor.
-	FStateTreeBindableStructDesc StructDesc;
-	StructDesc.StatePath = State ? State->GetPath() : UE::StateTree::Editor::GlobalStateName;
-	StructDesc.ID = TaskNode.ID;
-	StructDesc.Name = TaskNode.GetName();
-	StructDesc.DataSource = EStateTreeBindableStructSource::Task;
+	FStateTreeBindableStructDesc InstanceDesc;
+	InstanceDesc.StatePath = State ? State->GetPath() : UE::StateTree::Editor::GlobalStateName;
+	InstanceDesc.ID = TaskNode.ID;
+	InstanceDesc.Name = TaskNode.GetName();
+	InstanceDesc.DataSource = EStateTreeBindableStructSource::Task;
 
 	// Check that node has valid instance initialized.
 	if (!TaskNode.Instance.IsValid() && TaskNode.InstanceObject == nullptr)
 	{
-		Log.Reportf(EMessageSeverity::Error, StructDesc,
+		Log.Reportf(EMessageSeverity::Error, InstanceDesc,
 			TEXT("Malformed task, missing instance value."));
 		return false;
 	}
@@ -2118,7 +2198,7 @@ bool FStateTreeCompiler::CreateTask(UStateTreeState* State, const FStateTreeEdit
 	InstantiateStructSubobjects(Node);
 	
 	FStateTreeTaskBase& Task = Node.GetMutable<FStateTreeTaskBase>();
-	TOptional<FStateTreeDataView> InstanceDataView = CreateNodeInstanceData(TaskNode, Task, StructDesc, TaskDataHandle, InstanceStructs);
+	TOptional<FStateTreeDataView> InstanceDataView = CreateNodeInstanceData(TaskNode, Task, InstanceDesc, TaskDataHandle, InstanceStructs);
 	if (!InstanceDataView.IsSet())
 	{
 		return false;
@@ -2126,17 +2206,16 @@ bool FStateTreeCompiler::CreateTask(UStateTreeState* State, const FStateTreeEdit
 
 	if (!Task.bTaskEnabled)
 	{
-		Log.Reportf(EMessageSeverity::Info, StructDesc, TEXT("Task is disabled and will have no effect."));
+		Log.Reportf(EMessageSeverity::Info, InstanceDesc, TEXT("Task is disabled and will have no effect."));
 	}
 
-	if (!CompileAndValidateNode(State, StructDesc, Node,  InstanceDataView.GetValue()))
+	if (!CompileAndValidateNode(State, InstanceDesc, Node,  InstanceDataView.GetValue()))
 	{
 		return false;
 	}
 
-	// Mark the instance as binding source.
-	BindingsCompiler.AddSourceStruct(StructDesc);
-	
+	CreateBindingSourceStructsForNode(TaskNode, InstanceDesc);
+
 	return true;
 }
 
@@ -2148,17 +2227,17 @@ bool FStateTreeCompiler::CreateEvaluator(const FStateTreeEditorNode& EvalNode, c
 		return true;
 	}
 
-	// Create binding source struct descriptor.
-	FStateTreeBindableStructDesc StructDesc;
-	StructDesc.StatePath = UE::StateTree::Editor::GlobalStateName;
-    StructDesc.ID = EvalNode.ID;
-	StructDesc.Name = EvalNode.GetName();
-	StructDesc.DataSource = EStateTreeBindableStructSource::Evaluator;
+	// Create binding source descriptor for instance.
+	FStateTreeBindableStructDesc InstanceDesc;
+	InstanceDesc.StatePath = UE::StateTree::Editor::GlobalStateName;
+    InstanceDesc.ID = EvalNode.ID;
+	InstanceDesc.Name = EvalNode.GetName();
+	InstanceDesc.DataSource = EStateTreeBindableStructSource::Evaluator;
 
     // Check that node has valid instance initialized.
     if (!EvalNode.Instance.IsValid() && EvalNode.InstanceObject == nullptr)
     {
-        Log.Reportf(EMessageSeverity::Error, StructDesc,
+        Log.Reportf(EMessageSeverity::Error, InstanceDesc,
         	TEXT("Malformed evaluator, missing instance value."));
         return false;
     }
@@ -2169,45 +2248,24 @@ bool FStateTreeCompiler::CreateEvaluator(const FStateTreeEditorNode& EvalNode, c
 	InstantiateStructSubobjects(Node);
 	
 	FStateTreeEvaluatorBase& Eval = Node.GetMutable<FStateTreeEvaluatorBase>();
-	TOptional<FStateTreeDataView> InstanceDataView = CreateNodeInstanceData(EvalNode, Eval, StructDesc, EvalDataHandle, InstanceStructs);
+	TOptional<FStateTreeDataView> InstanceDataView = CreateNodeInstanceData(EvalNode, Eval, InstanceDesc, EvalDataHandle, InstanceStructs);
 	if (!InstanceDataView.IsSet())
 	{
 		return false;
 	}
 
-
-	if (!CompileAndValidateNode(nullptr, StructDesc, Node,  InstanceDataView.GetValue()))
+	if (!CompileAndValidateNode(nullptr, InstanceDesc, Node,  InstanceDataView.GetValue()))
 	{
 		return false;
 	}
 
-	// Mark the instance as binding source.
-	BindingsCompiler.AddSourceStruct(StructDesc);
+	CreateBindingSourceStructsForNode(EvalNode, InstanceDesc);
 
 	return true;
 }
 
-bool FStateTreeCompiler::IsPropertyOfTypeOrChild(UScriptStruct& Type, const FStateTreeBindableStructDesc& Struct, FPropertyBindingPath Path) const
-{
-	TArray<FPropertyBindingPathIndirection> Indirection;
-	const bool bResolved = Path.ResolveIndirections(Struct.Struct, Indirection);
-	
-	if (bResolved && Indirection.Num() > 0)
-	{
-		check(Indirection.Last().GetProperty());
-		if (const FProperty* OwnerProperty = Indirection.Last().GetProperty()->GetOwnerProperty())
-		{
-			if (const FStructProperty* OwnerStructProperty = CastField<FStructProperty>(OwnerProperty))
-			{
-				return OwnerStructProperty->Struct->IsChildOf(&Type);
-			}
-		}
-	}
-	return false;
-}
-
-bool FStateTreeCompiler::ValidateStructRef(const FStateTreeBindableStructDesc& SourceStruct, FPropertyBindingPath SourcePath,
-											const FStateTreeBindableStructDesc& TargetStruct, FPropertyBindingPath TargetPath) const
+bool FStateTreeCompiler::ValidateStructRef(const FStateTreeBindableStructDesc& SourceStruct, const FPropertyBindingPath& SourcePath,
+											const FStateTreeBindableStructDesc& TargetStruct, const FPropertyBindingPath& TargetPath) const
 {
 	FString ResolveError;
 	TArray<FPropertyBindingPathIndirection> TargetIndirection;
@@ -2291,9 +2349,37 @@ bool FStateTreeCompiler::ValidateStructRef(const FStateTreeBindableStructDesc& S
 	return true;
 }
 
+bool FStateTreeCompiler::ValidateBindingOnNode(const FStateTreeBindableStructDesc& TargetStruct, const FPropertyBindingPath& TargetPath) const
+{
+	FString ResolveError;
+	TArray<FPropertyBindingPathIndirection> TargetIndirection;
+	if (!TargetPath.ResolveIndirections(TargetStruct.Struct, TargetIndirection, &ResolveError))
+	{
+		// This will later be reported by the bindings compiler.
+		Log.Reportf(EMessageSeverity::Error, TargetStruct, TEXT("Failed to resolve binding path in %s: %s"), *TargetStruct.ToString(), *ResolveError);
+		return false;
+	}
+
+	check(TargetStruct.Struct);
+	if (TargetStruct.Struct->IsChildOf<FStateTreeNodeBase>())
+	{
+		if (!UE::StateTree::Compiler::IsPropertyChildOfAnyStruct<FStateTreePropertyRef, FStateTreeDelegateListener>(TargetStruct, TargetPath))
+		{
+			Log.Reportf(EMessageSeverity::Error, TargetStruct,
+				TEXT("Node struct(non-instance) only supports DelegateListener and PropertyReference as binding targets. %s is not supported."),
+				*UE::StateTree::GetDescAndPathAsString(TargetStruct, TargetPath));
+
+			return false;
+		}
+	}
+
+	return true;
+}
 
 bool FStateTreeCompiler::GetAndValidateBindings(const FStateTreeBindableStructDesc& TargetStruct, FStateTreeDataView TargetValue, FValidatedPathBindings& OutValidatedBindings) const
 {
+	using namespace UE::StateTree::Compiler;
+
 	check(EditorData);
 	
 	OutValidatedBindings = FValidatedPathBindings();
@@ -2311,7 +2397,7 @@ bool FStateTreeCompiler::GetAndValidateBindings(const FStateTreeBindableStructDe
 	{
 		if (Binding.GetSourcePath().GetStructID() == TargetStruct.ID)
 		{
-			if (IsPropertyOfTypeOrChild(*FStateTreeDelegateDispatcher::StaticStruct(), TargetStruct, Binding.GetSourcePath()))
+			if (IsPropertyChildOfAnyStruct<FStateTreeDelegateDispatcher>(TargetStruct, Binding.GetSourcePath()))
 			{
 				FStateTreePropertyPathBinding& BindingCopy = OutValidatedBindings.DelegateDispatchers.Add_GetRef(Binding);
 				BindingCopy.SetSourceDataHandle(TargetStruct.DataHandle);
@@ -2328,11 +2414,11 @@ bool FStateTreeCompiler::GetAndValidateBindings(const FStateTreeBindableStructDe
 		}
 
 		TArray<FStateTreePropertyPathBinding>* OutputBindings = nullptr;
-		if (IsPropertyOfTypeOrChild(*FStateTreeDelegateListener::StaticStruct(), TargetStruct, Binding.GetTargetPath()))
+		if (IsPropertyChildOfAnyStruct<FStateTreeDelegateListener>(TargetStruct, Binding.GetTargetPath()))
 		{
 			OutputBindings = &OutValidatedBindings.DelegateListeners;
 		}
-		else if (IsPropertyOfTypeOrChild(*FStateTreePropertyRef::StaticStruct(), TargetStruct, Binding.GetTargetPath()))
+		else if (IsPropertyChildOfAnyStruct<FStateTreePropertyRef>(TargetStruct, Binding.GetTargetPath()))
 		{
 			OutputBindings = &OutValidatedBindings.ReferenceBindings;
 		}
@@ -2425,13 +2511,16 @@ bool FStateTreeCompiler::GetAndValidateBindings(const FStateTreeBindableStructDe
 			}
 		}
 
-		if (!SourceStruct->DataHandle.IsValid())
+		if (SourceStruct->Struct && !SourceStruct->Struct->IsChildOf<FStateTreeNodeBase>())
 		{
-			Log.Reportf(EMessageSeverity::Error, TargetStruct,
-				TEXT("Malformed source'%s for property binding property '%s'."),
-				*UE::StateTree::GetDescAndPathAsString(*SourceStruct, Binding.GetSourcePath()), *Binding.GetSourcePath().ToString());
-			bSucceeded = false;
-			continue;
+			if (!SourceStruct->DataHandle.IsValid())
+			{
+				Log.Reportf(EMessageSeverity::Error, TargetStruct,
+					TEXT("Malformed source'%s for property binding property '%s'."),
+					*UE::StateTree::GetDescAndPathAsString(*SourceStruct, Binding.GetSourcePath()), *Binding.GetSourcePath().ToString());
+				bSucceeded = false;
+				continue;
+			}
 		}
 		
 		FStateTreePropertyPathBinding& BindingCopy = OutputBindings->Add_GetRef(Binding);
@@ -2440,8 +2529,8 @@ bool FStateTreeCompiler::GetAndValidateBindings(const FStateTreeBindableStructDe
 		// Special case for AnyEnum. StateTreeBindingExtension allows AnyEnums to bind to other enum types.
 		// The actual copy will be done via potential type promotion copy, into the value property inside the AnyEnum.
 		// We amend the paths here to point to the 'Value' property.
-		const bool bSourceIsAnyEnum = IsPropertyOfTypeOrChild(*TBaseStructure<FStateTreeAnyEnum>::Get(), *SourceStruct, Binding.GetSourcePath());
-		const bool bTargetIsAnyEnum = IsPropertyOfTypeOrChild(*TBaseStructure<FStateTreeAnyEnum>::Get(), TargetStruct, Binding.GetTargetPath());
+		const bool bSourceIsAnyEnum = IsPropertyChildOfAnyStruct<FStateTreeAnyEnum>(*SourceStruct, Binding.GetSourcePath());
+		const bool bTargetIsAnyEnum = IsPropertyChildOfAnyStruct<FStateTreeAnyEnum>(TargetStruct, Binding.GetTargetPath());
 		if (bSourceIsAnyEnum || bTargetIsAnyEnum)
 		{
 			if (bSourceIsAnyEnum)
@@ -2460,6 +2549,12 @@ bool FStateTreeCompiler::GetAndValidateBindings(const FStateTreeBindableStructDe
 			bSucceeded = false;
 			continue;
 		}
+
+		if (!ValidateBindingOnNode(TargetStruct, Binding.GetTargetPath()))
+		{
+			bSucceeded = false;
+			continue;
+		}
 	}
 
 	if (!bSucceeded)
@@ -2610,15 +2705,6 @@ void FStateTreeCompiler::InstantiateStructSubobjects(FStructView Struct)
 				}
 			}
 		}
-		if (const FStructProperty* StructProperty = CastField<FStructProperty>(It->Key))
-		{
-			// If we encounter instanced struct, recursively handle it too.
-			if (StructProperty->Struct == TBaseStructure<FInstancedStruct>::Get())
-			{
-				FInstancedStruct& InstancedStruct = *static_cast<FInstancedStruct*>(const_cast<void*>(It->Value));
-				InstantiateStructSubobjects(InstancedStruct);
-			}
-		}
 	}
 }
 
@@ -2687,3 +2773,22 @@ bool FStateTreeCompiler::NotifyInternalPost()
 
 	return true;
 }
+
+void FStateTreeCompiler::CreateBindingSourceStructsForNode(const FStateTreeEditorNode& EditorNode, const FStateTreeBindableStructDesc& InstanceDesc)
+{
+	using namespace UE::StateTree::Compiler;
+
+	check(EditorNode.Node.IsValid());
+
+	BindingsCompiler.AddSourceStruct(InstanceDesc);
+
+	if (IsNodeStructEligibleForBinding(EditorNode.GetNode()))
+	{
+		FStateTreeBindableStructDesc NodeDesc = InstanceDesc;
+		NodeDesc.ID = EditorNode.GetNodeID();
+		NodeDesc.Struct = EditorNode.Node.GetScriptStruct();
+		NodeDesc.DataHandle = FStateTreeDataHandle::Invalid;	// DataHandle is only for Instance data
+
+		BindingsCompiler.AddSourceStruct(NodeDesc);
+	}
+}
diff --git a/Engine/Plugins/Runtime/StateTree/Source/StateTreeEditorModule/Private/StateTreeEditorData.cpp b/Engine/Plugins/Runtime/StateTree/Source/StateTreeEditorModule/Private/StateTreeEditorData.cpp
--- a/Engine/Plugins/Runtime/StateTree/Source/StateTreeEditorModule/Private/StateTreeEditorData.cpp
+++ b/Engine/Plugins/Runtime/StateTree/Source/StateTreeEditorModule/Private/StateTreeEditorData.cpp
@@ -1236,6 +1236,19 @@ EStateTreeVisitor UStateTreeEditorData::VisitStateNodes(const UStateTreeState& S
 					bContinue = false;
 					break;
 				}
+
+				FStateTreeBindableStructDesc NodeDesc;
+				NodeDesc.StatePath = StatePathWithConditions;
+				NodeDesc.Struct = Node.Node.GetScriptStruct();
+				NodeDesc.Name = Node.GetName();
+				NodeDesc.ID = Node.GetNodeID();
+				NodeDesc.DataSource = EStateTreeBindableStructSource::Condition;
+
+				if (InFunc(&State, NodeDesc, Node.GetNode()) == EStateTreeVisitor::Break)
+				{
+					bContinue = false;
+					break;
+				}
 			}
 		}
 	}
@@ -1264,6 +1277,19 @@ EStateTreeVisitor UStateTreeEditorData::VisitStateNodes(const UStateTreeState& S
 					bContinue = false;
 					break;
 				}
+
+				FStateTreeBindableStructDesc NodeDesc;
+				NodeDesc.StatePath = StatePathWithConsiderations;
+				NodeDesc.Struct = Node.Node.GetScriptStruct();
+				NodeDesc.Name = Node.GetName();
+				NodeDesc.ID = Node.GetNodeID();
+				NodeDesc.DataSource = EStateTreeBindableStructSource::Consideration;
+
+				if (InFunc(&State, NodeDesc, Node.GetNode()) == EStateTreeVisitor::Break)
+				{
+					bContinue = false;
+					break;
+				}
 			}
 		}
 	}
@@ -1291,6 +1317,19 @@ EStateTreeVisitor UStateTreeEditorData::VisitStateNodes(const UStateTreeState& S
 					bContinue = false;
 					break;
 				}
+
+				FStateTreeBindableStructDesc NodeDesc;
+				NodeDesc.StatePath = StatePath;
+				NodeDesc.Struct = Node.Node.GetScriptStruct();
+				NodeDesc.Name = Node.GetName();
+				NodeDesc.ID = Node.GetNodeID();
+				NodeDesc.DataSource = EStateTreeBindableStructSource::Task;
+
+				if (InFunc(&State, NodeDesc, Node.GetNode()) == EStateTreeVisitor::Break)
+				{
+					bContinue = false;
+					break;
+				}
 			}
 		}
 	}
@@ -1315,6 +1354,18 @@ EStateTreeVisitor UStateTreeEditorData::VisitStateNodes(const UStateTreeState& S
 				{
 					bContinue = false;
 				}
+
+				FStateTreeBindableStructDesc NodeDesc;
+				NodeDesc.StatePath = StatePath;
+				NodeDesc.Struct = State.SingleTask.Node.GetScriptStruct();
+				NodeDesc.Name = State.SingleTask.GetName();
+				NodeDesc.ID = State.SingleTask.GetNodeID();
+				NodeDesc.DataSource = EStateTreeBindableStructSource::Task;
+
+				if (InFunc(&State, NodeDesc, State.SingleTask.GetNode()) == EStateTreeVisitor::Break)
+				{
+					bContinue = false;
+				}
 			}
 		}
 
@@ -1379,6 +1430,19 @@ EStateTreeVisitor UStateTreeEditorData::VisitStateNodes(const UStateTreeState& S
 						bContinue = false;
 						break;
 					}
+
+					FStateTreeBindableStructDesc NodeDesc;
+					NodeDesc.StatePath = StatePathWithTransition;
+					NodeDesc.Struct = Node.Node.GetScriptStruct();
+					NodeDesc.Name = Node.GetName();
+					NodeDesc.ID = Node.GetNodeID();
+					NodeDesc.DataSource = EStateTreeBindableStructSource::Condition;
+
+					if (InFunc(&State, NodeDesc, Node.GetNode()) == EStateTreeVisitor::Break)
+					{
+						bContinue = false;
+						break;
+					}
 				}
 			}
 		}
@@ -1537,6 +1601,18 @@ EStateTreeVisitor UStateTreeEditorData::VisitGlobalNodes(TFunctionRef<EStateTree
 			{
 				return EStateTreeVisitor::Break;
 			}
+
+			FStateTreeBindableStructDesc NodeDesc;
+			NodeDesc.StatePath = UE::StateTree::Editor::GlobalStateName;
+			NodeDesc.Struct = Node.Node.GetScriptStruct();
+			NodeDesc.Name = Node.GetName();
+			NodeDesc.ID = Node.GetNodeID();
+			NodeDesc.DataSource = EStateTreeBindableStructSource::Evaluator;
+
+			if (InFunc(nullptr, NodeDesc, Node.GetNode()) == EStateTreeVisitor::Break)
+			{
+				return EStateTreeVisitor::Break;
+			}
 		}
 	}
 
@@ -1561,6 +1637,18 @@ EStateTreeVisitor UStateTreeEditorData::VisitGlobalNodes(TFunctionRef<EStateTree
 			{
 				return EStateTreeVisitor::Break;
 			}
+
+			FStateTreeBindableStructDesc NodeDesc;
+			NodeDesc.StatePath = UE::StateTree::Editor::GlobalStateName;
+			NodeDesc.Struct = Node.Node.GetScriptStruct();
+			NodeDesc.Name = Node.GetName();
+			NodeDesc.ID = Node.GetNodeID();
+			NodeDesc.DataSource = EStateTreeBindableStructSource::GlobalTask;
+
+			if (InFunc(nullptr, NodeDesc, Node.GetNode()) == EStateTreeVisitor::Break)
+			{
+				return EStateTreeVisitor::Break;
+			}
 		}
 	}
 
diff --git a/Engine/Plugins/Runtime/StateTree/Source/StateTreeEditorModule/Private/StateTreePropertyBindingCompiler.cpp b/Engine/Plugins/Runtime/StateTree/Source/StateTreeEditorModule/Private/StateTreePropertyBindingCompiler.cpp
--- a/Engine/Plugins/Runtime/StateTree/Source/StateTreeEditorModule/Private/StateTreePropertyBindingCompiler.cpp
+++ b/Engine/Plugins/Runtime/StateTree/Source/StateTreeEditorModule/Private/StateTreePropertyBindingCompiler.cpp
@@ -323,7 +323,7 @@ bool FStateTreePropertyBindingCompiler::CompileReferences(const FStateTreeBindab
 		if (!UE::StateTree::PropertyRefHelpers::IsPropertyAccessibleForPropertyRef(SourceIndirections, *SourceStruct))
 		{
 			Log->Reportf(EMessageSeverity::Error, TargetStruct,
-					TEXT("%s cannot reference non-output %s "),
+					TEXT("%s cannot reference non-output or property function output %s "),
 					*UE::StateTree::GetDescAndPathAsString(TargetStruct, Binding.GetTargetPath()),
 					*UE::StateTree::GetDescAndPathAsString(*SourceStruct, Binding.GetSourcePath()));
 			return false;
@@ -428,8 +428,11 @@ int32 FStateTreePropertyBindingCompiler::AddSourceStruct(const FStateTreeBindabl
 		UE_LOG(LogStateTree, Error, TEXT("%s already exists as %s using ID '%s'"),
 			*SourceStruct.ToString(), *ExistingStruct->ToString(), *ExistingStruct->ID.ToString());
 	}
-	
-	UE_CLOG(!SourceStruct.DataHandle.IsValid(), LogStateTree, Error, TEXT("%s does not have a valid data handle."), *SourceStruct.ToString());
+
+	if (SourceStruct.Struct && !SourceStruct.Struct->IsChildOf<FStateTreeNodeBase>())
+	{
+		UE_CLOG(!SourceStruct.DataHandle.IsValid(), LogStateTree, Error, TEXT("%s does not have a valid data handle."), *SourceStruct.ToString())
+	}
 	
 	SourceStructs.Add(SourceStruct);
 	return SourceStructs.Num() - 1;
diff --git a/Engine/Plugins/Runtime/StateTree/Source/StateTreeEditorModule/Public/StateTreeCompiler.h b/Engine/Plugins/Runtime/StateTree/Source/StateTreeEditorModule/Public/StateTreeCompiler.h
--- a/Engine/Plugins/Runtime/StateTree/Source/StateTreeEditorModule/Public/StateTreeCompiler.h
+++ b/Engine/Plugins/Runtime/StateTree/Source/StateTreeEditorModule/Public/StateTreeCompiler.h
@@ -86,9 +86,9 @@ private:
 	bool CreateTask(UStateTreeState* State, const FStateTreeEditorNode& TaskNode, const FStateTreeDataHandle TaskDataHandle);
 	bool CreateEvaluator(const FStateTreeEditorNode& EvalNode, const FStateTreeDataHandle EvalDataHandle);
 
-	FInstancedStruct* CreateNode(UStateTreeState* State, const FStateTreeEditorNode& EditorNode, FStateTreeBindableStructDesc& StructDesc, const FStateTreeDataHandle DataHandle, TArray<FInstancedStruct>& InstancedStructContainer);
-	FInstancedStruct* CreateNodeWithSharedInstanceData(UStateTreeState* State, const FStateTreeEditorNode& EditorNode, FStateTreeBindableStructDesc& StructDesc);
-	TOptional<FStateTreeDataView> CreateNodeInstanceData(const FStateTreeEditorNode& EditorNode, FStateTreeNodeBase& Node, FStateTreeBindableStructDesc& StructDesc, const FStateTreeDataHandle DataHandle, TArray<FInstancedStruct>& InstancedStructContainer);
+	FInstancedStruct* CreateNode(UStateTreeState* State, const FStateTreeEditorNode& EditorNode, FStateTreeBindableStructDesc& InstanceDesc, const FStateTreeDataHandle DataHandle, TArray<FInstancedStruct>& InstancedStructContainer);
+	FInstancedStruct* CreateNodeWithSharedInstanceData(UStateTreeState* State, const FStateTreeEditorNode& EditorNode, FStateTreeBindableStructDesc& InstanceDesc);
+	TOptional<FStateTreeDataView> CreateNodeInstanceData(const FStateTreeEditorNode& EditorNode, FStateTreeNodeBase& Node, FStateTreeBindableStructDesc& InstanceDesc, const FStateTreeDataHandle DataHandle, TArray<FInstancedStruct>& InstancedStructContainer);
 
 	struct FValidatedPathBindings
 	{
@@ -99,15 +99,17 @@ private:
 	};
 
 	bool GetAndValidateBindings(const FStateTreeBindableStructDesc& TargetStruct, FStateTreeDataView TargetValue, FValidatedPathBindings& OutValidatedBindings) const;
-	bool IsPropertyOfTypeOrChild(UScriptStruct& Type, const FStateTreeBindableStructDesc& Struct, FPropertyBindingPath Path) const;
-	bool ValidateStructRef(const FStateTreeBindableStructDesc& SourceStruct, FPropertyBindingPath SourcePath,
-					const FStateTreeBindableStructDesc& TargetStruct, FPropertyBindingPath TargetPath) const;
-	bool CompileAndValidateNode(const UStateTreeState* SourceState, const FStateTreeBindableStructDesc& NodeDesc, FStructView NodeView, const FStateTreeDataView InstanceData);
+	bool ValidateStructRef(const FStateTreeBindableStructDesc& SourceStruct, const FPropertyBindingPath& SourcePath,
+	const FStateTreeBindableStructDesc& TargetStruct, const FPropertyBindingPath& TargetPath) const;
+	bool ValidateBindingOnNode(const FStateTreeBindableStructDesc& TargetStruct, const FPropertyBindingPath& TargetPath) const;
+	bool CompileAndValidateNode(const UStateTreeState* SourceState, const FStateTreeBindableStructDesc& InstanceDesc, FStructView NodeView, const FStateTreeDataView InstanceData);
 
 	void InstantiateStructSubobjects(FStructView Struct);
-
+        
 	bool NotifyInternalPost();
-
+	
+	void CreateBindingSourceStructsForNode(const FStateTreeEditorNode& EditorNode, const FStateTreeBindableStructDesc& InstanceDesc);
+      
 private:
 	FStateTreeCompilerLog& Log;
 	TObjectPtr<UStateTree> StateTree = nullptr;
diff --git a/Engine/Plugins/Runtime/StateTree/Source/StateTreeEditorModule/Public/StateTreeEditorNode.h b/Engine/Plugins/Runtime/StateTree/Source/StateTreeEditorModule/Public/StateTreeEditorNode.h
--- a/Engine/Plugins/Runtime/StateTree/Source/StateTreeEditorModule/Public/StateTreeEditorNode.h
+++ b/Engine/Plugins/Runtime/StateTree/Source/StateTreeEditorModule/Public/StateTreeEditorNode.h
@@ -41,6 +41,24 @@ struct FStateTreeEditorNode
 	 */
 	UE_API FName GetName() const;
 
+	/**
+	 * Get DataView for the node template.
+	 * @return DataView for the node template
+	 */
+	TStructView<FStateTreeNodeBase> GetNode() const
+	{
+		return TStructView<FStateTreeNodeBase>(Node.GetScriptStruct(), const_cast<uint8*>(Node.GetMemory()));
+	}
+
+	/**
+	 * Get ID for the node template, which is different from the ID for the Instance.
+	 * @return ID for the node template
+	 */
+	FGuid GetNodeID() const
+	{
+		return FGuid::Combine(ID, FGuid::NewDeterministicGuid(TEXT("Node Struct")));
+	}
+
 	FStateTreeDataView GetInstance() const
 	{
 		return InstanceObject ? FStateTreeDataView(InstanceObject) : FStateTreeDataView(const_cast<FInstancedStruct&>(Instance));
@@ -59,7 +77,10 @@ struct FStateTreeEditorNode
 
 	UPROPERTY(EditDefaultsOnly, Instanced, Category = Node)
 	TObjectPtr<UObject> InstanceObject = nullptr;
-	
+
+	/**
+	 * ID for the node instance.
+	 */
 	UPROPERTY(EditDefaultsOnly, Category = Node)
 	FGuid ID;
 
diff --git a/Engine/Plugins/Runtime/StateTree/Source/StateTreeTestSuite/Private/StateTreeBindingTest.cpp b/Engine/Plugins/Runtime/StateTree/Source/StateTreeTestSuite/Private/StateTreeBindingTest.cpp
--- a/Engine/Plugins/Runtime/StateTree/Source/StateTreeTestSuite/Private/StateTreeBindingTest.cpp
+++ b/Engine/Plugins/Runtime/StateTree/Source/StateTreeTestSuite/Private/StateTreeBindingTest.cpp
@@ -360,6 +360,7 @@ struct FStateTreeTest_CopyObjects : FStateTreeTestBase
 };
 IMPLEMENT_STATE_TREE_INSTANT_TEST(FStateTreeTest_CopyObjects, "System.StateTree.Binding.CopyObjects");
 
+//@todo: Add test coverage for propertyrefs across frames
 
 struct FStateTreeTest_References : FStateTreeTestBase
 {
@@ -516,6 +517,101 @@ struct FStateTreeTest_ReferencesConstness : FStateTreeTestBase
 };
 IMPLEMENT_STATE_TREE_INSTANT_TEST(FStateTreeTest_ReferencesConstness, "System.StateTree.Binding.ReferencesConstness");
 
+struct FStateTreeTest_ReferencesOnNode : FStateTreeTestBase
+{
+	virtual bool InstantTest() override
+	{
+		// Root IntA, IntB
+		//     State1 PropertyRefOnNode -> IntA, PropertyRefOnInstance -> IntB
+		//         State2 PropertyRefOnNode -> PropertyRefOnNode(State1), PropertyRefOnInstance -> PropertyRefOnInstance(State1)
+
+		UStateTree& StateTree = NewStateTree();
+		UStateTreeEditorData& EditorData = *Cast<UStateTreeEditorData>(StateTree.EditorData);
+
+		UStateTreeState& Root = EditorData.AddSubTree(FName(TEXT("Root")));
+		UStateTreeState& State1 = Root.AddChildState(FName(TEXT("State1")));
+		UStateTreeState& State2 = State1.AddChildState(FName(TEXT("State2")));
+
+		TStateTreeEditorNode<FTestTask_IntegersOutput>& RootTaskNode = Root.AddTask<FTestTask_IntegersOutput>(FName(TEXT("RootTask")));
+		TStateTreeEditorNode<FTestTask_PropertyRefOnNodeAndInstance>& State1TaskNode = State1.AddTask<FTestTask_PropertyRefOnNodeAndInstance>(FName(TEXT("State1Task")));
+		TStateTreeEditorNode<FTestTask_PropertyRefOnNodeAndInstance>& State2TaskNode = State2.AddTask<FTestTask_PropertyRefOnNodeAndInstance>(FName(TEXT("State2Task")));
+
+		EditorData.AddPropertyBinding(
+			FPropertyBindingPath(RootTaskNode.ID, GET_MEMBER_NAME_CHECKED(FTestTask_IntegersOutput_InstanceData, IntA)),
+			FPropertyBindingPath(State1TaskNode.GetNodeID(), GET_MEMBER_NAME_CHECKED(FTestTask_PropertyRefOnNodeAndInstance, RefOnNode)));
+		EditorData.AddPropertyBinding(
+			FPropertyBindingPath(RootTaskNode.ID, GET_MEMBER_NAME_CHECKED(FTestTask_IntegersOutput_InstanceData, IntB)),
+			FPropertyBindingPath(State1TaskNode.ID, GET_MEMBER_NAME_CHECKED(FTestTask_PropertyRefOnNodeAndInstance_InstanceData, RefOnInstance)));
+		EditorData.AddPropertyBinding(
+			FPropertyBindingPath(State1TaskNode.GetNodeID(), GET_MEMBER_NAME_CHECKED(FTestTask_PropertyRefOnNodeAndInstance, RefOnNode)),
+			FPropertyBindingPath(State2TaskNode.GetNodeID(), GET_MEMBER_NAME_CHECKED(FTestTask_PropertyRefOnNodeAndInstance, RefOnNode)));
+		EditorData.AddPropertyBinding(
+			FPropertyBindingPath(State1TaskNode.ID, GET_MEMBER_NAME_CHECKED(FTestTask_PropertyRefOnNodeAndInstance_InstanceData, RefOnInstance)),
+			FPropertyBindingPath(State2TaskNode.ID, GET_MEMBER_NAME_CHECKED(FTestTask_PropertyRefOnNodeAndInstance_InstanceData, RefOnInstance)));
+
+		void* IntAAddress = nullptr;
+		void* IntBAddress = nullptr;
+		void* State1RefOnNodeAddress = nullptr;
+		void* State1RefOnInstanceAddress = nullptr;
+		void* State2RefOnNodeAddress = nullptr;
+		void* State2RefOnInstanceAddress = nullptr;
+
+		RootTaskNode.GetNode().CustomEnterStateFunc = 
+			[&IntAAddress, &IntBAddress](FStateTreeExecutionContext& ExecContext, const FTestTask_IntegersOutput& Node)
+			{
+				auto& InstanceData = ExecContext.GetInstanceData<FTestTask_IntegersOutput::FInstanceDataType>(Node);
+				IntAAddress = &InstanceData.IntA;
+				IntBAddress = &InstanceData.IntB;
+			};
+
+		State1TaskNode.GetNode().CustomEnterStateFunc =
+			[&State1RefOnNodeAddress, &State1RefOnInstanceAddress](FStateTreeExecutionContext& ExecContext, const FTestTask_PropertyRefOnNodeAndInstance& Node)
+			{
+				auto& InstanceData = ExecContext.GetInstanceData<FTestTask_PropertyRefOnNodeAndInstance::FInstanceDataType>(Node);
+
+				State1RefOnNodeAddress = Node.RefOnNode.GetMutablePtr<int32>(ExecContext);
+				State1RefOnInstanceAddress = InstanceData.RefOnInstance.GetMutablePtr<int32>(ExecContext);
+			};
+
+		State2TaskNode.GetNode().CustomEnterStateFunc =
+			[&State2RefOnNodeAddress, &State2RefOnInstanceAddress](FStateTreeExecutionContext& ExecContext, const FTestTask_PropertyRefOnNodeAndInstance& Node)
+			{
+				auto& InstanceData = ExecContext.GetInstanceData<FTestTask_PropertyRefOnNodeAndInstance::FInstanceDataType>(Node);
+
+				State2RefOnNodeAddress = Node.RefOnNode.GetMutablePtr<int32>(ExecContext);
+				State2RefOnInstanceAddress = InstanceData.RefOnInstance.GetMutablePtr<int32>(ExecContext);
+			};
+
+		{
+			FStateTreeCompilerLog Log;
+			FStateTreeCompiler Compiler(Log);
+			const bool bResult = Compiler.Compile(StateTree);
+			AITEST_TRUE("StateTree should get compiled", bResult);
+		}
+
+		{
+			FStateTreeInstanceData InstanceData;
+			FTestStateTreeExecutionContext Exec(StateTree, StateTree, InstanceData);
+			const bool bInitSucceeded = Exec.IsValid();
+			AITEST_TRUE("StateTree should init", bInitSucceeded);
+
+			Exec.Start();
+			AITEST_TRUE("Expected Root to be active.", Exec.ExpectInActiveStates("Root", "State1", "State2"));
+			AITEST_TRUE("int addresses should have been initialized", IntAAddress && IntBAddress);
+			AITEST_EQUAL("Ref on Node should fetch IntA", IntAAddress, State1RefOnNodeAddress);
+			AITEST_EQUAL("Chained refs on Node should point to the same address", State1RefOnNodeAddress, State2RefOnNodeAddress);
+			AITEST_EQUAL("Ref on Instance should fetch IntB", IntBAddress, State1RefOnInstanceAddress);
+			AITEST_EQUAL("Chained refs on Instance should point to the same address", State1RefOnInstanceAddress, State2RefOnInstanceAddress);
+			Exec.LogClear();
+
+			Exec.Stop();
+		}
+
+		return true;
+	}
+};
+IMPLEMENT_STATE_TREE_INSTANT_TEST(FStateTreeTest_ReferencesOnNode, "System.StateTree.Binding.ReferencesOnNode");
+
 struct FStateTreeTest_MutableArray : FStateTreeTestBase
 {
 	virtual bool InstantTest() override
diff --git a/Engine/Plugins/Runtime/StateTree/Source/StateTreeTestSuite/Private/StateTreeDelegateTest.cpp b/Engine/Plugins/Runtime/StateTree/Source/StateTreeTestSuite/Private/StateTreeDelegateTest.cpp
--- a/Engine/Plugins/Runtime/StateTree/Source/StateTreeTestSuite/Private/StateTreeDelegateTest.cpp
+++ b/Engine/Plugins/Runtime/StateTree/Source/StateTreeTestSuite/Private/StateTreeDelegateTest.cpp
@@ -460,6 +460,188 @@ struct FStateTreeTest_Delegate_ListeningToDelegateOnExit2 : FStateTreeTestBase
 };
 IMPLEMENT_AI_INSTANT_TEST(FStateTreeTest_Delegate_ListeningToDelegateOnExit2, "System.StateTree.Delegate.ListeningToDelegateOnExit2");
 
+struct FStateTreeTest_Delegate_ListenerDispatcherOnNode : FStateTreeTestBase
+{
+	virtual bool InstantTest() override
+	{
+		// Root(Dispatcher, ListenerOnNode1 -> DispatcherOnNode, ListenerOnNode2 -> DispatcherOnInstance, ListenerOnInstance -> DispatcherOnNode)
+
+		UStateTree& StateTree = NewStateTree();
+		UStateTreeEditorData& EditorData = *Cast<UStateTreeEditorData>(StateTree.EditorData);
+		UStateTreeState& Root = EditorData.AddSubTree(FName(TEXT("Root")));
+
+		bool bListenerOnNode1Broadcasted = false;
+		bool bListenerOnNode2Broadcasted = false;
+		bool bListenerOnInstanceBroadcasted = false;
+
+		{
+			TStateTreeEditorNode<FTestTask_DispatcherOnNodeAndInstance>& DispatcherTaskNode = Root.AddTask<FTestTask_DispatcherOnNodeAndInstance>(FName(TEXT("Dispatcher")));
+
+			TStateTreeEditorNode<FTestTask_ListenerOnNode>& ListenerOnNode1TaskNode = Root.AddTask<FTestTask_ListenerOnNode>(FName(TEXT("ListenerOnNode1")));
+			ListenerOnNode1TaskNode.GetNode().CustomFunc = [&bListenerOnNode1Broadcasted](const FStateTreeWeakExecutionContext& WeakExecContext)
+			{
+				bListenerOnNode1Broadcasted = true;
+			};
+
+			TStateTreeEditorNode<FTestTask_ListenerOnNode>& ListenerOnNode2TaskNode = Root.AddTask<FTestTask_ListenerOnNode>(FName(TEXT("ListenerOnNode2")));
+			ListenerOnNode2TaskNode.GetNode().CustomFunc = [&bListenerOnNode2Broadcasted](const FStateTreeWeakExecutionContext& WeakExecContext)
+			{
+				bListenerOnNode2Broadcasted = true;
+			};
+
+			TStateTreeEditorNode<FTestTask_ListenerOnInstance>& ListenerOnInstanceTaskNode = Root.AddTask<FTestTask_ListenerOnInstance>(FName(TEXT("ListenerOnInstance")));
+			ListenerOnInstanceTaskNode.GetNode().CustomFunc = [&bListenerOnInstanceBroadcasted](const FStateTreeWeakExecutionContext& WeakExecContext)
+			{
+				bListenerOnInstanceBroadcasted = true;
+			};
+
+			EditorData.AddPropertyBinding(
+				FPropertyBindingPath(DispatcherTaskNode.GetNodeID(), GET_MEMBER_NAME_CHECKED(FTestTask_DispatcherOnNodeAndInstance, DispatcherOnNode)),
+				FPropertyBindingPath(ListenerOnNode1TaskNode.GetNodeID(), GET_MEMBER_NAME_CHECKED(FTestTask_ListenerOnNode, ListenerOnNode)));
+
+			EditorData.AddPropertyBinding(
+				FPropertyBindingPath(DispatcherTaskNode.ID, GET_MEMBER_NAME_CHECKED(FTestTask_DispatcherOnNodeAndInstance_InstanceData, DispatcherOnInstance)),
+				FPropertyBindingPath(ListenerOnNode2TaskNode.GetNodeID(), GET_MEMBER_NAME_CHECKED(FTestTask_ListenerOnNode, ListenerOnNode)));
+
+			EditorData.AddPropertyBinding(
+				FPropertyBindingPath(DispatcherTaskNode.GetNodeID(), GET_MEMBER_NAME_CHECKED(FTestTask_DispatcherOnNodeAndInstance, DispatcherOnNode)),
+				FPropertyBindingPath(ListenerOnInstanceTaskNode.ID, GET_MEMBER_NAME_CHECKED(FTestTask_ListenerOnInstance_InstanceData, ListenerOnInstance)));
+		}
+
+		{
+			FStateTreeCompilerLog Log;
+			FStateTreeCompiler Compiler(Log);
+			const bool bResult = Compiler.Compile(StateTree);
+			AITEST_TRUE("StateTree should get compiled", bResult);
+
+			const FCompactStateTreeState& RootState = StateTree.GetStates()[0];
+			int32 TaskNodeIndex = RootState.TasksBegin;
+
+			FConstStructView DispatcherView = StateTree.GetNode(TaskNodeIndex++);
+			FConstStructView DispatcherInstanceDataView = StateTree.GetDefaultInstanceData().GetStruct(DispatcherView.Get<const FStateTreeTaskBase>().InstanceTemplateIndex.Get());
+
+			FConstStructView ListenerOnNode1View = StateTree.GetNode(TaskNodeIndex++);
+
+			FConstStructView ListenerOnNode2View = StateTree.GetNode(TaskNodeIndex++);
+
+			FConstStructView ListenerOnInstanceView = StateTree.GetNode(TaskNodeIndex++);
+			FConstStructView ListenerOnInstanceInstanceDataView = StateTree.GetDefaultInstanceData().GetStruct(ListenerOnInstanceView.Get<const FStateTreeTaskBase>().InstanceTemplateIndex.Get());
+
+			const FStateTreeDelegateDispatcher DispatcherOnNode = DispatcherView.Get<const FTestTask_DispatcherOnNodeAndInstance>().DispatcherOnNode;
+			const FStateTreeDelegateDispatcher DispatcherOnInstance = DispatcherInstanceDataView.Get<const FTestTask_DispatcherOnNodeAndInstance_InstanceData>().DispatcherOnInstance;
+			const FStateTreeDelegateListener ListenerOnNode1 = ListenerOnNode1View.Get<const FTestTask_ListenerOnNode>().ListenerOnNode;
+			const FStateTreeDelegateListener ListenerOnNode2 = ListenerOnNode2View.Get<const FTestTask_ListenerOnNode>().ListenerOnNode;
+			const FStateTreeDelegateListener ListenerOnInstance = ListenerOnInstanceInstanceDataView.Get<const FTestTask_ListenerOnInstance_InstanceData>().ListenerOnInstance;
+
+			AITEST_TRUE("Dispatcher on Node should init", DispatcherOnNode.IsValid());
+			AITEST_TRUE("Dispatcher on Instance should init", DispatcherOnInstance.IsValid());
+			AITEST_TRUE("Listener on Node should init", ListenerOnNode1.IsValid() && ListenerOnNode2.IsValid());
+			AITEST_TRUE("Listener on Instance should init", ListenerOnInstance.IsValid());
+
+			AITEST_EQUAL("ListenerOnNode1 should be bound to Dispatcher on Node", DispatcherOnNode, ListenerOnNode1.GetDispatcher());
+			AITEST_EQUAL("ListenerOnNode2 should be bound to Dispatcher on Instance", DispatcherOnInstance, ListenerOnNode2.GetDispatcher());
+			AITEST_EQUAL("ListenerOnInstance should be bound to Dispatcher on Node", DispatcherOnNode, ListenerOnInstance.GetDispatcher());
+		}
+
+		{
+			FStateTreeInstanceData InstanceData;
+			FTestStateTreeExecutionContext Exec(StateTree, StateTree, InstanceData);
+			const bool bInitSucceeded = Exec.IsValid();
+			AITEST_TRUE("StateTree should init", bInitSucceeded);
+
+			Exec.Start();
+			AITEST_TRUE("Expected Root to be active.", Exec.ExpectInActiveStates("Root"));
+			Exec.LogClear();
+
+			Exec.Tick(0.1f);
+			AITEST_TRUE("Expected Root to be active.", Exec.ExpectInActiveStates("Root"));
+			AITEST_TRUE("ListenerOnNode1 should be broadcasted", bListenerOnNode1Broadcasted);
+			AITEST_TRUE("ListenerOnNode2 should be broadcasted", bListenerOnNode2Broadcasted);
+			AITEST_TRUE("ListenerOnInstance should be broadcasted", bListenerOnInstanceBroadcasted);
+
+			Exec.Stop();
+		}
+
+
+		return true;
+	}
+};
+IMPLEMENT_AI_INSTANT_TEST(FStateTreeTest_Delegate_ListenerDispatcherOnNode, "System.StateTree.Delegate.ListenerDispatcherOnNode");
+
+struct FStateTreeTest_Delegate_DispatcherOnNodeListenerOnTransition : FStateTreeTestBase
+{
+	virtual bool InstantTest() override
+	{
+		// Root(Dispatcher)
+		//     State1(Transition on Delegate -> DispatcherOnNode) -> State2
+		//	   State2(Transition on Delegate -> DispatcherOnInstance) -> Tree Succeeded
+
+		UStateTree& StateTree = NewStateTree();
+		UStateTreeEditorData& EditorData = *Cast<UStateTreeEditorData>(StateTree.EditorData);
+
+		UStateTreeState& Root = EditorData.AddSubTree(FName(TEXT("Root")));
+		TStateTreeEditorNode<FTestTask_DispatcherOnNodeAndInstance> DispatcherTaskNode = Root.AddTask<FTestTask_DispatcherOnNodeAndInstance>(TEXT("Dispatcher"));
+		UStateTreeState& State1 = Root.AddChildState(FName(TEXT("State1")));
+		FStateTreeTransition& DelegateOnNodeTransition = State1.AddTransition(EStateTreeTransitionTrigger::OnDelegate, EStateTreeTransitionType::NextState);
+
+		UStateTreeState& State2 = Root.AddChildState(FName(TEXT("State2")));
+		FStateTreeTransition& DelegateOnInstanceTransition = State2.AddTransition(EStateTreeTransitionTrigger::OnDelegate, EStateTreeTransitionType::Succeeded);
+
+		EditorData.AddPropertyBinding(
+			FPropertyBindingPath(DispatcherTaskNode.GetNodeID(), GET_MEMBER_NAME_CHECKED(FTestTask_DispatcherOnNodeAndInstance, DispatcherOnNode)),
+			FPropertyBindingPath(DelegateOnNodeTransition.ID, GET_MEMBER_NAME_CHECKED(FStateTreeTransition, DelegateListener)));
+
+		EditorData.AddPropertyBinding(
+			FPropertyBindingPath(DispatcherTaskNode.ID, GET_MEMBER_NAME_CHECKED(FTestTask_DispatcherOnNodeAndInstance_InstanceData, DispatcherOnInstance)),
+			FPropertyBindingPath(DelegateOnInstanceTransition.ID, GET_MEMBER_NAME_CHECKED(FStateTreeTransition, DelegateListener)));
+
+		{
+			FStateTreeCompilerLog Log;
+			FStateTreeCompiler Compiler(Log);
+			const bool bResult = Compiler.Compile(StateTree);
+			AITEST_TRUE("StateTree should get compiled", bResult);
+
+			const FCompactStateTreeState& RootState = StateTree.GetStates()[0];
+			const FCompactStateTreeState& State1State = StateTree.GetStates()[1];
+			const FCompactStateTreeState& State2State = StateTree.GetStates()[2];
+
+			const FTestTask_DispatcherOnNodeAndInstance& DispatcherTask = StateTree.GetNode(RootState.TasksBegin).Get<const FTestTask_DispatcherOnNodeAndInstance>();
+			const FStateTreeDelegateDispatcher DispatcherOnNode = DispatcherTask.DispatcherOnNode;
+			const FStateTreeDelegateDispatcher DispatcherOnInstance = StateTree.GetDefaultInstanceData().GetStruct(DispatcherTask.InstanceTemplateIndex.Get()).Get<const FTestTask_DispatcherOnNodeAndInstance_InstanceData>().DispatcherOnInstance;
+
+			AITEST_TRUE("Dispatcher on node should be valid", DispatcherOnNode.IsValid());
+			AITEST_TRUE("Dispatcher on instance should be valid", DispatcherOnInstance.IsValid());
+
+			const FCompactStateTransition* CompactDelegateOnNodeTransition = StateTree.GetTransitionFromIndex(FStateTreeIndex16(State1State.TransitionsBegin));
+			const FCompactStateTransition* CompactDelegateOnInstanceTransition = StateTree.GetTransitionFromIndex(FStateTreeIndex16(State2State.TransitionsBegin));
+
+			AITEST_EQUAL("State1 Transition Delegate Listener should be bound to Dispatcher on Node", CompactDelegateOnNodeTransition->RequiredDelegateDispatcher, DispatcherOnNode);
+			AITEST_EQUAL("State2 Transition Delegate Listener should be bound to Dispatcher on Instance", CompactDelegateOnInstanceTransition->RequiredDelegateDispatcher, DispatcherOnInstance);
+		}
+
+		{
+			FStateTreeInstanceData InstanceData;
+			FTestStateTreeExecutionContext Exec(StateTree, StateTree, InstanceData);
+			const bool bInitSucceeded = Exec.IsValid();
+			AITEST_TRUE("StateTree should init", bInitSucceeded);
+
+			Exec.Start();
+			AITEST_TRUE("Expected [Root, State1] to be active.", Exec.ExpectInActiveStates("Root", "State1"));
+			Exec.LogClear();
+
+			Exec.Tick(0.1f);
+			AITEST_TRUE("Expected [Root, State2] to be active.", Exec.ExpectInActiveStates("Root", "State2"));
+			Exec.LogClear();
+
+			Exec.Tick(0.1f);
+			AITEST_TRUE("Expected Tree to be succeeded.", Exec.GetStateTreeRunStatus() == EStateTreeRunStatus::Succeeded);
+			Exec.LogClear();
+		}
+
+		return true;
+	}
+};
+IMPLEMENT_AI_INSTANT_TEST(FStateTreeTest_Delegate_DispatcherOnNodeListenerOnTransition, "System.StateTree.Delegate.DispatcherOnNodeListenerOnTransition");
 } // namespace UE::StateTree::Tests
 
 UE_ENABLE_OPTIMIZATION_SHIP
diff --git a/Engine/Plugins/Runtime/StateTree/Source/StateTreeTestSuite/Private/StateTreeTestTypes.h b/Engine/Plugins/Runtime/StateTree/Source/StateTreeTestSuite/Private/StateTreeTestTypes.h
--- a/Engine/Plugins/Runtime/StateTree/Source/StateTreeTestSuite/Private/StateTreeTestTypes.h
+++ b/Engine/Plugins/Runtime/StateTree/Source/StateTreeTestSuite/Private/StateTreeTestTypes.h
@@ -639,6 +639,89 @@ struct FTestTask_StandNoTick : public FTestTask_Stand
 	}
 };
 
+
+USTRUCT()
+struct FTestTask_IntegersOutput_InstanceData
+{
+	GENERATED_BODY()
+
+	UPROPERTY(EditAnywhere, Category=Output)
+	int32 IntA = 0;
+
+	UPROPERTY(EditAnywhere, Category=Output)
+	int32 IntB = 0;
+
+};
+
+USTRUCT()
+struct FTestTask_IntegersOutput : public FStateTreeTaskBase
+{
+	GENERATED_BODY()
+
+	using FInstanceDataType = FTestTask_IntegersOutput_InstanceData;
+
+	FTestTask_IntegersOutput() = default;
+
+	explicit FTestTask_IntegersOutput(FName InName)
+	{
+		Name = InName;
+	}
+
+	virtual const UStruct* GetInstanceDataType() const override { return FInstanceDataType::StaticStruct(); }
+
+	virtual EStateTreeRunStatus EnterState(FStateTreeExecutionContext& Context, const FStateTreeTransitionResult& Transition) const override
+	{
+		if (CustomEnterStateFunc)
+		{
+			CustomEnterStateFunc(Context, *this);
+		}
+
+		return EStateTreeRunStatus::Running;
+	}
+
+	TFunction<void(FStateTreeExecutionContext&, const FTestTask_IntegersOutput&)> CustomEnterStateFunc;
+};
+
+USTRUCT()
+struct FTestTask_PropertyRefOnNodeAndInstance_InstanceData
+{
+	GENERATED_BODY()
+
+	UPROPERTY(EditAnywhere, Category = Parameter, meta = (RefType = "int32"))
+	FStateTreePropertyRef RefOnInstance;
+};
+
+USTRUCT()
+struct FTestTask_PropertyRefOnNodeAndInstance : public FStateTreeTaskBase
+{
+	GENERATED_BODY()
+
+	using FInstanceDataType = FTestTask_PropertyRefOnNodeAndInstance_InstanceData;
+
+	FTestTask_PropertyRefOnNodeAndInstance() = default;
+	FTestTask_PropertyRefOnNodeAndInstance(FName InName)
+	{
+		Name = InName;
+	}
+
+	virtual const UStruct* GetInstanceDataType() const override { return FInstanceDataType::StaticStruct(); }
+
+	virtual EStateTreeRunStatus EnterState(FStateTreeExecutionContext& Context, const FStateTreeTransitionResult& Transition) const override
+	{
+		if (CustomEnterStateFunc)
+		{
+			CustomEnterStateFunc(Context, *this);
+		}
+
+		return EStateTreeRunStatus::Running;
+	}
+
+	UPROPERTY(EditAnywhere, Category = Parameter, meta = (RefType = "int32"))
+	FStateTreePropertyRef RefOnNode;
+
+	TFunction<void(FStateTreeExecutionContext&, const FTestTask_PropertyRefOnNodeAndInstance&)> CustomEnterStateFunc;
+};
+
 USTRUCT()
 struct FStateTreeTestConditionInstanceData
 {
@@ -1165,4 +1248,122 @@ struct FTestTask_CustomFuncOnDelegate : public FStateTreeTaskBase
 	TFunction<void(const FStateTreeWeakExecutionContext&, FStateTreeDelegateListener)> CustomFunc;
 
 	bool bRemoveOnExit = true;
-};
\ No newline at end of file
+};
+
+USTRUCT()
+struct FTestTask_DispatcherOnNodeAndInstance_InstanceData
+{
+	GENERATED_BODY()
+
+	UPROPERTY(EditAnywhere, Category = Parameter)
+	FStateTreeDelegateDispatcher DispatcherOnInstance;
+};
+
+USTRUCT()
+struct FTestTask_DispatcherOnNodeAndInstance : public FStateTreeTaskBase
+{
+	GENERATED_BODY()
+
+	using FInstanceDataType = FTestTask_DispatcherOnNodeAndInstance_InstanceData;
+
+	FTestTask_DispatcherOnNodeAndInstance() = default;
+
+	explicit FTestTask_DispatcherOnNodeAndInstance(FName InName)
+	{
+		Name = InName;
+	}
+
+	virtual const UStruct* GetInstanceDataType() const override { return FInstanceDataType::StaticStruct(); }
+
+	virtual EStateTreeRunStatus Tick(FStateTreeExecutionContext& Context, const float DeltaTime) const override
+	{
+		const FInstanceDataType& InstanceData = Context.GetInstanceData<FInstanceDataType>(*this);
+
+		Context.BroadcastDelegate(DispatcherOnNode);
+		Context.BroadcastDelegate(InstanceData.DispatcherOnInstance);
+
+		return EStateTreeRunStatus::Running;
+	}
+
+	UPROPERTY(EditAnywhere, Category = Parameter)
+	FStateTreeDelegateDispatcher DispatcherOnNode;
+};
+
+USTRUCT()
+struct FTestTask_ListenerOnNode_InstanceData
+{
+	GENERATED_BODY()
+};
+
+USTRUCT()
+struct FTestTask_ListenerOnNode : public FStateTreeTaskBase
+{
+	GENERATED_BODY()
+
+	using FInstanceDataType = FTestTask_ListenerOnNode_InstanceData;
+
+	FTestTask_ListenerOnNode() = default;
+
+	explicit FTestTask_ListenerOnNode(FName InName)
+	{
+		Name = InName;
+	}
+
+	virtual const UStruct* GetInstanceDataType() const override { return FInstanceDataType::StaticStruct();  }
+
+	virtual EStateTreeRunStatus EnterState(FStateTreeExecutionContext& Context, const FStateTreeTransitionResult& Transition) const override
+	{
+		Context.BindDelegate(ListenerOnNode, FSimpleDelegate::CreateLambda(
+			[this, WeakExecContext = Context.MakeWeakExecutionContext()]()
+			{
+				CustomFunc(WeakExecContext);
+			}));
+
+		return EStateTreeRunStatus::Running;
+	}
+
+	UPROPERTY(EditAnywhere, Category = Parameter)
+	FStateTreeDelegateListener ListenerOnNode;
+
+	TFunction<void(const FStateTreeWeakExecutionContext&)> CustomFunc;
+};
+
+USTRUCT()
+struct FTestTask_ListenerOnInstance_InstanceData
+{
+	GENERATED_BODY()
+
+	UPROPERTY(EditAnywhere, Category = Parameter)
+	FStateTreeDelegateListener ListenerOnInstance;
+};
+
+USTRUCT()
+struct FTestTask_ListenerOnInstance : public FStateTreeTaskBase
+{
+	GENERATED_BODY()
+
+	using FInstanceDataType = FTestTask_ListenerOnInstance_InstanceData;
+
+	FTestTask_ListenerOnInstance() = default;
+
+	explicit FTestTask_ListenerOnInstance(FName InName)
+	{
+		Name = InName;
+	}
+
+	virtual const UStruct* GetInstanceDataType() const override { return FInstanceDataType::StaticStruct(); }
+
+	virtual EStateTreeRunStatus EnterState(FStateTreeExecutionContext& Context, const FStateTreeTransitionResult& Transition) const override
+	{
+		const FInstanceDataType& InstanceData = Context.GetInstanceData<FInstanceDataType>(*this);
+		Context.BindDelegate(InstanceData.ListenerOnInstance, FSimpleDelegate::CreateLambda(
+			[this, WeakExecContext = Context.MakeWeakExecutionContext()]()
+			{
+				CustomFunc(WeakExecContext);
+			}));
+
+		return EStateTreeRunStatus::Running;
+	}
+
+	TFunction<void(const FStateTreeWeakExecutionContext&)> CustomFunc;
+};
diff --git a/Engine/Source/Runtime/CoreUObject/Public/StructUtils/StructView.h b/Engine/Source/Runtime/CoreUObject/Public/StructUtils/StructView.h
--- a/Engine/Source/Runtime/CoreUObject/Public/StructUtils/StructView.h
+++ b/Engine/Source/Runtime/CoreUObject/Public/StructUtils/StructView.h
@@ -134,6 +134,14 @@ public:
 	{
 	}
 
+	template<typename T = BaseStructT>
+		requires (std::is_base_of_v<BaseStructT, std::decay_t<T>>)
+	explicit TStructView(const UScriptStruct* InScriptStruct, uint8* InStructMemory = nullptr)
+		: FStructView(InScriptStruct, InStructMemory)
+	{
+		checkf(InScriptStruct->IsChildOf(TBaseStructure<T>::Get()), TEXT("ScriptStruct must be a child of T!"));
+	}
+
 	template<typename T = BaseStructT>
 		requires (std::is_base_of_v<BaseStructT, std::decay_t<T>>)
 	TStructView(TInstancedStruct<T>& InstancedStruct)
