[StateTree] Add unit test state tree reactivate

#tests automation

diff --git a/Engine/Plugins/Runtime/StateTree/Source/StateTreeTestSuite/Private/StateTreeTestTypes.h b/Engine/Plugins/Runtime/StateTree/Source/StateTreeTestSuite/Private/StateTreeTestTypes.h
--- a/Engine/Plugins/Runtime/StateTree/Source/StateTreeTestSuite/Private/StateTreeTestTypes.h
+++ b/Engine/Plugins/Runtime/StateTree/Source/StateTreeTestSuite/Private/StateTreeTestTypes.h
@@ -600,6 +600,7 @@ struct FTestTask_Stand : public FStateTreeTaskBase
 		FInstanceDataType& InstanceData = Context.GetInstanceData(*this);
 		
 		InstanceData.CurrentTick++;
+		TestLog.Log(Name, FString::Printf(TEXT("CurrentTick=%u"), InstanceData.CurrentTick));
 		
 		return (InstanceData.CurrentTick >= TicksToCompletion) ? TickCompletionResult : EStateTreeRunStatus::Running;
 	};
@@ -639,6 +640,22 @@ struct FTestTask_StandNoTick : public FTestTask_Stand
 	}
 };
 
+USTRUCT()
+struct FTestTask_StandNoStateChangeOnReselect : public FTestTask_Stand
+{
+	GENERATED_BODY()
+
+	FTestTask_StandNoStateChangeOnReselect()
+	{
+		bShouldStateChangeOnReselect = false;
+	}
+
+	FTestTask_StandNoStateChangeOnReselect(const FName InName)
+		: FTestTask_Stand(InName)
+	{
+		bShouldStateChangeOnReselect = false;
+	}
+};
 
 USTRUCT()
 struct FTestTask_IntegersOutput_InstanceData
diff --git a/Engine/Plugins/Runtime/StateTree/Source/StateTreeTestSuite/Private/StateTreeTransitionTest.cpp b/Engine/Plugins/Runtime/StateTree/Source/StateTreeTestSuite/Private/StateTreeTransitionTest.cpp
--- a/Engine/Plugins/Runtime/StateTree/Source/StateTreeTestSuite/Private/StateTreeTransitionTest.cpp
+++ b/Engine/Plugins/Runtime/StateTree/Source/StateTreeTestSuite/Private/StateTreeTransitionTest.cpp
@@ -454,6 +454,104 @@ struct FStateTreeTest_TransitionDelay : FStateTreeTestBase
 };
 IMPLEMENT_STATE_TREE_INSTANT_TEST(FStateTreeTest_TransitionDelay, "System.StateTree.Transition.Delay");
 
+struct FStateTreeTest_TransitionReactivateTargetState : FStateTreeTestBase
+{
+	virtual bool InstantTest() override
+	{
+		//Root
+		// State0
+
+		UStateTree& StateTree = NewStateTree();
+		UStateTreeEditorData& EditorData = *Cast<UStateTreeEditorData>(StateTree.EditorData);
+		FInstancedPropertyBag& RootPropertyBag = GetRootPropertyBag(EditorData);
+		RootPropertyBag.AddProperty(FName(TEXT("Int")), EPropertyBagPropertyType::Int32);
+		RootPropertyBag.SetValueInt32(FName(TEXT("Int")), 1);
+
+		UStateTreeState& Root = EditorData.AddSubTree(FName(TEXT("Root")));
+		UStateTreeState& State0 = Root.AddChildState(FName(TEXT("State0")));
+
+		// Note: Task does not reselect by default, so will not trigger entry state on reselect
+		auto& Task0 = State0.AddTask<FTestTask_StandNoStateChangeOnReselect>(FName(TEXT("Task0")));
+		Task0.GetNode().TicksToCompletion = 100;
+
+		// Reactivate transtion back to root, by default will not fire due to condition. But should reactivate / re-trigger enter state when activated.
+		FStateTreeTransition& ReactivateRootTransition = State0.AddTransition(EStateTreeTransitionTrigger::OnTick, EStateTreeTransitionType::GotoState, &Root);
+		TStateTreeEditorNode<FStateTreeCompareIntCondition>& TransIntCond = ReactivateRootTransition.AddCondition<FStateTreeCompareIntCondition>(EGenericAICheck::Equal);
+		TransIntCond.GetInstanceData().Right = 0;
+		EditorData.AddPropertyBinding(
+			FPropertyBindingPath(EditorData.GetRootParametersGuid(), TEXT("Int")),
+			FPropertyBindingPath(TransIntCond.ID, TEXT("Left")));
+		ReactivateRootTransition.bReactivateTargetState = true;
+
+		// Transition back to root, will always fire, but will not reactivate / re-trigger enter state
+		FStateTreeTransition& DefaultRootTransition = State0.AddTransition(EStateTreeTransitionTrigger::OnTick, EStateTreeTransitionType::GotoState, &Root);
+
+		FStateTreeCompilerLog Log;
+		FStateTreeCompiler Compiler(Log);
+		const bool bResult = Compiler.Compile(StateTree);
+		AITEST_TRUE("StateTree should get compiled", bResult);
+
+		EStateTreeRunStatus Status = EStateTreeRunStatus::Unset;
+		FStateTreeInstanceData InstanceData;
+		FTestStateTreeExecutionContext Exec(StateTree, StateTree, InstanceData);
+		const bool bInitSucceeded = Exec.IsValid();
+		AITEST_TRUE("StateTree should init", bInitSucceeded);
+
+		const TCHAR* TickStr(TEXT("Tick"));
+		const TCHAR* EnterStateStr(TEXT("EnterState"));
+		const TCHAR* ExitStateStr(TEXT("ExitState"));
+		const TCHAR* StateCompletedStr(TEXT("StateCompleted"));
+
+		FInstancedPropertyBag GlobalParameters = StateTree.GetDefaultParameters();
+
+		// Start and enter state
+		Status = Exec.Start();
+		AITEST_TRUE(TEXT("StateTree Task0 should enter state"), Exec.Expect(Task0.GetName(), EnterStateStr));
+		Exec.LogClear();
+
+		// Transition from State1A -> State1A non-reactivate
+		Status = Exec.Tick(0.1f);
+		AITEST_TRUE(TEXT("StateTree Task0 1st tick count 1"), Exec.Expect(Task0.GetName(), TEXT("CurrentTick=1")));
+		AITEST_FALSE(TEXT("StateTree Task0 should not re-enter state 1st"), Exec.Expect(Task0.GetName(), EnterStateStr));
+		Exec.LogClear();
+
+		// Transition from State1A -> State1A non-reactivate - should still hold after second tick 
+		Status = Exec.Tick(0.1f);
+		AITEST_TRUE(TEXT("StateTree Task0 2nd tick count 2"), Exec.Expect(Task0.GetName(), TEXT("CurrentTick=2")));
+		AITEST_FALSE(TEXT("StateTree Task0 should not re-enter state 2nd"), Exec.Expect(Task0.GetName(), EnterStateStr));
+		Exec.LogClear();
+
+		GlobalParameters.SetValueInt32(FName(TEXT("Int")), 0);
+		InstanceData.GetMutableStorage().SetGlobalParameters(GlobalParameters);
+
+		// Transition from State1A -> State1A reactivate
+		Status = Exec.Tick(0.1f);
+		AITEST_TRUE(TEXT("StateTree Task0 tick before 1st re-enter tick count 3"), Exec.Expect(Task0.GetName(), TEXT("CurrentTick=3")));
+		AITEST_TRUE(TEXT("StateTree Task0 should re-enter state 1st"), Exec.Expect(Task0.GetName(), EnterStateStr));
+		Exec.LogClear();
+
+		// Transition from State1A -> State1A reactivate  should still hold after second tick 
+		Status = Exec.Tick(0.1f);
+		AITEST_TRUE(TEXT("StateTree Task0 tick after 1st re-enter tick count 1"), Exec.Expect(Task0.GetName(), TEXT("CurrentTick=1")));
+		AITEST_TRUE(TEXT("StateTree Task0 should re-enter state 2nd"), Exec.Expect(Task0.GetName(), EnterStateStr));
+		Exec.LogClear();
+
+		GlobalParameters.SetValueInt32(FName(TEXT("Int")), 1);
+		InstanceData.GetMutableStorage().SetGlobalParameters(GlobalParameters);
+
+		// Transition from State1A -> State1A final non-reactivate after condition reset
+		Status = Exec.Tick(0.1f);
+		AITEST_TRUE(TEXT("StateTree Task0 tick after 2nd re-enter tick count 1"), Exec.Expect(Task0.GetName(), TEXT("CurrentTick=1")));
+		AITEST_FALSE(TEXT("StateTree Task0 should not re-enter state final"), Exec.Expect(Task0.GetName(), EnterStateStr));
+		Exec.LogClear();
+
+		Exec.Stop();
+
+		return true;
+	}
+};
+IMPLEMENT_STATE_TREE_INSTANT_TEST(FStateTreeTest_TransitionReactivateTargetState, "System.StateTree.Transition.ReactivateTargetState");
+
 struct FStateTreeTest_TransitionDelayZero : FStateTreeTestBase
 {
 	virtual bool InstantTest() override
