[State Tree] check if the root tree has changed when an ExecContext with write access to instance storage has been created

Runtime validation checks if the root tree has been changed. The old behavior checks if the tree has changed in every Context. This has issues because if the instance data storage hasn't been initialized(written into) yet then changing root tree should be allowed.

Now, we check if the tree has changed once a MinimalExecContext(w/ write access to Instance Storage) has been created, since which changing root tree is not allowed anymore. Because the Instance Data might has been written into based on the previous tree.

diff --git a/Engine/Plugins/Runtime/StateTree/Source/StateTreeModule/Private/Debugger/StateTreeRuntimeValidation.cpp b/Engine/Plugins/Runtime/StateTree/Source/StateTreeModule/Private/Debugger/StateTreeRuntimeValidation.cpp
--- a/Engine/Plugins/Runtime/StateTree/Source/StateTreeModule/Private/Debugger/StateTreeRuntimeValidation.cpp
+++ b/Engine/Plugins/Runtime/StateTree/Source/StateTreeModule/Private/Debugger/StateTreeRuntimeValidation.cpp
@@ -71,13 +71,13 @@ FRuntimeValidationInstanceData::~FRuntimeValidationInstanceData()
 	}
 }
 
-void FRuntimeValidationInstanceData::SetContext(const UObject* InNewOwner, const UStateTree* InNewStateTree)
+void FRuntimeValidationInstanceData::SetContext(const UObject* InNewOwner, const UStateTree* InNewStateTree, bool bInInstanceDataWriteAccessAcquired)
 {
 	TWeakObjectPtr<const UStateTree> NewStateTree = InNewStateTree;
 	TWeakObjectPtr<const UObject> NewOwner = InNewOwner;
 	if (Private::bRuntimeValidationContext)
 	{
-		if (StateTree.IsValid() && StateTree != NewStateTree)
+		if (StateTree.IsValid() && StateTree != NewStateTree && bInstanceDataWriteAccessAcquired)
 		{
 			ensureAlwaysMsgf(false, TEXT("StateTree runtime check failed: The StateTree '%s' is different from the previously set '%s'.\n"
 				"Make sure you initialize FStateTreeExecutionContext with the same value every time.\n"
@@ -100,6 +100,7 @@ void FRuntimeValidationInstanceData::SetContext(const UObject* InNewOwner, const
 
 	StateTree = NewStateTree;
 	Owner = NewOwner;
+	bInstanceDataWriteAccessAcquired |= bInInstanceDataWriteAccessAcquired;
 }
 
 void FRuntimeValidationInstanceData::NodeEnterState(FGuid NodeID, FActiveFrameID FrameID)
@@ -266,17 +267,17 @@ FRuntimeValidationInstanceData* FRuntimeValidation::GetInstanceData() const
 	return RuntimeValidationData;
 }
 
-void FRuntimeValidation::SetContext(const UObject* Owner, const UStateTree* StateTree) const
+void FRuntimeValidation::SetContext(const UObject* Owner, const UStateTree* StateTree, bool bInstanceDataWriteAccessAcquired) const
 {
 	if (RuntimeValidationData)
 	{
-		RuntimeValidationData->SetContext(Owner, StateTree);
+		RuntimeValidationData->SetContext(Owner, StateTree, bInstanceDataWriteAccessAcquired);
 	}
 }
 
 #else
 
-void FRuntimeValidation::SetContext(const UObject*, const UStateTree*) const
+void FRuntimeValidation::SetContext(const UObject*, const UStateTree*, bool) const
 {
 }
 
diff --git a/Engine/Plugins/Runtime/StateTree/Source/StateTreeModule/Private/Debugger/StateTreeRuntimeValidationInstanceData.h b/Engine/Plugins/Runtime/StateTree/Source/StateTreeModule/Private/Debugger/StateTreeRuntimeValidationInstanceData.h
--- a/Engine/Plugins/Runtime/StateTree/Source/StateTreeModule/Private/Debugger/StateTreeRuntimeValidationInstanceData.h
+++ b/Engine/Plugins/Runtime/StateTree/Source/StateTreeModule/Private/Debugger/StateTreeRuntimeValidationInstanceData.h
@@ -20,7 +20,7 @@ public:
 	FRuntimeValidationInstanceData() = default;
 	~FRuntimeValidationInstanceData();
 
-	void SetContext(const UObject* Owner, const UStateTree* StateTree);
+	void SetContext(const UObject* InOwner, const UStateTree* InStateTree, bool bInInstanceDataWriteAccessAcquired);
 
 	void NodeEnterState(FGuid NodeID, FActiveFrameID FrameID);
 	void NodeExitState(FGuid NodeID, FActiveFrameID FrameID);
@@ -46,6 +46,9 @@ private:
 	TArray<FNodeStatePair> NodeStates;
 	TWeakObjectPtr<const UStateTree> StateTree;
 	TWeakObjectPtr<const UObject> Owner;
+
+	/** an ExecContext with write access to the Instance Storage data this is debugging on has been created */
+	bool bInstanceDataWriteAccessAcquired = false;
 };
 
 ENUM_CLASS_FLAGS(FRuntimeValidationInstanceData::EState);
diff --git a/Engine/Plugins/Runtime/StateTree/Source/StateTreeModule/Private/StateTreeExecutionContext.cpp b/Engine/Plugins/Runtime/StateTree/Source/StateTreeModule/Private/StateTreeExecutionContext.cpp
--- a/Engine/Plugins/Runtime/StateTree/Source/StateTreeModule/Private/StateTreeExecutionContext.cpp
+++ b/Engine/Plugins/Runtime/StateTree/Source/StateTreeModule/Private/StateTreeExecutionContext.cpp
@@ -341,19 +341,11 @@ FStateTreeReadOnlyExecutionContext::FStateTreeReadOnlyExecutionContext(TNotNull<
 	, Storage(InStorage)
 {
 	Storage.AcquireReadAccess();
-     
-    const UStateTree* PrevRootStateTree = nullptr;
-	const FStateTreeExecutionState& Exec = Storage.GetExecutionState();
-	if (Exec.ActiveFrames.Num())
-	{
-		PrevRootStateTree = Exec.ActiveFrames[0].StateTree;
-	}
 
-	if (!ensureMsgf(!PrevRootStateTree || PrevRootStateTree == &RootStateTree, 
-		TEXT("Tree Instance Storage will run on a different State Tree '%s' than before '%s'! Resetting the Tree Instance Storage"),
-		*GetFullNameSafe(&RootStateTree), *GetFullNameSafe(PrevRootStateTree)))
-	{		
-		Storage.Reset();
+	if (IsValid())
+	{
+		constexpr bool bWriteAccessAcquired = false;
+		Storage.GetRuntimeValidation().SetContext(&Owner, &RootStateTree, bWriteAccessAcquired);
 	}
 }
 
@@ -814,6 +806,12 @@ FStateTreeMinimalExecutionContext::FStateTreeMinimalExecutionContext(TNotNull<UO
 	: FStateTreeReadOnlyExecutionContext(InOwner, InStateTree, InStorage)
 {
 	Storage.AcquireWriteAccess();
+
+	if (IsValid())
+	{
+		constexpr bool bWriteAccessAcquired = true;
+		Storage.GetRuntimeValidation().SetContext(&Owner, &RootStateTree, bWriteAccessAcquired);
+	}
 }
 
 FStateTreeMinimalExecutionContext::~FStateTreeMinimalExecutionContext()
@@ -1376,7 +1374,8 @@ EStateTreeRunStatus FStateTreeExecutionContext::Start(FStartParameters Parameter
 	// Initialize instance data. No active states yet, so we'll initialize the evals and global tasks.
 	InstanceData.Reset();
 
-	Storage.GetRuntimeValidation().SetContext(&Owner, &RootStateTree);
+	constexpr bool bWriteAccessAcquired = true;
+	Storage.GetRuntimeValidation().SetContext(&Owner, &RootStateTree, bWriteAccessAcquired);
 	Exec.ExecutionExtension = MoveTemp(Parameters.ExecutionExtension);
 	if (Parameters.SharedEventQueue)
 	{
diff --git a/Engine/Plugins/Runtime/StateTree/Source/StateTreeModule/Public/Debugger/StateTreeRuntimeValidation.h b/Engine/Plugins/Runtime/StateTree/Source/StateTreeModule/Public/Debugger/StateTreeRuntimeValidation.h
--- a/Engine/Plugins/Runtime/StateTree/Source/StateTreeModule/Public/Debugger/StateTreeRuntimeValidation.h
+++ b/Engine/Plugins/Runtime/StateTree/Source/StateTreeModule/Public/Debugger/StateTreeRuntimeValidation.h
@@ -35,7 +35,7 @@ public:
 #endif
 
 	/** Set the Owner of the data */
-	void SetContext(const UObject* Owner, const UStateTree* StateTree) const;
+	void SetContext(const UObject* Owner, const UStateTree* StateTree, bool bInstanceDataWriteAccessAcquired) const;
 
 private:
 #if WITH_STATETREE_DEBUG
