[StateTree] Add a callback when a state tree compiles. Similar to UStatetreeEditorDataExtension::HandlePostInternalCompiler, but globally.

diff --git a/Engine/Plugins/Runtime/StateTree/Source/StateTreeEditorModule/Private/StateTreeCompiler.cpp b/Engine/Plugins/Runtime/StateTree/Source/StateTreeEditorModule/Private/StateTreeCompiler.cpp
--- a/Engine/Plugins/Runtime/StateTree/Source/StateTreeEditorModule/Private/StateTreeCompiler.cpp
+++ b/Engine/Plugins/Runtime/StateTree/Source/StateTreeEditorModule/Private/StateTreeCompiler.cpp
@@ -2745,28 +2745,28 @@ bool FStateTreeCompiler::NotifyInternalPost()
 		FStateTreeCompilerLog& Log;
 	};
 
-	if (EditorData->EditorSchema || EditorData->Extensions.Num() > 0)
-	{
-		FPostInternalContextImpl Context = FPostInternalContextImpl(StateTree, EditorData, Log);
+	FPostInternalContextImpl Context = FPostInternalContextImpl(StateTree, EditorData, Log);
+
+	// Notify the module, schema, and extension. Use that order to go from the less specific to the more specific.
+	FStateTreeEditorModule::GetModule().OnPostInternalCompile().Broadcast(Context);
 
-		if (EditorData->EditorSchema)
+	if (EditorData->EditorSchema)
+	{
+		if (!EditorData->EditorSchema->HandlePostInternalCompile(Context))
 		{
-			if (!EditorData->EditorSchema->HandlePostInternalCompile(Context))
-			{
-				Log.Reportf(EMessageSeverity::Error, TEXT("The schema failed compilation. See log for more info."));
-				return false;
-			}
+			Log.Reportf(EMessageSeverity::Error, TEXT("The schema failed compilation. See log for more info."));
+			return false;
 		}
+	}
 
-		for (UStateTreeEditorDataExtension* Extension : EditorData->Extensions)
+	for (UStateTreeEditorDataExtension* Extension : EditorData->Extensions)
+	{
+		if (Extension)
 		{
-			if (Extension)
+			if (!Extension->HandlePostInternalCompile(Context))
 			{
-				if (!Extension->HandlePostInternalCompile(Context))
-				{
-					Log.Reportf(EMessageSeverity::Error, TEXT("An extension failed compilation. See log for more info."));
-					return false;
-				}
+				Log.Reportf(EMessageSeverity::Error, TEXT("An extension failed compilation. See log for more info."));
+				return false;
 			}
 		}
 	}
diff --git a/Engine/Plugins/Runtime/StateTree/Source/StateTreeEditorModule/Public/StateTreeEditorModule.h b/Engine/Plugins/Runtime/StateTree/Source/StateTreeEditorModule/Public/StateTreeEditorModule.h
--- a/Engine/Plugins/Runtime/StateTree/Source/StateTreeEditorModule/Public/StateTreeEditorModule.h
+++ b/Engine/Plugins/Runtime/StateTree/Source/StateTreeEditorModule/Public/StateTreeEditorModule.h
@@ -18,6 +18,11 @@ class UUserDefinedStruct;
 class IStateTreeEditor;
 struct FStateTreeNodeClassCache;
 
+namespace UE::StateTree::Compiler
+{
+	struct FPostInternalContext;
+}
+
 namespace UE::StateTreeDebugger
 {
 	class FRewindDebuggerPlaybackExtension;
@@ -50,13 +55,34 @@ public:
 	/** Sets the Details View with required State Tree Detail Property Handlers */
 	static UE_API void SetDetailPropertyHandlers(IDetailsView& DetailsView);
 
-	virtual TSharedPtr<FExtensibilityManager> GetMenuExtensibilityManager() override { return MenuExtensibilityManager; }
-	virtual TSharedPtr<FExtensibilityManager> GetToolBarExtensibilityManager() override { return ToolBarExtensibilityManager; }
+	virtual TSharedPtr<FExtensibilityManager> GetMenuExtensibilityManager() override
+	{
+		return MenuExtensibilityManager;
+	}
+	virtual TSharedPtr<FExtensibilityManager> GetToolBarExtensibilityManager() override
+	{
+		return ToolBarExtensibilityManager;
+	}
 
 	UE_API TSharedPtr<FStateTreeNodeClassCache> GetNodeClassCache();
 	
 	DECLARE_EVENT_OneParam(FStateTreeEditorModule, FOnRegisterLayoutExtensions, FLayoutExtender&);
-	FOnRegisterLayoutExtensions& OnRegisterLayoutExtensions() { return RegisterLayoutExtensions; }
+	FOnRegisterLayoutExtensions& OnRegisterLayoutExtensions()
+	{
+		return RegisterLayoutExtensions;
+	}
+
+	DECLARE_EVENT_OneParam(FStateTreeEditorModule, FPostInternalCompile, const UE::StateTree::Compiler::FPostInternalContext&);
+	/**
+	 * Handle post internal compilation for all state tree assets.
+	 * The state tree asset compiled successfully.
+	 * @note Use the UStateTreeEditorExtension::HandlePostInternalCompile for controlling a single asset.
+	 * @note Use the UStateTreeEditorSchema::HandlePostInternalCompile for controlling a type of state tree asset.
+	 */
+	FPostInternalCompile& OnPostInternalCompile()
+	{
+		return PostInternalCompile;
+	}
 
 	/** Register the editor data type for a specific schema. */
 	UE_API void RegisterEditorDataClass(TNonNullSubclassOf<const UStateTreeSchema> Schema, TNonNullSubclassOf<const UStateTreeEditorData> EditorData);
@@ -95,6 +121,7 @@ protected:
 	UE_DEPRECATED(5.7, "OnUserDefinedStructReinstancedHandle is not used.")
 	FDelegateHandle OnUserDefinedStructReinstancedHandle;
 	FOnRegisterLayoutExtensions RegisterLayoutExtensions;
+	FPostInternalCompile PostInternalCompile;
 };
 
 #undef UE_API
