[StateTree] Add the "log dependency" option to the editor development menu.

diff --git a/Engine/Plugins/Runtime/StateTree/Source/StateTreeEditorModule/Private/StateTreeEditor.cpp b/Engine/Plugins/Runtime/StateTree/Source/StateTreeEditorModule/Private/StateTreeEditor.cpp
--- a/Engine/Plugins/Runtime/StateTree/Source/StateTreeEditorModule/Private/StateTreeEditor.cpp
+++ b/Engine/Plugins/Runtime/StateTree/Source/StateTreeEditorModule/Private/StateTreeEditor.cpp
@@ -390,6 +390,7 @@ void FillDeveloperMenu(UToolMenu* InMenu)
 	{
 		FToolMenuSection& Section = InMenu->AddSection("FileDeveloperCompilerSettings", LOCTEXT("CompileOptionsHeading", "Compiler Settings"));
 		Section.AddMenuEntry(Commands.LogCompilationResult);
+		Section.AddMenuEntry(Commands.LogDependencies);
 	}
 	{
 		FToolMenuSection& Section = InMenu->AddSection("FileDeveloperSettings", LOCTEXT("DeveloperOptionsHeading", "Settings"));
diff --git a/Engine/Plugins/Runtime/StateTree/Source/StateTreeEditorModule/Private/StateTreeEditorCommands.cpp b/Engine/Plugins/Runtime/StateTree/Source/StateTreeEditorModule/Private/StateTreeEditorCommands.cpp
--- a/Engine/Plugins/Runtime/StateTree/Source/StateTreeEditorModule/Private/StateTreeEditorCommands.cpp
+++ b/Engine/Plugins/Runtime/StateTree/Source/StateTreeEditorModule/Private/StateTreeEditorCommands.cpp
@@ -22,6 +22,7 @@ void FStateTreeEditorCommands::RegisterCommands()
 	UI_COMMAND(SaveOnCompile_SuccessOnly, "On Success Only", "Sets the save-on-compile option to 'Success Only', meaning that your StateTree will be saved whenever they are successfully compiled", EUserInterfaceActionType::RadioButton, FInputChord());
 	UI_COMMAND(SaveOnCompile_Always, "Always", "Sets the save-on-compile option to 'Always', meaning that your StateTree will be saved whenever they are compiled (even if there were errors)", EUserInterfaceActionType::RadioButton, FInputChord());
 	UI_COMMAND(LogCompilationResult, "Log compilation result", "After a StateTree compiles, log the internal content of the StateTree.", EUserInterfaceActionType::ToggleButton, FInputChord());
+	UI_COMMAND(LogDependencies, "Log dependencies", "After a StateTree compiles, log the compilation dependencies of the StateTree.", EUserInterfaceActionType::ToggleButton, FInputChord());
 
 	UI_COMMAND(AddSiblingState, "Add Sibling State", "Add a Sibling State", EUserInterfaceActionType::Button, FInputChord());
 	UI_COMMAND(AddChildState, "Add Child State", "Add a Child State", EUserInterfaceActionType::Button, FInputChord());
diff --git a/Engine/Plugins/Runtime/StateTree/Source/StateTreeEditorModule/Private/StateTreeEditorCommands.h b/Engine/Plugins/Runtime/StateTree/Source/StateTreeEditorModule/Private/StateTreeEditorCommands.h
--- a/Engine/Plugins/Runtime/StateTree/Source/StateTreeEditorModule/Private/StateTreeEditorCommands.h
+++ b/Engine/Plugins/Runtime/StateTree/Source/StateTreeEditorModule/Private/StateTreeEditorCommands.h
@@ -21,6 +21,7 @@ public:
 	TSharedPtr<FUICommandInfo> SaveOnCompile_SuccessOnly;
 	TSharedPtr<FUICommandInfo> SaveOnCompile_Always;
 	TSharedPtr<FUICommandInfo> LogCompilationResult;
+	TSharedPtr<FUICommandInfo> LogDependencies;
 
 	TSharedPtr<FUICommandInfo> AddSiblingState;
 	TSharedPtr<FUICommandInfo> AddChildState;
diff --git a/Engine/Plugins/Runtime/StateTree/Source/StateTreeEditorModule/Private/StateTreeEditorMode.cpp b/Engine/Plugins/Runtime/StateTree/Source/StateTreeEditorModule/Private/StateTreeEditorMode.cpp
--- a/Engine/Plugins/Runtime/StateTree/Source/StateTreeEditorModule/Private/StateTreeEditorMode.cpp
+++ b/Engine/Plugins/Runtime/StateTree/Source/StateTreeEditorModule/Private/StateTreeEditorMode.cpp
@@ -104,7 +104,7 @@ void UStateTreeEditorMode::OnSchemaChanged(const UStateTree& InStateTree)
 
 		if(UStateTreeEditingSubsystem* StateTreeEditingSubsystem = GEditor->GetEditorSubsystem<UStateTreeEditingSubsystem>())
 		{
-			TSharedRef<FStateTreeViewModel> ViewModel = StateTreeEditingSubsystem->FindOrAddViewModel(GetStateTree());		
+			TSharedRef<FStateTreeViewModel> ViewModel = StateTreeEditingSubsystem->FindOrAddViewModel(GetStateTree());
 			ViewModel->NotifyAssetChangedExternally();
 		}
 		
@@ -331,8 +331,29 @@ static bool IsLogCompilationResult()
 	IConsoleVariable* LogResultCVar = GetLogCompilationResultCVar();
 	return LogResultCVar ? LogResultCVar->GetBool() : false;
 }
+
+static IConsoleVariable* GetLogDependenciesCVar()
+{
+	static IConsoleVariable* FoundVariable = IConsoleManager::Get().FindConsoleVariable(TEXT("StateTree.Compiler.LogDependenciesOnCompilation"));
+	return FoundVariable;
 }
 
+static void ToggleLogDependencies()
+{
+	IConsoleVariable* LogResultCVar = GetLogDependenciesCVar();
+	if (ensure(LogResultCVar))
+	{
+		LogResultCVar->Set(!LogResultCVar->GetBool(), ECVF_SetByConsole);
+	}
+}
+
+static bool IsLogDependencies()
+{
+	IConsoleVariable* LogResultCVar = GetLogDependenciesCVar();
+	return LogResultCVar ? LogResultCVar->GetBool() : false;
+}
+} // namespace UE::StateTree::Editor::Internal
+
 void UStateTreeEditorMode::BindToolkitCommands(const TSharedRef<FUICommandList>& ToolkitCommands)
 {
 	FStateTreeEditorCommands::Register();
@@ -372,6 +393,12 @@ void UStateTreeEditorMode::BindToolkitCommands(const TSharedRef<FUICommandList>&
 		FCanExecuteAction(),
 		FIsActionChecked::CreateStatic(&UE::StateTree::Editor::Internal::IsLogCompilationResult)
 	);
+	ToolkitCommands->MapAction(
+		FStateTreeEditorCommands::Get().LogDependencies,
+		FExecuteAction::CreateStatic(&UE::StateTree::Editor::Internal::ToggleLogDependencies),
+		FCanExecuteAction(),
+		FIsActionChecked::CreateStatic(&UE::StateTree::Editor::Internal::IsLogDependencies)
+	);
 }
 
 void UStateTreeEditorMode::OnPropertyBindingChanged(const FPropertyBindingPath& SourcePath, const FPropertyBindingPath& TargetPath)
