[StateTree] Add comments and fix code layout.

#rnx

diff --git a/Engine/Plugins/Runtime/StateTree/Source/StateTreeModule/Public/StateTreeExecutionContext.h b/Engine/Plugins/Runtime/StateTree/Source/StateTreeModule/Public/StateTreeExecutionContext.h
--- a/Engine/Plugins/Runtime/StateTree/Source/StateTreeModule/Public/StateTreeExecutionContext.h
+++ b/Engine/Plugins/Runtime/StateTree/Source/StateTreeModule/Public/StateTreeExecutionContext.h
@@ -1176,7 +1176,11 @@ protected:
 	/** Sets up NextTransition based on the provided parameters and the current execution status. */
 	UE_API void SetupNextTransition(const FTransitionArguments& Args, FStateTreeTransitionResult& OutTransitionResult);
 
-	/** Tick the states transitions. */
+	/**
+	 * Tick (if needed) event, delayed, delegate and, tick state completion transitions.
+	 * Test the completed state transitions.
+	 * @return if the transition is valid and should be entered.
+	 */
 	UE_API bool TriggerTransitions();
 
 	UE_DEPRECATED(5.7, "Use ForceTransition")
