[Backout] - 44893320
Backout part of the CL. It is too consuming (as presented right now) to allow a single transition to override a rule when every transition. It is preferable to modify the schema. The CL doesn't address issues with linked assets. It doesn't work with rules already applied by the schema, and can't remove default schema rules. The CL deprecated unnecessary functions.

Original CL Desc
-----------------------------------------------------------------
[StateTree] Add ability to mark a transition for reselection.

[Backout] - CL45201737
#rnx
Original CL Desc
-----------------------------------------------------------------
[StateTree] Add unit test state tree reactivate

#tests automation

diff --git a/Engine/Plugins/Runtime/StateTree/Source/StateTreeEditorModule/Private/Customizations/StateTreeTransitionDetails.cpp b/Engine/Plugins/Runtime/StateTree/Source/StateTreeEditorModule/Private/Customizations/StateTreeTransitionDetails.cpp
--- a/Engine/Plugins/Runtime/StateTree/Source/StateTreeEditorModule/Private/Customizations/StateTreeTransitionDetails.cpp
+++ b/Engine/Plugins/Runtime/StateTree/Source/StateTreeEditorModule/Private/Customizations/StateTreeTransitionDetails.cpp
@@ -66,7 +66,6 @@ void FStateTreeTransitionDetails::CustomizeHeader(TSharedRef<IPropertyHandle> St
 	DelegateListener = StructProperty->GetChildHandle(GET_MEMBER_NAME_CHECKED(FStateTreeTransition, DelegateListener));
 	StateProperty = StructProperty->GetChildHandle(GET_MEMBER_NAME_CHECKED(FStateTreeTransition, State));
 	DelayTransitionProperty = StructProperty->GetChildHandle(GET_MEMBER_NAME_CHECKED(FStateTreeTransition, bDelayTransition));
-	TransitionReactivateProperty = StructProperty->GetChildHandle(GET_MEMBER_NAME_CHECKED(FStateTreeTransition, bReactivateTargetState));
 	DelayDurationProperty = StructProperty->GetChildHandle(GET_MEMBER_NAME_CHECKED(FStateTreeTransition, DelayDuration));
 	DelayRandomVarianceProperty = StructProperty->GetChildHandle(GET_MEMBER_NAME_CHECKED(FStateTreeTransition, DelayRandomVariance));
 	ConditionsProperty = StructProperty->GetChildHandle(GET_MEMBER_NAME_CHECKED(FStateTreeTransition, Conditions));
@@ -161,7 +160,6 @@ void FStateTreeTransitionDetails::CustomizeChildren(TSharedRef<class IPropertyHa
 	check(RequiredEventProperty);
 	check(DelegateListener);
 	check(DelayTransitionProperty);
-	check(TransitionReactivateProperty);
 	check(DelayDurationProperty);
 	check(DelayRandomVarianceProperty);
 	check(StateProperty);
@@ -226,9 +224,6 @@ void FStateTreeTransitionDetails::CustomizeChildren(TSharedRef<class IPropertyHa
 	StructBuilder.AddProperty(DelayRandomVarianceProperty.ToSharedRef())
 		.Visibility(TAttribute<EVisibility>::Create(TAttribute<EVisibility>::FGetter::CreateLambda(IsNotCompletionTransition)));
 
-	// Transition Reselection
-	StructBuilder.AddProperty(TransitionReactivateProperty.ToSharedRef());
-
 	// Show conditions always expanded, with simplified header (remove item count)
 	IDetailPropertyRow& ConditionsRow = StructBuilder.AddProperty(ConditionsProperty.ToSharedRef());
 	ConditionsRow.ShouldAutoExpand(true);
diff --git a/Engine/Plugins/Runtime/StateTree/Source/StateTreeEditorModule/Private/Customizations/StateTreeTransitionDetails.h b/Engine/Plugins/Runtime/StateTree/Source/StateTreeEditorModule/Private/Customizations/StateTreeTransitionDetails.h
--- a/Engine/Plugins/Runtime/StateTree/Source/StateTreeEditorModule/Private/Customizations/StateTreeTransitionDetails.h
+++ b/Engine/Plugins/Runtime/StateTree/Source/StateTreeEditorModule/Private/Customizations/StateTreeTransitionDetails.h
@@ -57,7 +57,6 @@ private:
 	TSharedPtr<IPropertyHandle> DelegateListener;
 	TSharedPtr<IPropertyHandle> StateProperty;
 	TSharedPtr<IPropertyHandle> DelayTransitionProperty;
-	TSharedPtr<IPropertyHandle> TransitionReactivateProperty;
 	TSharedPtr<IPropertyHandle> DelayDurationProperty;
 	TSharedPtr<IPropertyHandle> DelayRandomVarianceProperty;
 	TSharedPtr<IPropertyHandle> ConditionsProperty;
diff --git a/Engine/Plugins/Runtime/StateTree/Source/StateTreeEditorModule/Private/SStateTreeViewRow.cpp b/Engine/Plugins/Runtime/StateTree/Source/StateTreeEditorModule/Private/SStateTreeViewRow.cpp
--- a/Engine/Plugins/Runtime/StateTree/Source/StateTreeEditorModule/Private/SStateTreeViewRow.cpp
+++ b/Engine/Plugins/Runtime/StateTree/Source/StateTreeEditorModule/Private/SStateTreeViewRow.cpp
@@ -995,7 +995,6 @@ TSharedRef<SWidget> SStateTreeViewRow::MakeTransitionWidgetInternal(const EState
 		FText Tooltip;
 		FStateTreeStateLink Link;
 		FGuid NodeID;
-		bool bSkipComma = false;
 	};
 
 	TArray<FItem> DescItems;
@@ -1023,15 +1022,6 @@ TSharedRef<SWidget> SStateTreeViewRow::MakeTransitionWidgetInternal(const EState
 		if (bMatch)
 		{
 			DescItems.Emplace(Transition.State, Transition.ID);
-
-			if (Transition.bReactivateTargetState)
-			{
-				FItem ReselectionNote;
-				ReselectionNote.Desc = LOCTEXT("ReactivateDesc", "(Re-Activate)");
-				ReselectionNote.Tooltip = LOCTEXT("ReactivateDescToolTip", "The target state and all the following states will be reactivated. Their instance data will be reinstantiated and will receive EnterState.");
-				ReselectionNote.bSkipComma = true;
-				DescItems.Emplace(ReselectionNote);
-			}
 		}
 	}
 
@@ -1100,7 +1090,7 @@ TSharedRef<SWidget> SStateTreeViewRow::MakeTransitionWidgetInternal(const EState
 	{
 		const FItem& Item = DescItems[Index];
 
-		if (Index > 0 && !Item.bSkipComma)
+		if (Index > 0)
 		{
 			TransitionContainer->AddSlot()
 			.AutoWidth()
diff --git a/Engine/Plugins/Runtime/StateTree/Source/StateTreeEditorModule/Private/StateTreeCompiler.cpp b/Engine/Plugins/Runtime/StateTree/Source/StateTreeEditorModule/Private/StateTreeCompiler.cpp
--- a/Engine/Plugins/Runtime/StateTree/Source/StateTreeEditorModule/Private/StateTreeCompiler.cpp
+++ b/Engine/Plugins/Runtime/StateTree/Source/StateTreeEditorModule/Private/StateTreeCompiler.cpp
@@ -1422,7 +1422,6 @@ bool FStateTreeCompiler::CreateStateTransitions()
 				}
 			}
 
-			CompactTransition.bReactivateTargetState = Transition.bReactivateTargetState;
 			CompactTransition.bTransitionEnabled = Transition.bTransitionEnabled;
 
 			if (Transition.bDelayTransition)
diff --git a/Engine/Plugins/Runtime/StateTree/Source/StateTreeEditorModule/Private/StateTreeDescriptionHelpers.cpp b/Engine/Plugins/Runtime/StateTree/Source/StateTreeEditorModule/Private/StateTreeDescriptionHelpers.cpp
--- a/Engine/Plugins/Runtime/StateTree/Source/StateTreeEditorModule/Private/StateTreeDescriptionHelpers.cpp
+++ b/Engine/Plugins/Runtime/StateTree/Source/StateTreeEditorModule/Private/StateTreeDescriptionHelpers.cpp
@@ -247,16 +247,7 @@ FText GetTransitionDesc(const UStateTreeEditorData* EditorData, const FStateTree
 			: LOCTEXT("ActionReturn", "return");
 	}
 	
-	FText ReactivateText = Formatting == EStateTreeNodeFormatting::RichText
-		? LOCTEXT("ReselectionTextRich", "<s> (Re-Activate)</>")
-		: LOCTEXT("ReselectionText", " (Re-Activate)");
-
-	if (!Transition.bReactivateTargetState)
-	{
-		ReactivateText = FText::GetEmpty();
-	}
-
-	return FText::Format(LOCTEXT("TransitionDesc", "{0} {1} {2}{3}"), TriggerText, ActionText, GetStateLinkDesc(EditorData, Transition.State, Formatting, bShowStatePath), ReactivateText);
+	return FText::Format(LOCTEXT("TransitionDesc", "{0} {1} {2}"), TriggerText, ActionText, GetStateLinkDesc(EditorData, Transition.State, Formatting, bShowStatePath));
 }
 
 } // UE::StateTree::Editor
diff --git a/Engine/Plugins/Runtime/StateTree/Source/StateTreeEditorModule/Public/StateTreeState.h b/Engine/Plugins/Runtime/StateTree/Source/StateTreeEditorModule/Public/StateTreeState.h
--- a/Engine/Plugins/Runtime/StateTree/Source/StateTreeEditorModule/Public/StateTreeState.h
+++ b/Engine/Plugins/Runtime/StateTree/Source/StateTreeEditorModule/Public/StateTreeState.h
@@ -156,10 +156,6 @@ PRAGMA_ENABLE_DEPRECATION_WARNINGS
 	UPROPERTY(EditDefaultsOnly, Category = "Transition", meta = (BaseStruct = "/Script/StateTreeModule.StateTreeConditionBase", BaseClass = "/Script/StateTreeModule.StateTreeConditionBlueprintBase"))
 	TArray<FStateTreeEditorNode> Conditions;
 
-	/** The target state and all the following states will be reactivated. Their instance data will be reinstantiated and will receive EnterState. */
-	UPROPERTY(EditDefaultsOnly, Category = "Transition")
-	bool bReactivateTargetState = false;
-
 	/** True if the Transition is Enabled (i.e. not explicitly disabled in the asset). */
 	UPROPERTY(EditDefaultsOnly, Category = "Debug")
 	bool bTransitionEnabled = true;
diff --git a/Engine/Plugins/Runtime/StateTree/Source/StateTreeModule/Private/StateTreeExecutionContext.cpp b/Engine/Plugins/Runtime/StateTree/Source/StateTreeModule/Private/StateTreeExecutionContext.cpp
--- a/Engine/Plugins/Runtime/StateTree/Source/StateTreeModule/Private/StateTreeExecutionContext.cpp
+++ b/Engine/Plugins/Runtime/StateTree/Source/StateTreeModule/Private/StateTreeExecutionContext.cpp
@@ -1442,6 +1442,7 @@ EStateTreeRunStatus FStateTreeExecutionContext::Start(FStartParameters Parameter
 		SelectStateArgs.SourceState = FActiveState(InitFrameID, FActiveStateID(), InitFrameKey.RootState);
 		SelectStateArgs.TargetState = FStateHandleContext(InitFrameKey.StateTree, InitFrameKey.RootState);
 		SelectStateArgs.Behavior = ESelectStateBehavior::StateTransition;
+		SelectStateArgs.SelectionRules = InitFrameKey.StateTree->StateSelectionRules;
 		if (SelectState(SelectStateArgs, SelectStateResult))
 		{
 			check(SelectStateResult->SelectedStates.Num() > 0);
@@ -1997,9 +1998,7 @@ void FStateTreeExecutionContext::RequestTransition(const FStateTreeTransitionReq
 
 		const FActiveStateID CurrentlyProcessedStateID = CurrentlyProcessedFrame->ActiveStates.FindStateID(CurrentlyProcessedState);
 		const FStateHandleContext Target = FStateHandleContext(CurrentlyProcessedFrame->StateTree, Request.TargetState);
-
-		FTransitionArguments TransitionArgs(*CurrentlyProcessedFrame, CurrentlyProcessedStateID, Target, Request.Priority, FStateTreeSharedEvent(), Request.Fallback);
-		const bool bRequested = RequestTransitionInternal(TransitionArgs);
+		const bool bRequested = RequestTransitionInternal(*CurrentlyProcessedFrame, CurrentlyProcessedStateID, Target, FTransitionArguments{ .Priority = Request.Priority, .Fallback = Request.Fallback });
 		if (bRequested)
 		{
 			check(RequestedTransition.IsValid());
@@ -2760,6 +2759,7 @@ bool FStateTreeExecutionContext::ForceTransitionInternal(const TArrayView<const
 	SelectStateArgs.SourceState = SelectStateArgs.ActiveStates[0];
 	SelectStateArgs.TargetState = StateContexts.Last();
 	SelectStateArgs.Behavior = ESelectStateBehavior::Forced;
+	SelectStateArgs.SelectionRules = RootStateTree.StateSelectionRules;
 
 	FSelectStateInternalArguments InternalArgs;
 	InternalArgs.MissingActiveStates = MakeConstArrayView(CurrentActiveStatePath);
@@ -4972,84 +4972,34 @@ bool FStateTreeExecutionContext::RequestTransition(
 	using namespace UE::StateTree::ExecutionContext;
 
 	const FStateHandleContext TargetState = FStateHandleContext(CurrentFrame.StateTree, NextState);
-	const FStateTreeSharedEvent Event = TransitionEvent != nullptr ? *TransitionEvent : FStateTreeSharedEvent();
-
-	FTransitionArguments TransitionArgs(CurrentFrame, FActiveStateID(), TargetState, Priority, Event, Fallback);
-	return RequestTransitionInternal(TransitionArgs);
+	FTransitionArguments TransitionArgs = FTransitionArguments{ .Priority = Priority, .Fallback = Fallback };
+	if (TransitionEvent)
+	{
+		TransitionArgs.TransitionEvent = *TransitionEvent;
+	}
+	
+	return RequestTransitionInternal(CurrentFrame, FActiveStateID(), TargetState, TransitionArgs);
 }
 
-// Deprecated
 bool FStateTreeExecutionContext::RequestTransitionInternal(
 	const FStateTreeExecutionFrame& SourceFrame,
 	const UE::StateTree::FActiveStateID SourceStateID,
 	const UE::StateTree::ExecutionContext::FStateHandleContext TargetState,
-	const EStateTreeTransitionPriority Priority,
-	const FStateTreeSharedEvent TransitionEvent,
-	const EStateTreeSelectionFallback Fallback)
-{
-	FTransitionArguments TransitionArgs(SourceFrame, SourceStateID, TargetState, Priority, TransitionEvent, Fallback);
-	return RequestTransitionInternal(TransitionArgs);
-}
-
-// Deprecated
-void FStateTreeExecutionContext::SetupNextTransition(const FStateTreeExecutionFrame& CurrentFrame, const FStateTreeStateHandle NextState, const EStateTreeTransitionPriority Priority)
-{
-	using namespace UE::StateTree;
-	using namespace UE::StateTree::ExecutionContext;
-
-	const FStateHandleContext TargetState = FStateHandleContext(CurrentFrame.StateTree, NextState);
-	FTransitionArguments TransitionArgs(CurrentFrame, FActiveStateID(), TargetState, Priority);
-	SetupNextTransition(TransitionArgs);
-}
-
-// Deprecated
-void FStateTreeExecutionContext::SetupNextTransition(
-	const FStateTreeExecutionFrame& SourceFrame,
-	const UE::StateTree::FActiveStateID SourceStateID,
-	const UE::StateTree::ExecutionContext::FStateHandleContext TargetState,
-	const EStateTreeTransitionPriority Priority)
-{
-	FTransitionArguments TransitionArgs(SourceFrame, SourceStateID, TargetState, Priority);
-	RequestedTransition = MakeUnique<FRequestTransitionResult>();
-	SetupNextTransition(TransitionArgs, RequestedTransition->Transition);
-}
-
-// Deprecated
-void FStateTreeExecutionContext::SetupNextTransition(
-	const FStateTreeExecutionFrame& SourceFrame,
-	const UE::StateTree::FActiveStateID SourceStateID,
-	const UE::StateTree::ExecutionContext::FStateHandleContext TargetState,
-	const EStateTreeTransitionPriority Priority,
-	FStateTreeTransitionResult& OutTransitionResult)
-{
-	FTransitionArguments TransitionArgs(SourceFrame, SourceStateID, TargetState, Priority);
-	SetupNextTransition(TransitionArgs, OutTransitionResult);
-}
-
-bool FStateTreeExecutionContext::RequestTransitionInternal(const FTransitionArguments& Args)
+	const FTransitionArguments& Args)
 {
 	using namespace UE::StateTree;
 	using namespace UE::StateTree::ExecutionContext;
 	using namespace UE::StateTree::ExecutionContext::Private;
 
-	const FStateTreeExecutionFrame& SourceFrame = Args.SourceFrame;
-	const UE::StateTree::FActiveStateID& SourceStateID = Args.SourceStateID;
-	const UE::StateTree::ExecutionContext::FStateHandleContext& TargetState = Args.TargetState;
-	const EStateTreeTransitionPriority& Priority = Args.Priority;
-	const FStateTreeSharedEvent& TransitionEvent = Args.TransitionEvent;
-	const EStateTreeSelectionFallback& Fallback = Args.Fallback;
-	const EStateTreeStateSelectionRules AdditionalSelectionRules = Args.AdditionalSelectionRules;
-
 	// Skip lower priority transitions.
-	if (RequestedTransition.IsValid() && RequestedTransition->Transition.Priority >= Priority)
+	if (RequestedTransition.IsValid() && RequestedTransition->Transition.Priority >= Args.Priority)
 	{
 		return false;
 	}
 
 	if (TargetState.StateHandle.IsCompletionState())
 	{
-		FTransitionArguments TransitionArgs(SourceFrame, SourceStateID, TargetState, Priority);
-		SetupNextTransition(TransitionArgs);
+		SetupNextTransition(SourceFrame, SourceStateID, TargetState, Args);
 		STATETREE_LOG(Verbose, TEXT("Transition on state '%s' -> state '%s'"),
 			*GetSafeStateName(SourceFrame, SourceFrame.ActiveStates.Last()), *TargetState.StateHandle.Describe());
 		return true;
@@ -5057,8 +5007,7 @@ bool FStateTreeExecutionContext::RequestTransitionInternal(const FTransitionArgu
 	if (!TargetState.StateHandle.IsValid() || TargetState.StateTree == nullptr)
 	{
 		// NotSet is no-operation, but can be used to mask a transition at parent state. Returning unset keeps updating current state.
-		FTransitionArguments TransitionArgs(SourceFrame, SourceStateID, FStateHandleContext(), Priority);
-		SetupNextTransition(TransitionArgs);
+		SetupNextTransition(SourceFrame, SourceStateID, FStateHandleContext(), Args);
 		return true;
 	}
 
@@ -5084,10 +5033,10 @@ bool FStateTreeExecutionContext::RequestTransitionInternal(const FTransitionArgu
 		SelectStateArgs.SourceState = FActiveState(SourceFrame.FrameID, FActiveStateID(), FStateTreeStateHandle());
 	}
 	SelectStateArgs.TargetState = TargetState;
-	SelectStateArgs.TransitionEvent = TransitionEvent;
-	SelectStateArgs.Fallback = Fallback;
+	SelectStateArgs.TransitionEvent = Args.TransitionEvent;
+	SelectStateArgs.Fallback = Args.Fallback;
 	SelectStateArgs.Behavior = ESelectStateBehavior::StateTransition;
-	SelectStateArgs.AdditionalSelectionRules = AdditionalSelectionRules;
+	SelectStateArgs.SelectionRules = SourceFrame.StateTree->StateSelectionRules;
 
 	TSharedRef<FSelectStateResult> StateSelectionResult = MakeShared<FSelectStateResult>();
 	TGuardValue<TSharedPtr<FSelectStateResult>> TemporaryStorageScope(CurrentlyProcessedTemporaryStorage, StateSelectionResult);
@@ -5129,8 +5078,7 @@ bool FStateTreeExecutionContext::RequestTransitionInternal(const FTransitionArgu
 
 		const FStateTreeExecutionState& Exec = GetExecState();
 
-		FTransitionArguments TransitionArgs(SourceFrame, SourceStateID, TargetState, Priority);
-		SetupNextTransition(TransitionArgs);
+		SetupNextTransition(SourceFrame, SourceStateID, TargetState, Args);
 		RequestedTransition->Selection = StateSelectionResult; // It will keep the temporary storage alive until the end of the EnterState
 
 		// Fill NextActiveFrames & NextActiveFrameEvents for backward compatibility
@@ -5212,18 +5160,23 @@ PRAGMA_ENABLE_DEPRECATION_WARNINGS
 	return false;
 }
 
-void FStateTreeExecutionContext::SetupNextTransition(const FTransitionArguments& Args)
+void FStateTreeExecutionContext::SetupNextTransition(
+	const FStateTreeExecutionFrame& SourceFrame,
+	const UE::StateTree::FActiveStateID SourceStateID,
+	const UE::StateTree::ExecutionContext::FStateHandleContext TargetState,
+	const FTransitionArguments& Args)
 {
 	RequestedTransition = MakeUnique<FRequestTransitionResult>();
-	SetupNextTransition(Args, RequestedTransition->Transition);
+	SetupNextTransition(SourceFrame, SourceStateID, TargetState, Args, RequestedTransition->Transition);
 }
 
-void FStateTreeExecutionContext::SetupNextTransition(const FTransitionArguments& Args, FStateTreeTransitionResult& OutTransitionResult)
+void FStateTreeExecutionContext::SetupNextTransition(
+	const FStateTreeExecutionFrame& SourceFrame,
+	const UE::StateTree::FActiveStateID SourceStateID,
+	const UE::StateTree::ExecutionContext::FStateHandleContext TargetState,
+	const FTransitionArguments& Args,
+	FStateTreeTransitionResult& OutTransitionResult)
 {
-	const FStateTreeExecutionFrame& SourceFrame = Args.SourceFrame;
-	const UE::StateTree::FActiveStateID& SourceStateID = Args.SourceStateID;
-	const UE::StateTree::ExecutionContext::FStateHandleContext& TargetState = Args.TargetState;
-	const EStateTreeTransitionPriority& Priority = Args.Priority;
 	const FStateTreeExecutionState& Exec = GetExecState();
 
 	OutTransitionResult.SourceFrameID = SourceFrame.FrameID;
@@ -5232,7 +5185,7 @@ void FStateTreeExecutionContext::SetupNextTransition(const FTransitionArguments&
 	OutTransitionResult.CurrentState = FStateTreeStateHandle::Invalid;
 	OutTransitionResult.CurrentRunStatus = Exec.LastTickStatus;
 	OutTransitionResult.ChangeType = EStateTreeStateChangeType::Changed;
-	OutTransitionResult.Priority = Priority;
+	OutTransitionResult.Priority = Args.Priority;
 
 	// Fill NextActiveFrames & NextActiveFrameEvents for backward compatibility
 	if (UE::StateTree::ExecutionContext::Private::bSetDeprecatedTransitionResultProperties)
@@ -5296,9 +5249,9 @@ bool FStateTreeExecutionContext::TriggerTransitions()
 		const FStateTreeExecutionFrame* CurrentFrame = Exec.FindActiveFrame(Request.SourceFrameID);
 		if (CurrentFrame && CurrentFrame->ActiveStates.Contains(Request.SourceStateID))
 		{
-			UE::StateTree::ExecutionContext::FStateHandleContext TargetState(CurrentFrame->StateTree, Request.TargetState);
-			FTransitionArguments TransitionArgs(*CurrentFrame, Request.SourceStateID, TargetState, Request.Priority, FStateTreeSharedEvent(), Request.Fallback);
-			if (RequestTransitionInternal(TransitionArgs))
+			const UE::StateTree::ExecutionContext::FStateHandleContext TargetState(CurrentFrame->StateTree, Request.TargetState);
+			const FTransitionArguments TransitionArgs = FTransitionArguments{.Priority = Request.Priority, .Fallback = Request.Fallback};
+			if (RequestTransitionInternal(*CurrentFrame, Request.SourceStateID, TargetState, TransitionArgs))
 			{
 				check(RequestedTransition.IsValid());
 				RequestedTransition->Source = FStateTreeTransitionSource(CurrentFrame->StateTree, EStateTreeTransitionSourceType::ExternalRequest, Request.TargetState, Request.Priority);
@@ -5525,10 +5478,12 @@ bool FStateTreeExecutionContext::TriggerTransitions()
 								*GetSafeStateName(CurrentFrame, CurrentFrame.ActiveStates.Last()), *State.Name.ToString(), *GetSafeStateName(CurrentFrame, Transition.State));
 
 							// Trigger Delayed Transition when the delay has passed.
-							UE::StateTree::ExecutionContext::FStateHandleContext TargetState(CurrentFrame.StateTree, Transition.State);
-							EStateTreeStateSelectionRules TransitionRules = Transition.bReactivateTargetState ? EStateTreeStateSelectionRules::ReselectedStateCreatesNewStates : EStateTreeStateSelectionRules::None;
-							FTransitionArguments TransitionArgs(CurrentFrame, Handler.StateID, TargetState, Transition.Priority, DelayedTransition.CapturedEvent, Transition.Fallback, TransitionRules);
-							if (RequestTransitionInternal(TransitionArgs))
+							const UE::StateTree::ExecutionContext::FStateHandleContext TargetState(CurrentFrame.StateTree, Transition.State);
+							FTransitionArguments TransitionArgs;
+							TransitionArgs.Priority = Transition.Priority;
+							TransitionArgs.TransitionEvent = DelayedTransition.CapturedEvent;
+							TransitionArgs.Fallback = Transition.Fallback;
+							if (RequestTransitionInternal(CurrentFrame, Handler.StateID, TargetState, TransitionArgs))
 							{
 								// If the transition was successfully requested with a specific event, consume and remove the event, it's been used.
 								if (DelayedTransition.CapturedEvent.IsValid() && Transition.bConsumeEventOnSelect)
@@ -5616,10 +5571,12 @@ PRAGMA_ENABLE_DEPRECATION_WARNINGS
 						}
 
 						UE_STATETREE_DEBUG_TRANSITION_EVENT(this, FStateTreeTransitionSource(CurrentFrame.StateTree, FStateTreeIndex16(TransitionIndex), Transition.State, Transition.Priority), EStateTreeTraceEventType::OnRequesting);
-						UE::StateTree::ExecutionContext::FStateHandleContext TargetState(CurrentFrame.StateTree, Transition.State);
-						EStateTreeStateSelectionRules TransitionRules = Transition.bReactivateTargetState ? EStateTreeStateSelectionRules::ReselectedStateCreatesNewStates : EStateTreeStateSelectionRules::None;
-						FTransitionArguments TransitionArgs(CurrentFrame, Handler.StateID, TargetState, Transition.Priority, TransitionEvent, Transition.Fallback, TransitionRules);
-						if (RequestTransitionInternal(TransitionArgs))
+						const UE::StateTree::ExecutionContext::FStateHandleContext TargetState(CurrentFrame.StateTree, Transition.State);
+						FTransitionArguments TransitionArgs;
+						TransitionArgs.Priority = Transition.Priority;
+						TransitionArgs.TransitionEvent = TransitionEvent;
+						TransitionArgs.Fallback = Transition.Fallback;
+						if (RequestTransitionInternal(CurrentFrame, Handler.StateID, TargetState, TransitionArgs))
 						{
 							// If the transition was successfully requested with a specific event, consume and remove the event, it's been used.
 							if (TransitionEvent.IsValid() && Transition.bConsumeEventOnSelect)
@@ -5779,10 +5736,11 @@ PRAGMA_ENABLE_DEPRECATION_WARNINGS
 									// No delay allowed on completion conditions.
 									// No priority on completion transitions, use the priority to signal that state is selected.
 									UE_STATETREE_DEBUG_TRANSITION_EVENT(this, FStateTreeTransitionSource(CurrentFrame.StateTree, FStateTreeIndex16(TransitionIndex), Transition.State, Transition.Priority), EStateTreeTraceEventType::OnRequesting);
-									UE::StateTree::ExecutionContext::FStateHandleContext TargetState(CurrentFrame.StateTree, Transition.State);
-									EStateTreeStateSelectionRules TransitionRules = Transition.bReactivateTargetState ? EStateTreeStateSelectionRules::ReselectedStateCreatesNewStates : EStateTreeStateSelectionRules::None;
-									FTransitionArguments TransitionArgs(CurrentFrame, CurrentStateID, TargetState, EStateTreeTransitionPriority::Normal, FStateTreeSharedEvent(), Transition.Fallback, TransitionRules);
-									if (RequestTransitionInternal(TransitionArgs))
+									const UE::StateTree::ExecutionContext::FStateHandleContext TargetState(CurrentFrame.StateTree, Transition.State);
+									FTransitionArguments TransitionArgs;
+									TransitionArgs.Priority = EStateTreeTransitionPriority::Normal;
+									TransitionArgs.Fallback = Transition.Fallback;
+									if (RequestTransitionInternal(CurrentFrame, CurrentStateID, TargetState, TransitionArgs))
 									{
 										RequestedTransition->Source = FStateTreeTransitionSource(CurrentFrame.StateTree, FStateTreeIndex16(TransitionIndex), Transition.State, Transition.Priority);
 										break;
@@ -5821,8 +5779,7 @@ PRAGMA_ENABLE_DEPRECATION_WARNINGS
 				UE::StateTree::ExecutionContext::FStateHandleContext TargetState(RootFrame.StateTree, FStateTreeStateHandle::Root);
 				UE_STATETREE_DEBUG_TRANSITION_EVENT(this, FStateTreeTransitionSource(RootFrame.StateTree, EStateTreeTransitionSourceType::Internal, FStateTreeStateHandle::Root, EStateTreeTransitionPriority::Normal), EStateTreeTraceEventType::OnRequesting);
 
-				FTransitionArguments TransitionArgs(RootFrame, RootFrame.ActiveStates.StateIDs[0], TargetState);
-				if (RequestTransitionInternal(TransitionArgs))
+				if (RequestTransitionInternal(RootFrame, RootFrame.ActiveStates.StateIDs[0], TargetState, FTransitionArguments()))
 				{
 					RequestedTransition->Source = FStateTreeTransitionSource(RootFrame.StateTree, EStateTreeTransitionSourceType::Internal, FStateTreeStateHandle::Root, EStateTreeTransitionPriority::Normal);
 				}
@@ -6099,6 +6056,7 @@ bool FStateTreeExecutionContext::SelectState(const FStateTreeExecutionFrame& Cur
 	SelectStateArgs.TargetState = FStateHandleContext(CurrentFrame.StateTree, NextState);
 	SelectStateArgs.TransitionEvent = TransitionEvent != nullptr ? *TransitionEvent : FStateTreeSharedEvent();
 	SelectStateArgs.Fallback = Fallback;
+	SelectStateArgs.SelectionRules = CurrentFrame.StateTree->StateSelectionRules;
 
 	TSharedRef<FSelectStateResult> SelectStateResult = MakeShared<FSelectStateResult>();
 	TGuardValue<TSharedPtr<FSelectStateResult>> ExecutionFrameHolderScope(CurrentlyProcessedTemporaryStorage, SelectStateResult);
@@ -6678,22 +6636,20 @@ bool FStateTreeExecutionContext::SelectStateInternal(
 		StateSelectionEvents.Add(FStateTreeSharedEvent());
 	}
 
-	EStateTreeStateSelectionRules StateSelectionRules = NextFrame.StateTree->StateSelectionRules | Args.AdditionalSelectionRules;
-
 	bool bShouldCreateNewState = true;
 	const bool bIsNextTargetStateInActiveStates = !bNewFrameCreated
 		&& InternalArgs.MissingActiveStates.Num() > 0
 		&& InternalArgs.MissingActiveStates[0].GetStateHandle() == NextStateHandle
 		&& InternalArgs.MissingActiveStates[0].GetFrameID() == NextFrame.FrameID;
-	if (EnumHasAnyFlags(StateSelectionRules, EStateTreeStateSelectionRules::ReselectedStateCreatesNewStates | EStateTreeStateSelectionRules::CompletedTransitionStatesCreateNewStates))
+	if (EnumHasAnyFlags(Args.SelectionRules, EStateTreeStateSelectionRules::ReselectedStateCreatesNewStates | EStateTreeStateSelectionRules::CompletedTransitionStatesCreateNewStates))
 	{
 		bShouldCreateNewState = false;
-		if (EnumHasAllFlags(StateSelectionRules, EStateTreeStateSelectionRules::CompletedTransitionStatesCreateNewStates))
+		if (EnumHasAllFlags(Args.SelectionRules, EStateTreeStateSelectionRules::CompletedTransitionStatesCreateNewStates))
 		{
 			// Request a new state if it doesn't match the previous states or if the state is completed.
 			bShouldCreateNewState = !bIsNextTargetStateInActiveStates || NextFrame.ActiveTasksStatus.GetStatus(NextState).IsCompleted();
 		}
-		if (EnumHasAnyFlags(StateSelectionRules, EStateTreeStateSelectionRules::ReselectedStateCreatesNewStates))
+		if (EnumHasAnyFlags(Args.SelectionRules, EStateTreeStateSelectionRules::ReselectedStateCreatesNewStates))
 		{
 			bShouldCreateNewState = bShouldCreateNewState || bIsTargetState || !bIsNextTargetStateInActiveStates;
 		}
@@ -6887,21 +6843,19 @@ TOptional<bool> FStateTreeExecutionContext::SelectStateFromSourceInternal(
 		bool bContinueWithStateSelection = NextStateHandle != MissingActiveState.GetStateHandle();
 		bool bCompletedState = false;
 
-		EStateTreeStateSelectionRules StateSelectionRules = NextFrame.StateTree->StateSelectionRules | Args.AdditionalSelectionRules;
-
 		if (!bContinueWithStateSelection)
 		{
 			if (InternalArgs.MissingActiveStates.Num() == 1)
 			{
 				// MissingActiveState is the transition source.
-				if (EnumHasAllFlags(StateSelectionRules, EStateTreeStateSelectionRules::CompletedTransitionStatesCreateNewStates))
+				if (EnumHasAllFlags(Args.SelectionRules, EStateTreeStateSelectionRules::CompletedTransitionStatesCreateNewStates))
 				{
 					bContinueWithStateSelection = NextFrame.ActiveTasksStatus.GetStatus(NextState).IsCompleted();
 				}
 			}
 			else
 			{
-				if (EnumHasAllFlags(StateSelectionRules, EStateTreeStateSelectionRules::CompletedTransitionStatesCreateNewStates))
+				if (EnumHasAllFlags(Args.SelectionRules, EStateTreeStateSelectionRules::CompletedTransitionStatesCreateNewStates))
 				{
 					bCompletedState = NextFrame.ActiveTasksStatus.GetStatus(NextState).IsCompleted();
 				}
@@ -6930,7 +6884,7 @@ TOptional<bool> FStateTreeExecutionContext::SelectStateFromSourceInternal(
 
 			if (bSelectStateInternalSucceeded && bCompletedState)
 			{
-				check(EnumHasAllFlags(StateSelectionRules, EStateTreeStateSelectionRules::CompletedStateBeforeTransitionSourceFailsTransition));
+				check(EnumHasAllFlags(Args.SelectionRules, EStateTreeStateSelectionRules::CompletedStateBeforeTransitionSourceFailsTransition));
 				// Before the source, there's a completed state.
 				//The completed state should fail the transition(because of EStateTreeStateSelectionRules::CompletedStateBeforeTargetFailsTransition).
 				//The transition can contain a state with "follow transition" that would remove this state from the transition.
@@ -6971,10 +6925,10 @@ namespace UE::StateTree::ExecutionContext::Private
 	}
 
 	FStateTreeExecutionFrame* SelectedFrameLinkedFrame(FStateTreeExecutionState& Exec,
-		bool bShouldCreateNewState,
+		const bool bShouldCreateNewState,
 		TArrayView<const UE::StateTree::FActiveState> MatchingActiveStates,
 		const FExecutionFrameKey& LinkStateFrameKey,
-		EStateTreeStateSelectionRules AdditionalSelectionRules
+		const EStateTreeStateSelectionRules StateSelectionRules
 	)
 	{
 		FStateTreeExecutionFrame* Result = nullptr;
@@ -6988,7 +6942,6 @@ namespace UE::StateTree::ExecutionContext::Private
 				check(Result);
 				if (Result->HasRoot(LinkStateFrameKey))
 				{
-					EStateTreeStateSelectionRules StateSelectionRules = Result->StateTree->GetStateSelectionRules() | AdditionalSelectionRules;
 					if (EnumHasAllFlags(StateSelectionRules, EStateTreeStateSelectionRules::CompletedTransitionStatesCreateNewStates)
 						&& Result->ActiveTasksStatus.GetStatus(Result->StateTree).IsCompleted()
 						)
@@ -7072,7 +7025,7 @@ bool FStateTreeExecutionContext::SelectStateInternal_Linked(
 	}
 
 	// Do we have an existing frame.
-	FStateTreeExecutionFrame* SelectedFrame = SelectedFrameLinkedFrame(Exec, bShouldCreateNewState, InternalArgs.MissingActiveStates, LinkStateFrameKey, Args.AdditionalSelectionRules);
+	FStateTreeExecutionFrame* SelectedFrame = SelectedFrameLinkedFrame(Exec, bShouldCreateNewState, InternalArgs.MissingActiveStates, LinkStateFrameKey, Args.SelectionRules);
 	if (SelectedFrame && !SelectedFrame->HasRoot(LinkStateFrameKey))
 	{
 		if (Args.Behavior == ESelectStateBehavior::Forced)
@@ -7120,12 +7073,12 @@ bool FStateTreeExecutionContext::SelectStateInternal_Linked(
 	OutSelectionResult->SelectedFrames.Add(SelectedFrame->FrameID);
 
 	// Select the root state of the new frame.
-	FStateHandleContext RootState = FStateHandleContext(LinkStateFrameKey.StateTree, LinkStateFrameKey.RootState);
-	FSelectStateInternalArguments NewInternalArgs = FSelectStateInternalArguments{
+	const FStateHandleContext RootState = FStateHandleContext(LinkStateFrameKey.StateTree, LinkStateFrameKey.RootState);
+	const FSelectStateInternalArguments NewInternalArgs = FSelectStateInternalArguments{
 		.MissingActiveStates = bIsNewFrame ? TArrayView<const FActiveState>() : InternalArgs.MissingActiveStates,
 		.MissingSourceFrameID = SelectedFrame->FrameID,
 		.MissingSourceStates = TArrayView<const FActiveState>(),
-		.MissingStatesToReachTarget = bHasMissingState ? InternalArgs.MissingStatesToReachTarget : MakeArrayView(&RootState, 1)
+		.MissingStatesToReachTarget = bHasMissingState ? InternalArgs.MissingStatesToReachTarget : MakeConstArrayView(&RootState, 1)
 	};
 	const bool bSucceededToSelectState = SelectStateInternal(Args, NewInternalArgs, OutSelectionResult);
 
@@ -7229,7 +7182,9 @@ bool FStateTreeExecutionContext::SelectStateInternal_LinkedAsset(
 	}
 
 	// Do we have an existing frame.
-	FStateTreeExecutionFrame* SelectedFrame = SelectedFrameLinkedFrame(Exec, bShouldCreateNewState, InternalArgs.MissingActiveStates, LinkStateFrameKey, Args.AdditionalSelectionRules);
+	// Do not use the transition override selection rules. A transition outside the frame should not impact the current frame.
+	const EStateTreeStateSelectionRules StateSelectionRules = LinkStateFrameKey.StateTree->GetStateSelectionRules();
+	FStateTreeExecutionFrame* SelectedFrame = SelectedFrameLinkedFrame(Exec, bShouldCreateNewState, InternalArgs.MissingActiveStates, LinkStateFrameKey, StateSelectionRules);
 	if (SelectedFrame && !SelectedFrame->HasRoot(LinkStateFrameKey))
 	{
 		if (Args.Behavior == ESelectStateBehavior::Forced)
@@ -7308,14 +7263,16 @@ bool FStateTreeExecutionContext::SelectStateInternal_LinkedAsset(
 	}
 
 	// Select the root state of the new frame.
-	FStateHandleContext RootState = FStateHandleContext(LinkStateFrameKey.StateTree, LinkStateFrameKey.RootState);
-	FSelectStateInternalArguments NewInternalArgs = FSelectStateInternalArguments{
+	const FStateHandleContext RootState = FStateHandleContext(LinkStateFrameKey.StateTree, LinkStateFrameKey.RootState);
+	FSelectStateArguments NewArgs = Args;
+	NewArgs.SelectionRules = StateSelectionRules;
+	const FSelectStateInternalArguments NewInternalArgs = FSelectStateInternalArguments{
 		.MissingActiveStates = bIsNewFrame ? TArrayView<const FActiveState>() : InternalArgs.MissingActiveStates,
 		.MissingSourceFrameID = SelectedFrame->FrameID,
 		.MissingSourceStates = TArrayView<const FActiveState>(),
-		.MissingStatesToReachTarget = bHasMissingState ? InternalArgs.MissingStatesToReachTarget : MakeArrayView(&RootState, 1)
+		.MissingStatesToReachTarget = bHasMissingState ? InternalArgs.MissingStatesToReachTarget : MakeConstArrayView(&RootState, 1)
 	};
-	const bool bSucceededToSelectState = SelectStateInternal(Args, NewInternalArgs, OutSelectionResult);
+	const bool bSucceededToSelectState = SelectStateInternal(NewArgs, NewInternalArgs, OutSelectionResult);
 
 	if (!bSucceededToSelectState)
 	{
@@ -7673,6 +7630,7 @@ bool FStateTreeExecutionContext::SelectStateInternal_TryFollowTransitions(
 				SelectStateArgs.TargetState = FStateHandleContext(NextStateTree, Transition.State);
 				SelectStateArgs.TransitionEvent = SelectedStateTransitionEvent;
 				SelectStateArgs.Fallback = Transition.Fallback;
+				SelectStateArgs.SelectionRules = NextStateTree->StateSelectionRules;
 
 				OutSelectionResult->SelectedStates.Reset();
 				OutSelectionResult->SelectedFrames.Reset();
diff --git a/Engine/Plugins/Runtime/StateTree/Source/StateTreeModule/Public/StateTreeExecutionContext.h b/Engine/Plugins/Runtime/StateTree/Source/StateTreeModule/Public/StateTreeExecutionContext.h
--- a/Engine/Plugins/Runtime/StateTree/Source/StateTreeModule/Public/StateTreeExecutionContext.h
+++ b/Engine/Plugins/Runtime/StateTree/Source/StateTreeModule/Public/StateTreeExecutionContext.h
@@ -1108,73 +1108,36 @@ protected:
 		const FStateTreeSharedEvent* TransitionEvent = nullptr,
 		const EStateTreeSelectionFallback Fallback = EStateTreeSelectionFallback::None);
 
-	/** Requests transition to a specified state with specified priority. */
-	UE_DEPRECATED(5.7, "Use the RequestTransitionInternal with TransitionArguments")
-	UE_API bool RequestTransitionInternal(
-		const FStateTreeExecutionFrame& SourceFrame,
-		UE::StateTree::FActiveStateID SourceStateID,
-		UE::StateTree::ExecutionContext::FStateHandleContext TargetState,
-		EStateTreeTransitionPriority Priority = EStateTreeTransitionPriority::Normal,
-		FStateTreeSharedEvent TransitionEvent = FStateTreeSharedEvent(),
-		EStateTreeSelectionFallback Fallback = EStateTreeSelectionFallback::None);
 
 	UE_DEPRECATED(5.7, "Use the SetupNextTransition with TransitionArguments")
 	UE_API void SetupNextTransition(const FStateTreeExecutionFrame& CurrentFrame, const FStateTreeStateHandle NextState, const EStateTreeTransitionPriority Priority);
-
-	/** Sets up NextTransition based on the provided parameters and the current execution status. */
-	UE_DEPRECATED(5.7, "Use the SetupNextTransition with TransitionArguments")
-	UE_API void SetupNextTransition(
-		const FStateTreeExecutionFrame& SourceFrame,
-		UE::StateTree::FActiveStateID SourceStateID,
-		UE::StateTree::ExecutionContext::FStateHandleContext TargetState,
-		EStateTreeTransitionPriority Priority);
-
-	/** Sets up NextTransition based on the provided parameters and the current execution status. */
-	UE_DEPRECATED(5.7, "Use the SetupNextTransition with TransitionArguments")
-	UE_API void SetupNextTransition(
-		const FStateTreeExecutionFrame& SourceFrame,
-		UE::StateTree::FActiveStateID SourceStateID,
-		UE::StateTree::ExecutionContext::FStateHandleContext TargetState,
-		EStateTreeTransitionPriority Priority,
-		FStateTreeTransitionResult& OutTransitionResult);
 	
+	/** Optional arguments for a transition request. */
 	struct FTransitionArguments
 	{
-		FTransitionArguments(const FStateTreeExecutionFrame& InSourceFrame
-			, UE::StateTree::FActiveStateID InSourceStateID
-			, UE::StateTree::ExecutionContext::FStateHandleContext InTargetState
-			, EStateTreeTransitionPriority InPriority = EStateTreeTransitionPriority::Normal
-			, FStateTreeSharedEvent InTransitionEvent = FStateTreeSharedEvent()
-			, EStateTreeSelectionFallback InFallback = EStateTreeSelectionFallback::None
-			, EStateTreeStateSelectionRules InAdditionalSelectionRules = EStateTreeStateSelectionRules::None)
-			:
-			SourceFrame(InSourceFrame)
-			, SourceStateID(InSourceStateID)
-			, TargetState(InTargetState)
-			, Priority(InPriority)
-			, TransitionEvent(InTransitionEvent)
-			, Fallback(InFallback)
-			, AdditionalSelectionRules(InAdditionalSelectionRules)
-		{
-		};
-
-		const FStateTreeExecutionFrame& SourceFrame;
-		UE::StateTree::FActiveStateID SourceStateID;
-		UE::StateTree::ExecutionContext::FStateHandleContext TargetState;
 		EStateTreeTransitionPriority Priority = EStateTreeTransitionPriority::Normal;
-		FStateTreeSharedEvent TransitionEvent = FStateTreeSharedEvent();
+		FStateTreeSharedEvent TransitionEvent;
 		EStateTreeSelectionFallback Fallback = EStateTreeSelectionFallback::None;
-		EStateTreeStateSelectionRules AdditionalSelectionRules = EStateTreeStateSelectionRules::None;
 	};
 
-	/** Requests transition to a specified state with specified priority. */
-	UE_API bool RequestTransitionInternal(const FTransitionArguments& Args);
+	/** Requests transition to a specified state. */
+	UE_API bool RequestTransitionInternal(const FStateTreeExecutionFrame& SourceFrame,
+		UE::StateTree::FActiveStateID SourceStateID,
+		UE::StateTree::ExecutionContext::FStateHandleContext TargetState,
+		const FTransitionArguments& Args);
 	
 	/** Sets up NextTransition based on the provided parameters and the current execution status. */
-	UE_API void SetupNextTransition(const FTransitionArguments& Args);
+	UE_API void SetupNextTransition(const FStateTreeExecutionFrame& SourceFrame,
+		UE::StateTree::FActiveStateID SourceStateID,
+		UE::StateTree::ExecutionContext::FStateHandleContext TargetState,
+		const FTransitionArguments& Args);
 
 	/** Sets up NextTransition based on the provided parameters and the current execution status. */
-	UE_API void SetupNextTransition(const FTransitionArguments& Args, FStateTreeTransitionResult& OutTransitionResult);
+	UE_API void SetupNextTransition(const FStateTreeExecutionFrame& SourceFrame,
+		UE::StateTree::FActiveStateID SourceStateID,
+		UE::StateTree::ExecutionContext::FStateHandleContext TargetState,
+		const FTransitionArguments& Args,
+		FStateTreeTransitionResult& OutTransitionResult);
 
 	/**
 	 * Tick (if needed) event, delayed, delegate and, tick state completion transitions.
@@ -1255,8 +1218,8 @@ PRAGMA_ENABLE_DEPRECATION_WARNINGS
 		EStateTreeSelectionFallback Fallback = EStateTreeSelectionFallback::None;
 		/** The type of selection. */
 		ESelectStateBehavior Behavior = ESelectStateBehavior::StateTransition;
-		/** True implies we should always create a new state. */
-		EStateTreeStateSelectionRules AdditionalSelectionRules = EStateTreeStateSelectionRules::None;
+		/** Rule to Selection state rules. */
+		EStateTreeStateSelectionRules SelectionRules = EStateTreeStateSelectionRules::None;
 	};
 
 	/**
diff --git a/Engine/Plugins/Runtime/StateTree/Source/StateTreeTestSuite/Private/StateTreeTestTypes.h b/Engine/Plugins/Runtime/StateTree/Source/StateTreeTestSuite/Private/StateTreeTestTypes.h
--- a/Engine/Plugins/Runtime/StateTree/Source/StateTreeTestSuite/Private/StateTreeTestTypes.h
+++ b/Engine/Plugins/Runtime/StateTree/Source/StateTreeTestSuite/Private/StateTreeTestTypes.h
@@ -600,7 +600,6 @@ struct FTestTask_Stand : public FStateTreeTaskBase
 		FInstanceDataType& InstanceData = Context.GetInstanceData(*this);
 		
 		InstanceData.CurrentTick++;
-		TestLog.Log(Name, FString::Printf(TEXT("CurrentTick=%u"), InstanceData.CurrentTick));
 		
 		return (InstanceData.CurrentTick >= TicksToCompletion) ? TickCompletionResult : EStateTreeRunStatus::Running;
 	};
@@ -640,22 +639,6 @@ struct FTestTask_StandNoTick : public FTestTask_Stand
 	}
 };
 
-USTRUCT()
-struct FTestTask_StandNoStateChangeOnReselect : public FTestTask_Stand
-{
-	GENERATED_BODY()
-
-	FTestTask_StandNoStateChangeOnReselect()
-	{
-		bShouldStateChangeOnReselect = false;
-	}
-
-	FTestTask_StandNoStateChangeOnReselect(const FName InName)
-		: FTestTask_Stand(InName)
-	{
-		bShouldStateChangeOnReselect = false;
-	}
-};
 
 USTRUCT()
 struct FTestTask_IntegersOutput_InstanceData
diff --git a/Engine/Plugins/Runtime/StateTree/Source/StateTreeTestSuite/Private/StateTreeTransitionTest.cpp b/Engine/Plugins/Runtime/StateTree/Source/StateTreeTestSuite/Private/StateTreeTransitionTest.cpp
--- a/Engine/Plugins/Runtime/StateTree/Source/StateTreeTestSuite/Private/StateTreeTransitionTest.cpp
+++ b/Engine/Plugins/Runtime/StateTree/Source/StateTreeTestSuite/Private/StateTreeTransitionTest.cpp
@@ -454,104 +454,6 @@ struct FStateTreeTest_TransitionDelay : FStateTreeTestBase
 };
 IMPLEMENT_STATE_TREE_INSTANT_TEST(FStateTreeTest_TransitionDelay, "System.StateTree.Transition.Delay");
 
-struct FStateTreeTest_TransitionReactivateTargetState : FStateTreeTestBase
-{
-	virtual bool InstantTest() override
-	{
-		//Root
-		// State0
-
-		UStateTree& StateTree = NewStateTree();
-		UStateTreeEditorData& EditorData = *Cast<UStateTreeEditorData>(StateTree.EditorData);
-		FInstancedPropertyBag& RootPropertyBag = GetRootPropertyBag(EditorData);
-		RootPropertyBag.AddProperty(FName(TEXT("Int")), EPropertyBagPropertyType::Int32);
-		RootPropertyBag.SetValueInt32(FName(TEXT("Int")), 1);
-
-		UStateTreeState& Root = EditorData.AddSubTree(FName(TEXT("Root")));
-		UStateTreeState& State0 = Root.AddChildState(FName(TEXT("State0")));
-
-		// Note: Task does not reselect by default, so will not trigger entry state on reselect
-		auto& Task0 = State0.AddTask<FTestTask_StandNoStateChangeOnReselect>(FName(TEXT("Task0")));
-		Task0.GetNode().TicksToCompletion = 100;
-
-		// Reactivate transtion back to root, by default will not fire due to condition. But should reactivate / re-trigger enter state when activated.
-		FStateTreeTransition& ReactivateRootTransition = State0.AddTransition(EStateTreeTransitionTrigger::OnTick, EStateTreeTransitionType::GotoState, &Root);
-		TStateTreeEditorNode<FStateTreeCompareIntCondition>& TransIntCond = ReactivateRootTransition.AddCondition<FStateTreeCompareIntCondition>(EGenericAICheck::Equal);
-		TransIntCond.GetInstanceData().Right = 0;
-		EditorData.AddPropertyBinding(
-			FPropertyBindingPath(EditorData.GetRootParametersGuid(), TEXT("Int")),
-			FPropertyBindingPath(TransIntCond.ID, TEXT("Left")));
-		ReactivateRootTransition.bReactivateTargetState = true;
-
-		// Transition back to root, will always fire, but will not reactivate / re-trigger enter state
-		FStateTreeTransition& DefaultRootTransition = State0.AddTransition(EStateTreeTransitionTrigger::OnTick, EStateTreeTransitionType::GotoState, &Root);
-
-		FStateTreeCompilerLog Log;
-		FStateTreeCompiler Compiler(Log);
-		const bool bResult = Compiler.Compile(StateTree);
-		AITEST_TRUE("StateTree should get compiled", bResult);
-
-		EStateTreeRunStatus Status = EStateTreeRunStatus::Unset;
-		FStateTreeInstanceData InstanceData;
-		FTestStateTreeExecutionContext Exec(StateTree, StateTree, InstanceData);
-		const bool bInitSucceeded = Exec.IsValid();
-		AITEST_TRUE("StateTree should init", bInitSucceeded);
-
-		const TCHAR* TickStr(TEXT("Tick"));
-		const TCHAR* EnterStateStr(TEXT("EnterState"));
-		const TCHAR* ExitStateStr(TEXT("ExitState"));
-		const TCHAR* StateCompletedStr(TEXT("StateCompleted"));
-
-		FInstancedPropertyBag GlobalParameters = StateTree.GetDefaultParameters();
-
-		// Start and enter state
-		Status = Exec.Start();
-		AITEST_TRUE(TEXT("StateTree Task0 should enter state"), Exec.Expect(Task0.GetName(), EnterStateStr));
-		Exec.LogClear();
-
-		// Transition from State1A -> State1A non-reactivate
-		Status = Exec.Tick(0.1f);
-		AITEST_TRUE(TEXT("StateTree Task0 1st tick count 1"), Exec.Expect(Task0.GetName(), TEXT("CurrentTick=1")));
-		AITEST_FALSE(TEXT("StateTree Task0 should not re-enter state 1st"), Exec.Expect(Task0.GetName(), EnterStateStr));
-		Exec.LogClear();
-
-		// Transition from State1A -> State1A non-reactivate - should still hold after second tick 
-		Status = Exec.Tick(0.1f);
-		AITEST_TRUE(TEXT("StateTree Task0 2nd tick count 2"), Exec.Expect(Task0.GetName(), TEXT("CurrentTick=2")));
-		AITEST_FALSE(TEXT("StateTree Task0 should not re-enter state 2nd"), Exec.Expect(Task0.GetName(), EnterStateStr));
-		Exec.LogClear();
-
-		GlobalParameters.SetValueInt32(FName(TEXT("Int")), 0);
-		InstanceData.GetMutableStorage().SetGlobalParameters(GlobalParameters);
-
-		// Transition from State1A -> State1A reactivate
-		Status = Exec.Tick(0.1f);
-		AITEST_TRUE(TEXT("StateTree Task0 tick before 1st re-enter tick count 3"), Exec.Expect(Task0.GetName(), TEXT("CurrentTick=3")));
-		AITEST_TRUE(TEXT("StateTree Task0 should re-enter state 1st"), Exec.Expect(Task0.GetName(), EnterStateStr));
-		Exec.LogClear();
-
-		// Transition from State1A -> State1A reactivate  should still hold after second tick 
-		Status = Exec.Tick(0.1f);
-		AITEST_TRUE(TEXT("StateTree Task0 tick after 1st re-enter tick count 1"), Exec.Expect(Task0.GetName(), TEXT("CurrentTick=1")));
-		AITEST_TRUE(TEXT("StateTree Task0 should re-enter state 2nd"), Exec.Expect(Task0.GetName(), EnterStateStr));
-		Exec.LogClear();
-
-		GlobalParameters.SetValueInt32(FName(TEXT("Int")), 1);
-		InstanceData.GetMutableStorage().SetGlobalParameters(GlobalParameters);
-
-		// Transition from State1A -> State1A final non-reactivate after condition reset
-		Status = Exec.Tick(0.1f);
-		AITEST_TRUE(TEXT("StateTree Task0 tick after 2nd re-enter tick count 1"), Exec.Expect(Task0.GetName(), TEXT("CurrentTick=1")));
-		AITEST_FALSE(TEXT("StateTree Task0 should not re-enter state final"), Exec.Expect(Task0.GetName(), EnterStateStr));
-		Exec.LogClear();
-
-		Exec.Stop();
-
-		return true;
-	}
-};
-IMPLEMENT_STATE_TREE_INSTANT_TEST(FStateTreeTest_TransitionReactivateTargetState, "System.StateTree.Transition.ReactivateTargetState");
-
 struct FStateTreeTest_TransitionDelayZero : FStateTreeTestBase
 {
 	virtual bool InstantTest() override
