[StateTree] Add compilation options to use the evaluation scope for condition and utility consideration.

diff --git a/Engine/Plugins/Runtime/StateTree/Source/StateTreeEditorModule/Private/StateTreeCompiler.cpp b/Engine/Plugins/Runtime/StateTree/Source/StateTreeEditorModule/Private/StateTreeCompiler.cpp
--- a/Engine/Plugins/Runtime/StateTree/Source/StateTreeEditorModule/Private/StateTreeCompiler.cpp
+++ b/Engine/Plugins/Runtime/StateTree/Source/StateTreeEditorModule/Private/StateTreeCompiler.cpp
@@ -66,13 +66,28 @@ namespace UE::StateTree::Compiler
 		TEXT("Enable binding from delegate dispatchers that are in the state tree parameters.")
 	);
 
-	FAutoConsoleVariable CVarLogEnablePropertyFunctionWithEvaluationScopeInstanceData(
+	bool bEnablePropertyFunctionWithEvaluationScopeInstanceData = true;
+	FAutoConsoleVariableRef CVarEnablePropertyFunctionWithEvaluationScopeInstanceData(
 		TEXT("StateTree.Compiler.EnablePropertyFunctionWithEvaluationScopeInstanceData"),
-		true,
+		bEnablePropertyFunctionWithEvaluationScopeInstanceData,
 		TEXT("Use EvaluationScope data for property functions instead of SharedInstance data.\n")
 		TEXT("SharedInstance is the previous behavior that is deprecated.")
 	);
 
+	bool bEnableConditionWithEvaluationScopeInstanceData = false;
+	FAutoConsoleVariableRef CVarEnableConditionWithEvaluationScopeInstanceData(
+		TEXT("StateTree.Compiler.EnableConditionWithEvaluationScopeInstanceData"),
+		bEnableConditionWithEvaluationScopeInstanceData,
+		TEXT("Use EvaluationScope data for conditions instead of SharedInstance data.")
+	);
+
+	bool bEnableUtilityConsiderationWithEvaluationScopeInstanceData = false;
+	FAutoConsoleVariableRef CVarEnableUtilityConsiderationWithEvaluationScopeInstanceData(
+		TEXT("StateTree.Compiler.EnableUtilityConsiderationWithEvaluationScopeInstanceData"),
+		bEnableUtilityConsiderationWithEvaluationScopeInstanceData,
+		TEXT("Use EvaluationScope data for utility considerations instead of SharedInstance data.")
+	);
+
 	FAutoConsoleVariable CVarLogCompiledStateTree(
 		TEXT("StateTree.Compiler.LogResultOnCompilationCompleted"),
 		false,
@@ -1263,7 +1278,9 @@ bool FStateTreeCompiler::CreateStateTransitions()
 		}
 		CompactState.EnterConditionsNum = uint8(EnterConditionsNum);
 
-		if (!CreateBindingsForNodes(SourceState->EnterConditions, FStateTreeIndex16(EnterConditionsBegin), SharedInstanceStructs))
+		const bool bUseEvaluationScopeInstanceData = UE::StateTree::Compiler::bEnableConditionWithEvaluationScopeInstanceData;
+		TArray<FInstancedStruct>& InstancedStructContainer = bUseEvaluationScopeInstanceData ? EvaluationScopeStructs : SharedInstanceStructs;
+		if (!CreateBindingsForNodes(SourceState->EnterConditions, FStateTreeIndex16(EnterConditionsBegin), InstancedStructContainer))
 		{
 			bSucceeded = false;
 			continue;
@@ -1552,7 +1569,7 @@ bool FStateTreeCompiler::CreateStateTransitions()
 			}
 			CompactTransition.ConditionsNum = uint8(ConditionsNum);
 
-			if (!CreateBindingsForNodes(Transition.Conditions, FStateTreeIndex16(ConditionsBegin), SharedInstanceStructs))
+			if (!CreateBindingsForNodes(Transition.Conditions, FStateTreeIndex16(ConditionsBegin), InstancedStructContainer))
 			{
 				bTransitionSucceeded = false;
 				continue;
@@ -1622,7 +1639,9 @@ bool FStateTreeCompiler::CreateStateConsiderations()
 		}
 		CompactState.UtilityConsiderationsNum = uint8(UtilityConsiderationsNum);
 
-		if (!CreateBindingsForNodes(SourceState->Considerations, FStateTreeIndex16(UtilityConsiderationsBegin), SharedInstanceStructs))
+		const bool bUseEvaluationScopeInstanceData = UE::StateTree::Compiler::bEnableUtilityConsiderationWithEvaluationScopeInstanceData;
+		TArray<FInstancedStruct>& InstancedStructContainer = bUseEvaluationScopeInstanceData ? EvaluationScopeStructs : SharedInstanceStructs;
+		if (!CreateBindingsForNodes(SourceState->Considerations, FStateTreeIndex16(UtilityConsiderationsBegin), InstancedStructContainer))
 		{
 			bSucceeded = false;
 			continue;
@@ -1793,7 +1812,7 @@ bool FStateTreeCompiler::CreatePropertyFunction(const FStateTreeEditorNode& Func
 	StructDesc.Name = FuncEditorNode.GetName();
 	StructDesc.DataSource = EStateTreeBindableStructSource::PropertyFunction;
 
-	const bool bUseEvaluationScopeInstanceData = UE::StateTree::Compiler::CVarLogEnablePropertyFunctionWithEvaluationScopeInstanceData->GetBool();
+	const bool bUseEvaluationScopeInstanceData = UE::StateTree::Compiler::bEnablePropertyFunctionWithEvaluationScopeInstanceData;
 	const FStateTreeDataHandle DataHandle = bUseEvaluationScopeInstanceData
 		? FStateTreeDataHandle(EStateTreeDataSourceType::EvaluationScopeInstanceData, EvaluationScopeStructs.Num())
 		: FStateTreeDataHandle(EStateTreeDataSourceType::SharedInstanceData, SharedInstanceStructs.Num());
@@ -2018,24 +2037,31 @@ bool FStateTreeCompiler::CreateCondition(UStateTreeState& State, const FString&
 	InstanceDesc.Name = CondNode.GetName();
 	InstanceDesc.DataSource = EStateTreeBindableStructSource::Condition;
 
-	FInstancedStruct* CreatedNode = CreateNodeWithSharedInstanceData(&State, CondNode, InstanceDesc);
-	if (CreatedNode)
+	const bool bUseEvaluationScopeInstanceData = UE::StateTree::Compiler::bEnableConditionWithEvaluationScopeInstanceData;
+	const FStateTreeDataHandle DataHandle = bUseEvaluationScopeInstanceData
+		? FStateTreeDataHandle(EStateTreeDataSourceType::EvaluationScopeInstanceData, EvaluationScopeStructs.Num())
+		: FStateTreeDataHandle(EStateTreeDataSourceType::SharedInstanceData, SharedInstanceStructs.Num());
+	TArray<FInstancedStruct>& InstancedStructContainer = bUseEvaluationScopeInstanceData ? EvaluationScopeStructs : SharedInstanceStructs;
+	FInstancedStruct* CreatedNode = CreateNode(&State, CondNode, InstanceDesc, DataHandle, InstancedStructContainer);
+	if (CreatedNode == nullptr)
 	{
-		FStateTreeConditionBase* Cond = CreatedNode->GetMutablePtr<FStateTreeConditionBase>();
-		if (ensure(Cond))
-		{
-			if (Cond->EvaluationMode == EStateTreeConditionEvaluationMode::ForcedFalse
-				|| Cond->EvaluationMode == EStateTreeConditionEvaluationMode::ForcedTrue)
-			{
-				Log.Reportf(EMessageSeverity::Info, InstanceDesc,
-						TEXT("The condition result will always be %s."),
-						Cond->EvaluationMode == EStateTreeConditionEvaluationMode::ForcedTrue ? TEXT("True") : TEXT("False"));
-			}
+		return false;
+	}
 
-			Cond->Operand = Operand;
-			Cond->DeltaIndent = DeltaIndent;
-			return true;
+	FStateTreeConditionBase* Cond = CreatedNode->GetMutablePtr<FStateTreeConditionBase>();
+	if (ensure(Cond))
+	{
+		if (Cond->EvaluationMode == EStateTreeConditionEvaluationMode::ForcedFalse
+			|| Cond->EvaluationMode == EStateTreeConditionEvaluationMode::ForcedTrue)
+		{
+			Log.Reportf(EMessageSeverity::Info, InstanceDesc,
+					TEXT("The condition result will always be %s."),
+					Cond->EvaluationMode == EStateTreeConditionEvaluationMode::ForcedTrue ? TEXT("True") : TEXT("False"));
 		}
+
+		Cond->Operand = Operand;
+		Cond->DeltaIndent = DeltaIndent;
+		return true;
 	}
 
 	return false;
@@ -2085,16 +2111,23 @@ bool FStateTreeCompiler::CreateConsideration(UStateTreeState& State, const FStri
 	InstanceDesc.Name = ConsiderationNode.GetName();
 	InstanceDesc.DataSource = EStateTreeBindableStructSource::Consideration;
 
-	FInstancedStruct* CreatedNode = CreateNodeWithSharedInstanceData(&State, ConsiderationNode, InstanceDesc);
-	if (CreatedNode)
+	const bool bUseEvaluationScopeInstanceData = UE::StateTree::Compiler::bEnableUtilityConsiderationWithEvaluationScopeInstanceData;
+	const FStateTreeDataHandle DataHandle = bUseEvaluationScopeInstanceData
+		? FStateTreeDataHandle(EStateTreeDataSourceType::EvaluationScopeInstanceData, EvaluationScopeStructs.Num())
+		: FStateTreeDataHandle(EStateTreeDataSourceType::SharedInstanceData, SharedInstanceStructs.Num());
+	TArray<FInstancedStruct>& InstancedStructContainer = bUseEvaluationScopeInstanceData ? EvaluationScopeStructs : SharedInstanceStructs;
+	FInstancedStruct* CreatedNode = CreateNode(&State, ConsiderationNode, InstanceDesc, DataHandle, InstancedStructContainer);
+	if (CreatedNode == nullptr)
 	{
-		FStateTreeConsiderationBase* Consideration = CreatedNode->GetMutablePtr<FStateTreeConsiderationBase>();
-		if (ensure(Consideration))
-		{
-			Consideration->Operand = Operand;
-			Consideration->DeltaIndent = DeltaIndent;
-			return true;
-		}
+		return false;
+	}
+
+	FStateTreeConsiderationBase* Consideration = CreatedNode->GetMutablePtr<FStateTreeConsiderationBase>();
+	if (ensure(Consideration))
+	{
+		Consideration->Operand = Operand;
+		Consideration->DeltaIndent = DeltaIndent;
+		return true;
 	}
 
 	return false;
diff --git a/Engine/Plugins/Runtime/StateTree/Source/StateTreeModule/Private/StateTree.cpp b/Engine/Plugins/Runtime/StateTree/Source/StateTreeModule/Private/StateTree.cpp
--- a/Engine/Plugins/Runtime/StateTree/Source/StateTreeModule/Private/StateTree.cpp
+++ b/Engine/Plugins/Runtime/StateTree/Source/StateTreeModule/Private/StateTree.cpp
@@ -758,6 +758,39 @@ void UStateTree::UpdateRuntimeFlags()
 				}
 			}
 		}
+
+		// Cache the amount of memory needed to execute the conditions.
+		{
+			UE::StateTree::InstanceData::FEvaluationScopeInstanceContainer::FMemoryRequirementBuilder Requirement;
+			for (int32 Index = 0; Index < State.EnterConditionsNum; ++Index)
+			{
+				const int32 NodeIndex = State.EnterConditionsBegin + Index;
+				const FStateTreeConditionBase& Cond = Nodes[NodeIndex].Get<const FStateTreeConditionBase>();
+				if (Cond.InstanceDataHandle.GetSource() == EStateTreeDataSourceType::EvaluationScopeInstanceData
+					|| Cond.InstanceDataHandle.GetSource() == EStateTreeDataSourceType::EvaluationScopeInstanceDataObject)
+				{
+					const FConstStructView DefaultInstanceView = DefaultEvaluationScopeInstanceData.GetStruct(Cond.InstanceTemplateIndex.Get());
+					Requirement.Add(DefaultInstanceView.GetScriptStruct());
+				}
+			}
+			State.EnterConditionEvaluationScopeMemoryRequirement = Requirement.Build();
+		}
+		// Cache the amount of memory needed to execute the considerations.
+		{
+			UE::StateTree::InstanceData::FEvaluationScopeInstanceContainer::FMemoryRequirementBuilder Requirement;
+			for (int32 Index = 0; Index < State.UtilityConsiderationsNum; ++Index)
+			{
+				const int32 NodeIndex = State.UtilityConsiderationsBegin + Index;
+				const FStateTreeConsiderationBase& Consideration = Nodes[NodeIndex].Get<const FStateTreeConsiderationBase>();
+				if (Consideration.InstanceDataHandle.GetSource() == EStateTreeDataSourceType::EvaluationScopeInstanceData
+					|| Consideration.InstanceDataHandle.GetSource() == EStateTreeDataSourceType::EvaluationScopeInstanceDataObject)
+				{
+					const FConstStructView DefaultInstanceView = DefaultEvaluationScopeInstanceData.GetStruct(Consideration.InstanceTemplateIndex.Get());
+					Requirement.Add(DefaultInstanceView.GetScriptStruct());
+				}
+			}
+			State.ConsiderationEvaluationScopeMemoryRequirement = Requirement.Build();
+		}
 	}
 
 	// Update the global task flags.
@@ -784,6 +817,24 @@ void UStateTree::UpdateRuntimeFlags()
 		}
 	}
 
+	// Cache the amount of memory needed to execute the transition's condition.
+	for (FCompactStateTransition& Transition : Transitions)
+	{
+		UE::StateTree::InstanceData::FEvaluationScopeInstanceContainer::FMemoryRequirementBuilder Requirement;
+		for (int32 Index = 0; Index < Transition.ConditionsNum; ++Index)
+		{
+			const int32 NodeIndex = Transition.ConditionsBegin + Index;
+			const FStateTreeConditionBase& Cond = Nodes[NodeIndex].Get<const FStateTreeConditionBase>();
+			if (Cond.InstanceDataHandle.GetSource() == EStateTreeDataSourceType::EvaluationScopeInstanceData
+				|| Cond.InstanceDataHandle.GetSource() == EStateTreeDataSourceType::EvaluationScopeInstanceDataObject)
+			{
+				const FConstStructView DefaultInstanceView = DefaultEvaluationScopeInstanceData.GetStruct(Cond.InstanceTemplateIndex.Get());
+				Requirement.Add(DefaultInstanceView.GetScriptStruct());
+			}
+		}
+		Transition.ConditionEvaluationScopeMemoryRequirement = Requirement.Build();
+	}
+
 	bScheduledTickAllowed = Schema ? Schema->IsScheduledTickAllowed() : false;
 	StateSelectionRules = Schema ? Schema->GetStateSelectionRules() : EStateTreeStateSelectionRules::Default;
 }
diff --git a/Engine/Plugins/Runtime/StateTree/Source/StateTreeModule/Private/StateTreeExecutionContext.cpp b/Engine/Plugins/Runtime/StateTree/Source/StateTreeModule/Private/StateTreeExecutionContext.cpp
--- a/Engine/Plugins/Runtime/StateTree/Source/StateTreeModule/Private/StateTreeExecutionContext.cpp
+++ b/Engine/Plugins/Runtime/StateTree/Source/StateTreeModule/Private/StateTreeExecutionContext.cpp
@@ -240,6 +240,19 @@ namespace UE::StateTree::ExecutionContext::Private
 				return DelayedState.StateID == StateID;
 			});
 	}
+
+	void InitEvaluationScopeInstanceData(UE::StateTree::InstanceData::FEvaluationScopeInstanceContainer& Container, TNotNull<const UStateTree*> StateTree, const int32 NodesBegin, const int32 NodesEnd)
+	{
+		for (int32 NodeIndex = NodesBegin; NodeIndex != NodesEnd; NodeIndex++)
+		{
+			const FStateTreeNodeBase& Node = StateTree->GetNodes()[NodeIndex].Get<const FStateTreeNodeBase>();
+			if (Node.InstanceDataHandle.GetSource() == EStateTreeDataSourceType::EvaluationScopeInstanceData
+				|| Node.InstanceDataHandle.GetSource() == EStateTreeDataSourceType::EvaluationScopeInstanceDataObject)
+			{
+				Container.Add(Node.InstanceDataHandle, StateTree->GetDefaultEvaluationScopeInstanceData().GetStruct(Node.InstanceTemplateIndex.Get()));
+			}
+		}
+	}
 } // namespace UE::StateTree::ExecutionContext::Private
 
 namespace UE::StateTree::ExecutionContext
@@ -3021,22 +3034,6 @@ bool FStateTreeExecutionContext::CopyBatchWithValidation(const FStateTreeExecuti
 	return CopyBatchInternal<bOnActiveInstances>(ParentFrame, CurrentFrame, TargetView, BindingsBatch);
 }
 
-namespace UE::StateTree::ExecutionContext::Private
-{
-	void InitPropertyFunctions(UE::StateTree::InstanceData::FEvaluationScopeInstanceContainer& Container, TNotNull<const UStateTree*> StateTree, int32 FuncBegin, int32 FuncEnd)
-	{
-		for (int32 FuncIndex = FuncBegin; FuncIndex != FuncEnd; ++FuncIndex)
-		{
-			const FStateTreePropertyFunctionBase& Func = StateTree->GetNodes()[FuncIndex].Get<const FStateTreePropertyFunctionBase>();
-			if (Func.InstanceDataHandle.GetSource() == EStateTreeDataSourceType::EvaluationScopeInstanceData
-				|| Func.InstanceDataHandle.GetSource() == EStateTreeDataSourceType::EvaluationScopeInstanceDataObject)
-			{
-				Container.Add(Func.InstanceDataHandle, StateTree->GetDefaultEvaluationScopeInstanceData().GetStruct(Func.InstanceTemplateIndex.Get()));
-			}
-		}
-	}
-}
-
 template<bool bOnActiveInstances>
 bool FStateTreeExecutionContext::CopyBatchInternal(const FStateTreeExecutionFrame* ParentFrame, const FStateTreeExecutionFrame& CurrentFrame, const FStateTreeDataView TargetView, const FStateTreeIndex16 BindingsBatch)
 {
@@ -3058,7 +3055,7 @@ bool FStateTreeExecutionContext::CopyBatchInternal(const FStateTreeExecutionFram
 			void* FunctionBindingMemory = FMemory_Alloca_Aligned(MemoryRequirement.Size, MemoryRequirement.Alignment);
 			EvaluationScopeContainer = FEvaluationScopeInstanceContainer(FunctionBindingMemory, MemoryRequirement);
 			PushEvaluationScopeInstanceContainer(EvaluationScopeContainer, CurrentFrame);
-			InitPropertyFunctions(EvaluationScopeContainer, CurrentFrame.StateTree, Batch.PropertyFunctionsBegin.Get(), Batch.PropertyFunctionsEnd.Get());
+			InitEvaluationScopeInstanceData(EvaluationScopeContainer, CurrentFrame.StateTree, Batch.PropertyFunctionsBegin.Get(), Batch.PropertyFunctionsEnd.Get());
 			bEvaluationScopeContainerPushed = true;
 		}
 
@@ -4659,24 +4656,27 @@ FStateTreeExecutionContext::FTickTaskResult FStateTreeExecutionContext::TickTask
 // Deprecated
 bool FStateTreeExecutionContext::TestAllConditions(const FStateTreeExecutionFrame* CurrentParentFrame, const FStateTreeExecutionFrame& CurrentFrame, const int32 ConditionsOffset, const int32 ConditionsNum)
 {
-	return TestAllConditionsInternal<false>(CurrentParentFrame, CurrentFrame, FStateTreeStateHandle(), ConditionsOffset, ConditionsNum, EStateTreeUpdatePhase::EnterConditions);
+	return TestAllConditionsInternal<false>(CurrentParentFrame, CurrentFrame, FStateTreeStateHandle(), FMemoryRequirement(), ConditionsOffset, ConditionsNum, EStateTreeUpdatePhase::EnterConditions);
 }
 
-bool FStateTreeExecutionContext::TestAllConditionsOnActiveInstances(const FStateTreeExecutionFrame* CurrentParentFrame, const FStateTreeExecutionFrame& CurrentFrame, FStateTreeStateHandle CurrentStateHandle, const int32 ConditionsOffset, const int32 ConditionsNum, EStateTreeUpdatePhase Phase)
+bool FStateTreeExecutionContext::TestAllConditionsOnActiveInstances(const FStateTreeExecutionFrame* CurrentParentFrame, const FStateTreeExecutionFrame& CurrentFrame, FStateTreeStateHandle CurrentStateHandle, const FMemoryRequirement& MemoryRequirement, const int32 ConditionsOffset, const int32 ConditionsNum, EStateTreeUpdatePhase Phase)
 {
-	return TestAllConditionsInternal<true>(CurrentParentFrame, CurrentFrame, CurrentStateHandle, ConditionsOffset, ConditionsNum, Phase);
+	return TestAllConditionsInternal<true>(CurrentParentFrame, CurrentFrame, CurrentStateHandle, MemoryRequirement, ConditionsOffset, ConditionsNum, Phase);
 }
 
-bool FStateTreeExecutionContext::TestAllConditionsWithValidation(const FStateTreeExecutionFrame* CurrentParentFrame, const FStateTreeExecutionFrame& CurrentFrame, FStateTreeStateHandle CurrentStateHandle, const int32 ConditionsOffset, const int32 ConditionsNum, EStateTreeUpdatePhase Phase)
+bool FStateTreeExecutionContext::TestAllConditionsWithValidation(const FStateTreeExecutionFrame* CurrentParentFrame, const FStateTreeExecutionFrame& CurrentFrame, FStateTreeStateHandle CurrentStateHandle, const FMemoryRequirement& MemoryRequirement, const int32 ConditionsOffset, const int32 ConditionsNum, EStateTreeUpdatePhase Phase)
 {
-	return TestAllConditionsInternal<false>(CurrentParentFrame, CurrentFrame, CurrentStateHandle, ConditionsOffset, ConditionsNum, Phase);
+	return TestAllConditionsInternal<false>(CurrentParentFrame, CurrentFrame, CurrentStateHandle, MemoryRequirement, ConditionsOffset, ConditionsNum, Phase);
 }
 
 template<bool bOnActiveInstances>
-bool FStateTreeExecutionContext::TestAllConditionsInternal(const FStateTreeExecutionFrame* CurrentParentFrame, const FStateTreeExecutionFrame& CurrentFrame, FStateTreeStateHandle CurrentStateHandle, const int32 ConditionsOffset, const int32 ConditionsNum, EStateTreeUpdatePhase Phase)
+bool FStateTreeExecutionContext::TestAllConditionsInternal(const FStateTreeExecutionFrame* CurrentParentFrame, const FStateTreeExecutionFrame& CurrentFrame, FStateTreeStateHandle CurrentStateHandle, const FMemoryRequirement& MemoryRequirement, const int32 ConditionsOffset, const int32 ConditionsNum, EStateTreeUpdatePhase Phase)
 {
 	CSV_SCOPED_TIMING_STAT_EXCLUSIVE(StateTree_TestConditions);
 
+	using namespace UE::StateTree::InstanceData;
+	using namespace UE::StateTree::ExecutionContext::Private;
+
 	if (ConditionsNum == 0)
 	{
 		return true;
@@ -4690,15 +4690,24 @@ bool FStateTreeExecutionContext::TestAllConditionsInternal(const FStateTreeExecu
 	TStaticArray<EStateTreeExpressionOperand, UE::StateTree::MaxExpressionIndent + 1> Operands(InPlace, EStateTreeExpressionOperand::Copy);
 	TStaticArray<bool, UE::StateTree::MaxExpressionIndent + 1> Values(InPlace, false);
 
+	FEvaluationScopeInstanceContainer EvaluationScopeContainer;
+	bool bEvaluationScopeContainerPushed = false;
+
+	if (MemoryRequirement.Size > 0)
+	{
+		void* FunctionBindingMemory = FMemory_Alloca_Aligned(MemoryRequirement.Size, MemoryRequirement.Alignment);
+		EvaluationScopeContainer = FEvaluationScopeInstanceContainer(FunctionBindingMemory, MemoryRequirement);
+		PushEvaluationScopeInstanceContainer(EvaluationScopeContainer, CurrentFrame);
+		InitEvaluationScopeInstanceData(EvaluationScopeContainer, CurrentFrame.StateTree, ConditionsOffset, ConditionsOffset + ConditionsNum);
+		bEvaluationScopeContainerPushed = true;
+	}
+
 	int32 Level = 0;
 
 	for (int32 Index = 0; Index < ConditionsNum; Index++)
 	{
 		const int32 ConditionIndex = ConditionsOffset + Index;
 		const FStateTreeConditionBase& Cond = CurrentFrame.StateTree->Nodes[ConditionIndex].Get<const FStateTreeConditionBase>();
-		checkf(Cond.InstanceDataHandle.GetSource() == EStateTreeDataSourceType::SharedInstanceData
-			|| Cond.InstanceDataHandle.GetSource() == EStateTreeDataSourceType::SharedInstanceDataObject,
-			TEXT("Condition only support SharedInstanceData"));
 
 		const FStateTreeDataView ConditionInstanceView = GetDataView(CurrentParentFrame, CurrentFrame, Cond.InstanceDataHandle);
 		FNodeInstanceDataScope DataScope(*this, &Cond, ConditionIndex, Cond.InstanceDataHandle, ConditionInstanceView);
@@ -4776,6 +4785,11 @@ bool FStateTreeExecutionContext::TestAllConditionsInternal(const FStateTreeExecu
 		}
 	}
 
+	if (bEvaluationScopeContainerPushed)
+	{
+		PopEvaluationScopeInstanceContainer(EvaluationScopeContainer);
+	}
+
 	return Values[0];
 }
 
@@ -4783,17 +4797,20 @@ bool FStateTreeExecutionContext::TestAllConditionsInternal(const FStateTreeExecu
 float FStateTreeExecutionContext::EvaluateUtility(const FStateTreeExecutionFrame* CurrentParentFrame, const FStateTreeExecutionFrame& CurrentFrame, const int32 ConsiderationsBegin, const int32 ConsiderationsNum, const float StateWeight)
 {
 	const FStateTreeStateHandle StateHandle = FStateTreeStateHandle::Invalid;
-	return EvaluateUtilityWithValidation(CurrentParentFrame, CurrentFrame, StateHandle, ConsiderationsBegin, ConsiderationsNum, StateWeight);
+	return EvaluateUtilityWithValidation(CurrentParentFrame, CurrentFrame, StateHandle, FMemoryRequirement(), ConsiderationsBegin, ConsiderationsNum, StateWeight);
 }
 
-float FStateTreeExecutionContext::EvaluateUtilityWithValidation(const FStateTreeExecutionFrame* CurrentParentFrame, const FStateTreeExecutionFrame& CurrentFrame, FStateTreeStateHandle CurrentStateHandle, const int32 ConsiderationsBegin, const int32 ConsiderationsNum, const float StateWeight)
+float FStateTreeExecutionContext::EvaluateUtilityWithValidation(const FStateTreeExecutionFrame* CurrentParentFrame, const FStateTreeExecutionFrame& CurrentFrame, FStateTreeStateHandle CurrentStateHandle, const FMemoryRequirement& MemoryRequirement, const int32 ConsiderationsBegin, const int32 ConsiderationsNum, const float StateWeight)
 {
 	// @todo: Tracing support
 	CSV_SCOPED_TIMING_STAT_EXCLUSIVE(StateTree_EvaluateUtility);
 
+	using namespace UE::StateTree::InstanceData;
+	using namespace UE::StateTree::ExecutionContext::Private;
+
 	if (ConsiderationsNum == 0)
 	{
-		return .0f;
+		return 0.0f;
 	}
 
 	FCurrentlyProcessedFrameScope FrameScope(*this, CurrentParentFrame, CurrentFrame);
@@ -4802,17 +4819,25 @@ float FStateTreeExecutionContext::EvaluateUtilityWithValidation(const FStateTree
 	TStaticArray<EStateTreeExpressionOperand, UE::StateTree::MaxExpressionIndent + 1> Operands(InPlace, EStateTreeExpressionOperand::Copy);
 	TStaticArray<float, UE::StateTree::MaxExpressionIndent + 1> Values(InPlace, false);
 
+	FEvaluationScopeInstanceContainer EvaluationScopeContainer;
+	bool bEvaluationScopeContainerPushed = false;
+
+	if (MemoryRequirement.Size > 0)
+	{
+		void* FunctionBindingMemory = FMemory_Alloca_Aligned(MemoryRequirement.Size, MemoryRequirement.Alignment);
+		EvaluationScopeContainer = FEvaluationScopeInstanceContainer(FunctionBindingMemory, MemoryRequirement);
+		PushEvaluationScopeInstanceContainer(EvaluationScopeContainer, CurrentFrame);
+		InitEvaluationScopeInstanceData(EvaluationScopeContainer, CurrentFrame.StateTree, ConsiderationsBegin, ConsiderationsBegin + ConsiderationsNum);
+		bEvaluationScopeContainerPushed = true;
+	}
+
 	int32 Level = 0;
-	float Value = .0f;
+	float Value = 0.0f;
 	for (int32 Index = 0; Index < ConsiderationsNum; Index++)
 	{
 		const int32 ConsiderationIndex = ConsiderationsBegin + Index;
 		const FStateTreeConsiderationBase& Consideration = CurrentFrame.StateTree->Nodes[ConsiderationIndex].Get<const FStateTreeConsiderationBase>();
 
-		checkf(Consideration.InstanceDataHandle.GetSource() == EStateTreeDataSourceType::SharedInstanceData
-			|| Consideration.InstanceDataHandle.GetSource() == EStateTreeDataSourceType::SharedInstanceDataObject
-			, TEXT("Consideration only support SharedInstanceData"));
-
 		const FStateTreeDataView ConsiderationInstanceView = GetDataView(CurrentParentFrame, CurrentFrame, Consideration.InstanceDataHandle);
 		FNodeInstanceDataScope DataScope(*this, &Consideration, ConsiderationIndex, Consideration.InstanceDataHandle, ConsiderationInstanceView);
 
@@ -4871,6 +4896,11 @@ float FStateTreeExecutionContext::EvaluateUtilityWithValidation(const FStateTree
 		}
 	}
 
+	if (bEvaluationScopeContainerPushed)
+	{
+		PopEvaluationScopeInstanceContainer(EvaluationScopeContainer);
+	}
+
 	return StateWeight * Values[0];
 }
 
@@ -5513,7 +5543,7 @@ bool FStateTreeExecutionContext::TriggerTransitions()
 					{
 						FCurrentlyProcessedTransitionEventScope TransitionEventScope(*this, TransitionEvent.IsValid() ? TransitionEvent.Get() : nullptr);
 						UE_STATETREE_DEBUG_TRANSITION_EVENT(this, FStateTreeTransitionSource(CurrentFrame.StateTree, FStateTreeIndex16(TransitionIndex), Transition.State, Transition.Priority), EStateTreeTraceEventType::OnEvaluating);
-						bPassed = TestAllConditionsOnActiveInstances(CurrentParentFrame, CurrentFrame, Handler.StateHandle, Transition.ConditionsBegin, Transition.ConditionsNum, EStateTreeUpdatePhase::TransitionConditions);
+						bPassed = TestAllConditionsOnActiveInstances(CurrentParentFrame, CurrentFrame, Handler.StateHandle, Transition.ConditionEvaluationScopeMemoryRequirement, Transition.ConditionsBegin, Transition.ConditionsNum, EStateTreeUpdatePhase::TransitionConditions);
 					}
 
 					if (bPassed)
@@ -5728,7 +5758,7 @@ PRAGMA_ENABLE_DEPRECATION_WARNINGS
 								bool bPassed = false;
 								{
 									UE_STATETREE_DEBUG_TRANSITION_EVENT(this, FStateTreeTransitionSource(CurrentFrame.StateTree, FStateTreeIndex16(TransitionIndex), Transition.State, Transition.Priority), EStateTreeTraceEventType::OnEvaluating);
-									bPassed = TestAllConditionsOnActiveInstances(CurrentParentFrame, CurrentFrame, CurrentStateHandle, Transition.ConditionsBegin, Transition.ConditionsNum, EStateTreeUpdatePhase::TransitionConditions);
+									bPassed = TestAllConditionsOnActiveInstances(CurrentParentFrame, CurrentFrame, CurrentStateHandle, Transition.ConditionEvaluationScopeMemoryRequirement, Transition.ConditionsBegin, Transition.ConditionsNum, EStateTreeUpdatePhase::TransitionConditions);
 								}
 
 								if (bPassed)
@@ -6737,7 +6767,7 @@ bool FStateTreeExecutionContext::SelectStateInternal(
 		if (bShouldPrerequisitesBeChecked)
 		{
 			// Check that the state can be entered
-			const bool bEnterConditionsPassed = TestAllConditionsWithValidation(NextParentFrame, NextFrame, NextStateHandle, NextState.EnterConditionsBegin, NextState.EnterConditionsNum, EStateTreeUpdatePhase::EnterConditions);
+			const bool bEnterConditionsPassed = TestAllConditionsWithValidation(NextParentFrame, NextFrame, NextStateHandle, NextState.EnterConditionEvaluationScopeMemoryRequirement, NextState.EnterConditionsBegin, NextState.EnterConditionsNum, EStateTreeUpdatePhase::EnterConditions);
 
 			if (!bEnterConditionsPassed)
 			{
@@ -7408,7 +7438,7 @@ bool FStateTreeExecutionContext::SelectStateInternal_TrySelectChildrenWithHighes
 		for (uint16 ChildState = NextState.ChildrenBegin; ChildState < NextState.ChildrenEnd; ChildState = NextStateTree->States[ChildState].GetNextSibling())
 		{
 			const FCompactStateTreeState& CurrentState = NextStateTree->States[ChildState];
-			const float Score = EvaluateUtilityWithValidation(NextParentFrame, *NextFrame, FStateTreeStateHandle(ChildState), CurrentState.UtilityConsiderationsBegin, CurrentState.UtilityConsiderationsNum, CurrentState.Weight);
+			const float Score = EvaluateUtilityWithValidation(NextParentFrame, *NextFrame, FStateTreeStateHandle(ChildState), CurrentState.ConsiderationEvaluationScopeMemoryRequirement, CurrentState.UtilityConsiderationsBegin, CurrentState.UtilityConsiderationsNum, CurrentState.Weight);
 			NextLevelChildStates.Emplace(ChildState, Score);
 		}
 	}
@@ -7483,7 +7513,7 @@ bool FStateTreeExecutionContext::SelectStateInternal_TrySelectChildrenAtRandomWe
 		for (uint16 ChildState = NextState.ChildrenBegin; ChildState < NextState.ChildrenEnd; ChildState = NextStateTree->States[ChildState].GetNextSibling())
 		{
 			const FCompactStateTreeState& CurrentState = NextStateTree->States[ChildState];
-			const float Score = EvaluateUtilityWithValidation(NextParentFrame, *NextFrame, FStateTreeStateHandle(ChildState), CurrentState.UtilityConsiderationsBegin, CurrentState.UtilityConsiderationsNum, CurrentState.Weight);
+			const float Score = EvaluateUtilityWithValidation(NextParentFrame, *NextFrame, FStateTreeStateHandle(ChildState), CurrentState.ConsiderationEvaluationScopeMemoryRequirement, CurrentState.UtilityConsiderationsBegin, CurrentState.UtilityConsiderationsNum, CurrentState.Weight);
 			if (Score > 0.0f)
 			{
 				NextLevelChildStates.Emplace(ChildState, Score);
@@ -7611,7 +7641,7 @@ bool FStateTreeExecutionContext::SelectStateInternal_TryFollowTransitions(
 
 				FCurrentlyProcessedTransitionEventScope TransitionEventScope(*this, SelectedStateTransitionEvent.IsValid() ? SelectedStateTransitionEvent.Get() : nullptr);
 				UE_STATETREE_DEBUG_TRANSITION_EVENT(this, FStateTreeTransitionSource(NextStateTree, FStateTreeIndex16(TransitionIndex), Transition.State, Transition.Priority), EStateTreeTraceEventType::OnEvaluating);
-				bTransitionConditionsPassed = TestAllConditionsWithValidation(NextParentFrame, *NextFrame, NextStateHandle, Transition.ConditionsBegin, Transition.ConditionsNum, EStateTreeUpdatePhase::TransitionConditions);
+				bTransitionConditionsPassed = TestAllConditionsWithValidation(NextParentFrame, *NextFrame, NextStateHandle, Transition.ConditionEvaluationScopeMemoryRequirement, Transition.ConditionsBegin, Transition.ConditionsNum, EStateTreeUpdatePhase::TransitionConditions);
 			}
 
 			if (bTransitionConditionsPassed)
diff --git a/Engine/Plugins/Runtime/StateTree/Source/StateTreeModule/Public/StateTreeEvaluationScopeInstanceContainer.h b/Engine/Plugins/Runtime/StateTree/Source/StateTreeModule/Public/StateTreeEvaluationScopeInstanceContainer.h
--- a/Engine/Plugins/Runtime/StateTree/Source/StateTreeModule/Public/StateTreeEvaluationScopeInstanceContainer.h
+++ b/Engine/Plugins/Runtime/StateTree/Source/StateTreeModule/Public/StateTreeEvaluationScopeInstanceContainer.h
@@ -17,12 +17,7 @@ namespace UE::StateTree::InstanceData
 	{
 	public:
 		/** The memory requirement for the container allocation. */
-		struct FMemoryRequirement
-		{
-			int32 Size = 0;
-			int32 Alignment = 0;
-			int32 NumberOfElements = 0;
-		};
+		using FMemoryRequirement = UE::StateTree::InstanceData::FEvaluationScopeMemoryRequirement;
 
 		/** Builds the amount of memory required to store the provided structures. */
 		struct FMemoryRequirementBuilder
diff --git a/Engine/Plugins/Runtime/StateTree/Source/StateTreeModule/Public/StateTreeExecutionContext.h b/Engine/Plugins/Runtime/StateTree/Source/StateTreeModule/Public/StateTreeExecutionContext.h
--- a/Engine/Plugins/Runtime/StateTree/Source/StateTreeModule/Public/StateTreeExecutionContext.h
+++ b/Engine/Plugins/Runtime/StateTree/Source/StateTreeModule/Public/StateTreeExecutionContext.h
@@ -1052,9 +1052,12 @@ protected:
 	/** Handles transition triggering part of the tick. */
 	UE_API void TickTriggerTransitionsInternal();
 
+protected:
+	using FMemoryRequirement = UE::StateTree::InstanceData::FEvaluationScopeInstanceContainer::FMemoryRequirement;
+
 private:
 	template<bool bOnActiveInstances>
-	bool TestAllConditionsInternal(const FStateTreeExecutionFrame* CurrentParentFrame, const FStateTreeExecutionFrame& CurrentFrame, FStateTreeStateHandle CurrentStateHandle, const int32 ConditionsOffset, const int32 ConditionsNum, EStateTreeUpdatePhase Phase);
+	bool TestAllConditionsInternal(const FStateTreeExecutionFrame* CurrentParentFrame, const FStateTreeExecutionFrame& CurrentFrame, FStateTreeStateHandle CurrentStateHandle, const FMemoryRequirement& MemoryRequirement, const int32 ConditionsOffset, const int32 ConditionsNum, EStateTreeUpdatePhase Phase);
 
 protected:
 	/**
@@ -1062,14 +1065,14 @@ protected:
 	 * Should be used only on active node instances, assumes valid data handles and does not consider temporary node instances.
 	 * @return True if all conditions pass.
 	 */
-	UE_API bool TestAllConditionsOnActiveInstances(const FStateTreeExecutionFrame* CurrentParentFrame, const FStateTreeExecutionFrame& CurrentFrame, FStateTreeStateHandle CurrentStateHandle, const int32 ConditionsOffset, const int32 ConditionsNum, EStateTreeUpdatePhase Phase);
+	UE_API bool TestAllConditionsOnActiveInstances(const FStateTreeExecutionFrame* CurrentParentFrame, const FStateTreeExecutionFrame& CurrentFrame, FStateTreeStateHandle CurrentStateHandle, const FMemoryRequirement& MemoryRequirement, const int32 ConditionsOffset, const int32 ConditionsNum, EStateTreeUpdatePhase Phase);
 
 	/**
 	 * Checks all conditions at given range.
 	 * This version validates the data handles and looks up temporary instances.
 	 * @return True if all conditions pass.
 	 */
-	UE_API bool TestAllConditionsWithValidation(const FStateTreeExecutionFrame* CurrentParentFrame, const FStateTreeExecutionFrame& CurrentFrame, FStateTreeStateHandle CurrentStateHandle, const int32 ConditionsOffset, const int32 ConditionsNum, EStateTreeUpdatePhase Phase);
+	UE_API bool TestAllConditionsWithValidation(const FStateTreeExecutionFrame* CurrentParentFrame, const FStateTreeExecutionFrame& CurrentFrame, FStateTreeStateHandle CurrentStateHandle, const FMemoryRequirement& MemoryRequirement, const int32 ConditionsOffset, const int32 ConditionsNum, EStateTreeUpdatePhase Phase);
 
 	UE_DEPRECATED(5.7, "Use TestAllConditionsWithValidation or TestAllConditionsOnActiveInstances. TestAllConditionsWithValidation can get the data from a temporary frame in state selection.")
 	UE_API bool TestAllConditions(const FStateTreeExecutionFrame* CurrentParentFrame, const FStateTreeExecutionFrame& CurrentFrame, const int32 ConditionsOffset, const int32 ConditionsNum);
@@ -1078,7 +1081,7 @@ protected:
 	 * Calculate the final score of all considerations at given range.
 	 * @return the final score
 	 */
-	UE_API float EvaluateUtilityWithValidation(const FStateTreeExecutionFrame* CurrentParentFrame, const FStateTreeExecutionFrame& CurrentFrame, FStateTreeStateHandle StateHandle, const int32 ConsiderationsBegin, const int32 ConsiderationsNum, const float StateWeight);
+	UE_API float EvaluateUtilityWithValidation(const FStateTreeExecutionFrame* CurrentParentFrame, const FStateTreeExecutionFrame& CurrentFrame, FStateTreeStateHandle StateHandle, const FMemoryRequirement& MemoryRequirement, const int32 ConsiderationsBegin, const int32 ConsiderationsNum, const float StateWeight);
 
 	UE_DEPRECATED(5.7, "Use EvaluateUtilityWithValidation. It can get the data from a temporary frame in state selection.")
 	UE_API float EvaluateUtility(const FStateTreeExecutionFrame* CurrentParentFrame, const FStateTreeExecutionFrame& CurrentFrame, const int32 ConsiderationsOffset, const int32 ConsiderationsNum, const float StateWeight);
diff --git a/Engine/Plugins/Runtime/StateTree/Source/StateTreeModule/Public/StateTreeTypes.h b/Engine/Plugins/Runtime/StateTree/Source/StateTreeModule/Public/StateTreeTypes.h
--- a/Engine/Plugins/Runtime/StateTree/Source/StateTreeModule/Public/StateTreeTypes.h
+++ b/Engine/Plugins/Runtime/StateTree/Source/StateTreeModule/Public/StateTreeTypes.h
@@ -461,6 +461,16 @@ private:
 	FStateTreeStateHandle StateHandle = FStateTreeStateHandle::Invalid; 
 };
 
+namespace UE::StateTree::InstanceData
+{
+	/** The memory requirement for the container allocation. */
+	struct FEvaluationScopeMemoryRequirement
+	{
+		int32 Size = 0;
+		int32 Alignment = 0;
+		int32 NumberOfElements = 0;
+	};
+};
 
 /**
  * Time duration with random variance. Stored compactly as two uint16s, which gives time range of about 650 seconds.
@@ -608,6 +618,9 @@ struct FCompactStateTransition
 	UPROPERTY()
 	FStateTreeDelegateDispatcher RequiredDelegateDispatcher;
 
+	/** The amount of memory used by the condition when it uses the evaluation scope data type. */
+	UE::StateTree::InstanceData::FEvaluationScopeMemoryRequirement ConditionEvaluationScopeMemoryRequirement;
+
 	/** Index to first condition to test */
 	UPROPERTY()
 	uint16 ConditionsBegin = 0;
@@ -734,6 +747,12 @@ struct FCompactStateTreeState
 	UPROPERTY()
 	FGameplayTag Tag;
 
+	/** The amount of memory used by the condition when it uses the evaluation scope data type. */
+	UE::StateTree::InstanceData::FEvaluationScopeMemoryRequirement EnterConditionEvaluationScopeMemoryRequirement;
+
+	/** The amount of memory used by the consideration when it uses the evaluation scope data type. */
+	UE::StateTree::InstanceData::FEvaluationScopeMemoryRequirement ConsiderationEvaluationScopeMemoryRequirement;
+
 	UPROPERTY()
 	TObjectPtr<UStateTree> LinkedAsset = nullptr;
 
