[StateTree] Add accessor for the UStateTree asset from the StateTreeLinker

diff --git a/Engine/Plugins/Runtime/StateTree/Source/StateTreeModule/Private/StateTree.cpp b/Engine/Plugins/Runtime/StateTree/Source/StateTreeModule/Private/StateTree.cpp
--- a/Engine/Plugins/Runtime/StateTree/Source/StateTreeModule/Private/StateTree.cpp
+++ b/Engine/Plugins/Runtime/StateTree/Source/StateTreeModule/Private/StateTree.cpp
@@ -645,7 +645,7 @@ bool UStateTree::Link()
 
 	// Resolves nodes references to other StateTree data
 	{
-		FStateTreeLinker Linker(Schema);
+		FStateTreeLinker Linker(this);
 
 		for (int32 Index = 0; Index < Nodes.Num(); Index++)
 		{
diff --git a/Engine/Plugins/Runtime/StateTree/Source/StateTreeModule/Private/StateTreeLinker.cpp b/Engine/Plugins/Runtime/StateTree/Source/StateTreeModule/Private/StateTreeLinker.cpp
new file mode 100644
--- /dev/null
+++ b/Engine/Plugins/Runtime/StateTree/Source/StateTreeModule/Private/StateTreeLinker.cpp
@@ -0,0 +1,40 @@
+// Copyright Epic Games, Inc. All Rights Reserved.
+
+#include "StateTreeLinker.h"
+#include "StateTree.h"
+#include "StateTreeSchema.h"
+#include "Templates/Casts.h"
+
+#include UE_INLINE_GENERATED_CPP_BY_NAME(StateTreeLinker)
+
+FStateTreeLinker::FStateTreeLinker(TNotNull<const UStateTree*> InStateTree)
+	: StateTree(InStateTree)
+	, Schema(InStateTree->GetSchema())
+{
+}
+
+void FStateTreeLinker::LinkExternalData(FStateTreeExternalDataHandle& Handle, const UStruct* Struct, const EStateTreeExternalDataRequirement Requirement)
+{
+	if (Schema != nullptr && !Schema->IsExternalItemAllowed(*Struct))
+	{
+		UE_LOG(LogStateTree, Error,
+			TEXT("External data of type '%s' used by current node is not allowed by schema '%s' (i.e. rejected by IsExternalItemAllowed)"),
+			*Struct->GetName(),
+			*Schema->GetClass()->GetName());
+
+		Handle = FStateTreeExternalDataHandle();
+		Status = EStateTreeLinkerStatus::Failed;
+		return;
+	}
+
+	const FStateTreeExternalDataDesc Desc(Struct, Requirement);
+	int32 Index = ExternalDataDescs.Find(Desc);
+
+	if (Index == INDEX_NONE)
+	{
+		Index = ExternalDataDescs.Add(Desc);
+		ExternalDataDescs[Index].Handle.DataHandle = FStateTreeDataHandle(EStateTreeDataSourceType::ExternalData, Index);
+	}
+
+	Handle.DataHandle = ExternalDataDescs[Index].Handle.DataHandle;
+}
\ No newline at end of file
diff --git a/Engine/Plugins/Runtime/StateTree/Source/StateTreeModule/Public/StateTreeLinker.h b/Engine/Plugins/Runtime/StateTree/Source/StateTreeModule/Public/StateTreeLinker.h
--- a/Engine/Plugins/Runtime/StateTree/Source/StateTreeModule/Public/StateTreeLinker.h
+++ b/Engine/Plugins/Runtime/StateTree/Source/StateTreeModule/Public/StateTreeLinker.h
@@ -2,11 +2,14 @@
 
 #pragma once
 
+#include "StateTree.h"
 #include "StateTreeSchema.h"
 #include "StateTreeExecutionTypes.h"
 #include "Templates/Casts.h"
 #include "StateTreeLinker.generated.h"
 
+#define UE_API STATETREEMODULE_API
+
 UENUM()
 enum class EStateTreeLinkerStatus : uint8
 {
@@ -20,10 +23,24 @@ enum class EStateTreeLinkerStatus : uint8
  */
 struct FStateTreeLinker
 {
-	explicit FStateTreeLinker(const UStateTreeSchema* InSchema) : Schema(InSchema) {}
+	UE_API explicit FStateTreeLinker(TNotNull<const UStateTree*> InStateTree);
+
+	UE_DEPRECATED(5.7, "Use the constructor with the StateTree pointer.")
+	explicit FStateTreeLinker(const UStateTreeSchema* InSchema)
+		: Schema(InSchema)
+	{}
 	
 	/** @returns the linking status. */
-	EStateTreeLinkerStatus GetStatus() const { return Status; }
+	EStateTreeLinkerStatus GetStatus() const
+	{
+		return Status;
+	}
+
+	/** @returns the StateTree asset. */
+	const UStateTree* GetStateTree() const
+	{
+		return StateTree.Get();
+	}
 	
 	/**
 	 * Links reference to an external UObject.
@@ -62,31 +79,7 @@ struct FStateTreeLinker
 	 * @param Struct Expected type of the Object or Struct to link to.
 	 * @param Requirement Describes if the external data is expected to be required or optional.
 	 */
-	void LinkExternalData(FStateTreeExternalDataHandle& Handle, const UStruct* Struct, const EStateTreeExternalDataRequirement Requirement)
-	{
-		if (Schema != nullptr && !Schema->IsExternalItemAllowed(*Struct))
-		{
-			UE_LOG(LogStateTree, Error,
-				TEXT("External data of type '%s' used by current node is not allowed by schema '%s' (i.e. rejected by IsExternalItemAllowed)"),
-				*Struct->GetName(),
-				*Schema->GetClass()->GetName());
-
-			Handle = FStateTreeExternalDataHandle();
-			Status = EStateTreeLinkerStatus::Failed;
-			return;
-		}
-		
-		const FStateTreeExternalDataDesc Desc(Struct, Requirement);
-		int32 Index = ExternalDataDescs.Find(Desc);
-		
-		if (Index == INDEX_NONE)
-		{
-			Index = ExternalDataDescs.Add(Desc);
-			ExternalDataDescs[Index].Handle.DataHandle = FStateTreeDataHandle(EStateTreeDataSourceType::ExternalData, Index);
-		}
-		
-		Handle.DataHandle = ExternalDataDescs[Index].Handle.DataHandle;
-	}
+	UE_API void LinkExternalData(FStateTreeExternalDataHandle& Handle, const UStruct* Struct, const EStateTreeExternalDataRequirement Requirement);
 
 	/** @return linked external data descriptors. */
 	const TArrayView<const FStateTreeExternalDataDesc> GetExternalDataDescs() const
@@ -95,7 +88,10 @@ struct FStateTreeLinker
 	}
 
 protected:
+	TStrongObjectPtr<const UStateTree> StateTree;
 	TStrongObjectPtr<const UStateTreeSchema> Schema;
 	EStateTreeLinkerStatus Status = EStateTreeLinkerStatus::Succeeded;
 	TArray<FStateTreeExternalDataDesc> ExternalDataDescs;
 };
+
+#undef UE_API
