[StateTree] Fix invalid access when cached frame values become invalid. The temporary frame array grows, and the values need to be re-cached. It happens when you select a state that contains +2 new frames (via linked states) and there are state selection events.

diff --git a/Engine/Plugins/Runtime/StateTree/Source/StateTreeModule/Private/StateTreeExecutionContext.cpp b/Engine/Plugins/Runtime/StateTree/Source/StateTreeModule/Private/StateTreeExecutionContext.cpp
--- a/Engine/Plugins/Runtime/StateTree/Source/StateTreeModule/Private/StateTreeExecutionContext.cpp
+++ b/Engine/Plugins/Runtime/StateTree/Source/StateTreeModule/Private/StateTreeExecutionContext.cpp
@@ -6529,17 +6529,22 @@ bool FStateTreeExecutionContext::SelectStateInternal(
 
 	// Use the ID to get the ExecutionFrame because it is saved in an array (and this is a recursive function), the array can grow and the pointer won't be valid.
 	check(OutSelectionResult->SelectedFrames.Num() != 0);
-	UE::StateTree::FActiveFrameID NextFrameID = OutSelectionResult->SelectedFrames.Last();
-	const FStateTreeExecutionFrame* NextFramePtr = FindExecutionFrame(NextFrameID, MakeConstArrayView(Exec.ActiveFrames), MakeConstArrayView(OutSelectionResult->TemporaryFrames));
-	check(NextFramePtr);
+	const UE::StateTree::FActiveFrameID NextFrameID = OutSelectionResult->SelectedFrames.Last();
 
-	const FStateTreeExecutionFrame& NextFrame = *NextFramePtr;
+	const FStateTreeExecutionFrame* NextFrame = nullptr;
 	const FStateTreeExecutionFrame* NextParentFrame = nullptr;
-	if (OutSelectionResult->SelectedFrames.Num() > 1)
-	{
-		NextParentFrame = FindExecutionFrame(OutSelectionResult->SelectedFrames.Last(1), MakeConstArrayView(Exec.ActiveFrames), MakeConstArrayView(OutSelectionResult->TemporaryFrames));
-		check(NextParentFrame);
-	}
+	auto CacheNextFrame = [&NextFrame, &NextParentFrame, &NextFrameID, &Exec, &OutSelectionResult]()
+		{
+			NextFrame = FindExecutionFrame(NextFrameID, MakeConstArrayView(Exec.ActiveFrames), MakeConstArrayView(OutSelectionResult->TemporaryFrames));
+			check(NextFrame);
+
+			if (OutSelectionResult->SelectedFrames.Num() > 1)
+			{
+				NextParentFrame = FindExecutionFrame(OutSelectionResult->SelectedFrames.Last(1), MakeConstArrayView(Exec.ActiveFrames), MakeConstArrayView(OutSelectionResult->TemporaryFrames));
+				check(NextParentFrame);
+			}
+		};
+	CacheNextFrame();
 
 	// Save the current result to use SelectionEvents in GetDataView
 	TGuardValue<FSelectStateResult*> GuardCurrentlyProcessedStateSelectionResult(CurrentlyProcessedStateSelectionResult, &OutSelectionResult.Get());
@@ -6558,11 +6563,11 @@ bool FStateTreeExecutionContext::SelectStateInternal(
 
 			STATETREE_LOG(VeryVerbose, TEXT("%hs: In state '%s', overriding linked asset '%s' with '%s'. '%s' using StateTree '%s'."),
 				__FUNCTION__,
-				*GetSafeStateName(NextFrame, NextStateHandle),
+				*GetSafeStateName(*NextFrame, NextStateHandle),
 				*GetFullNameSafe(NextState.LinkedAsset),
 				*GetFullNameSafe(NextLinkedStateAsset),
 				*GetNameSafe(&Owner),
-				*GetFullNameSafe(NextFrame.StateTree)
+				*GetFullNameSafe(NextFrame->StateTree)
 			);
 		}
 	}
@@ -6570,11 +6575,11 @@ bool FStateTreeExecutionContext::SelectStateInternal(
 	// Update state parameters.
 	if (NextState.ParameterDataHandle.IsValid())
 	{
-		FCurrentlyProcessedFrameScope FrameScope(*this, NextParentFrame, NextFrame);
+		FCurrentlyProcessedFrameScope FrameScope(*this, NextParentFrame, *NextFrame);
 		FCurrentlyProcessedStateScope NextStateScope(*this, NextStateHandle);
 
 		// Instantiate temporary state parameters if not done yet.
-		FStateTreeDataView NextStateParametersView = GetDataViewOrTemporary(NextParentFrame, NextFrame, NextState.ParameterDataHandle);
+		FStateTreeDataView NextStateParametersView = GetDataViewOrTemporary(NextParentFrame, *NextFrame, NextState.ParameterDataHandle);
 		if (!NextStateParametersView.IsValid())
 		{
 			// Allocate temporary instance for parameters if the state has params.
@@ -6583,7 +6588,7 @@ bool FStateTreeExecutionContext::SelectStateInternal(
 			if (NextLinkedStateParameterOverride)
 			{
 				// Create from an override.
-				FStateTreeDataView TempStateParametersView = AddTemporaryInstance(NextFrame, FStateTreeIndex16::Invalid, NextState.ParameterDataHandle, FConstStructView(TBaseStructure<FCompactStateTreeParameters>::Get()));
+				FStateTreeDataView TempStateParametersView = AddTemporaryInstance(*NextFrame, FStateTreeIndex16::Invalid, NextState.ParameterDataHandle, FConstStructView(TBaseStructure<FCompactStateTreeParameters>::Get()));
 				check(TempStateParametersView.IsValid());
 				FCompactStateTreeParameters& StateParams = TempStateParametersView.GetMutable<FCompactStateTreeParameters>();
 				StateParams.Parameters = *NextLinkedStateParameterOverride;
@@ -6592,8 +6597,8 @@ bool FStateTreeExecutionContext::SelectStateInternal(
 			else
 			{
 				// Create from template in the asset.
-				const FConstStructView DefaultStateParamsInstanceData = NextFrame.StateTree->DefaultInstanceData.GetStruct(NextState.ParameterTemplateIndex.Get());
-				FStateTreeDataView TempStateParametersView = AddTemporaryInstance(NextFrame, FStateTreeIndex16::Invalid, NextState.ParameterDataHandle, DefaultStateParamsInstanceData);
+				const FConstStructView DefaultStateParamsInstanceData = NextFrame->StateTree->DefaultInstanceData.GetStruct(NextState.ParameterTemplateIndex.Get());
+				FStateTreeDataView TempStateParametersView = AddTemporaryInstance(*NextFrame, FStateTreeIndex16::Invalid, NextState.ParameterDataHandle, DefaultStateParamsInstanceData);
 				check(TempStateParametersView.IsValid());
 				FCompactStateTreeParameters& StateParams = TempStateParametersView.GetMutable<FCompactStateTreeParameters>();
 				NextStateParametersView = FStateTreeDataView(StateParams.Parameters.GetMutableValue());
@@ -6607,7 +6612,7 @@ bool FStateTreeExecutionContext::SelectStateInternal(
 		{
 			// Note: the parameters are for the current (linked) state, stored in current frame.
 			// The copy can fail, if the overridden parameters do not match, this is by design.
-			CopyBatchWithValidation(NextParentFrame, NextFrame, NextStateParametersView, NextState.ParameterBindingsBatch);
+			CopyBatchWithValidation(NextParentFrame, *NextFrame, NextStateParametersView, NextState.ParameterBindingsBatch);
 		}
 	}
 
@@ -6670,14 +6675,14 @@ bool FStateTreeExecutionContext::SelectStateInternal(
 	const bool bIsNextTargetStateInActiveStates = !bNewFrameCreated
 		&& InternalArgs.MissingActiveStates.Num() > 0
 		&& InternalArgs.MissingActiveStates[0].GetStateHandle() == NextStateHandle
-		&& InternalArgs.MissingActiveStates[0].GetFrameID() == NextFrame.FrameID;
+		&& InternalArgs.MissingActiveStates[0].GetFrameID() == NextFrame->FrameID;
 	if (EnumHasAnyFlags(Args.SelectionRules, EStateTreeStateSelectionRules::ReselectedStateCreatesNewStates | EStateTreeStateSelectionRules::CompletedTransitionStatesCreateNewStates))
 	{
 		bShouldCreateNewState = false;
 		if (EnumHasAllFlags(Args.SelectionRules, EStateTreeStateSelectionRules::CompletedTransitionStatesCreateNewStates))
 		{
 			// Request a new state if it doesn't match the previous states or if the state is completed.
-			bShouldCreateNewState = !bIsNextTargetStateInActiveStates || NextFrame.ActiveTasksStatus.GetStatus(NextState).IsCompleted();
+			bShouldCreateNewState = !bIsNextTargetStateInActiveStates || NextFrame->ActiveTasksStatus.GetStatus(NextState).IsCompleted();
 		}
 		if (EnumHasAnyFlags(Args.SelectionRules, EStateTreeStateSelectionRules::ReselectedStateCreatesNewStates))
 		{
@@ -6700,12 +6705,12 @@ bool FStateTreeExecutionContext::SelectStateInternal(
 	// Add the state to the selected list. It will be removed if the selection fails.
 	if (bShouldCreateNewState)
 	{
-		OutSelectionResult->SelectedStates.Emplace(NextFrame.FrameID, FActiveStateID(Storage.GenerateUniqueId()), NextStateHandle);
+		OutSelectionResult->SelectedStates.Emplace(NextFrame->FrameID, FActiveStateID(Storage.GenerateUniqueId()), NextStateHandle);
 
 		int32 StateInFrameCounter = 0;
 		for (int32 Index = OutSelectionResult->SelectedStates.Num() - 1; Index >= 0; --Index)
 		{
-			if (OutSelectionResult->SelectedStates[Index].GetFrameID() != NextFrame.FrameID)
+			if (OutSelectionResult->SelectedStates[Index].GetFrameID() != NextFrame->FrameID)
 			{
 				break;
 			}
@@ -6716,8 +6721,8 @@ bool FStateTreeExecutionContext::SelectStateInternal(
 		{
 			STATETREE_LOG(Error, TEXT("%hs: Reached max execution depth when trying to select state %s from '%s'.  '%s' using StateTree '%s'."),
 				__FUNCTION__,
-				*GetSafeStateName(NextFrame, NextStateHandle),
-				*GetFullNameSafe(NextFrame.StateTree),
+				*GetSafeStateName(*NextFrame, NextStateHandle),
+				*GetFullNameSafe(NextFrame->StateTree),
 				*GetNameSafe(&Owner),
 				*GetFullNameSafe(&RootStateTree)
 			);
@@ -6736,8 +6741,16 @@ bool FStateTreeExecutionContext::SelectStateInternal(
 		OutSelectionResult->TargetState = OutSelectionResult->SelectedStates.Last();
 	}
 
-	for (const FStateTreeSharedEvent& StateSelectionEvent : StateSelectionEvents)
+	for (int32 EventIndex = 0; EventIndex < StateSelectionEvents.Num(); ++EventIndex)
 	{
+		const FStateTreeSharedEvent& StateSelectionEvent = StateSelectionEvents[EventIndex];
+
+		// A SelectStateInternal_X might have changed the TemporaryFrames, re-cache the values
+		if (EventIndex != 0)
+		{
+			CacheNextFrame();
+		}
+
 		// Add the event for GetDataView and for consuming it later.
 		const bool bContainsSelectionEvent = OutSelectionResult->SelectionEvents.ContainsByPredicate(
 			[&LatestState = OutSelectionResult->SelectedStates.Last()](const FSelectionEventWithID& Other)
@@ -6767,7 +6780,7 @@ bool FStateTreeExecutionContext::SelectStateInternal(
 		if (bShouldPrerequisitesBeChecked)
 		{
 			// Check that the state can be entered
-			const bool bEnterConditionsPassed = TestAllConditionsWithValidation(NextParentFrame, NextFrame, NextStateHandle, NextState.EnterConditionEvaluationScopeMemoryRequirement, NextState.EnterConditionsBegin, NextState.EnterConditionsNum, EStateTreeUpdatePhase::EnterConditions);
+			const bool bEnterConditionsPassed = TestAllConditionsWithValidation(NextParentFrame, *NextFrame, NextStateHandle, NextState.EnterConditionEvaluationScopeMemoryRequirement, NextState.EnterConditionsBegin, NextState.EnterConditionsNum, EStateTreeUpdatePhase::EnterConditions);
 
 			if (!bEnterConditionsPassed)
 			{
@@ -6776,7 +6789,7 @@ bool FStateTreeExecutionContext::SelectStateInternal(
 			}
 		}
 
-		FSelectStateInternalArguments NewInternalArgs = FSelectStateInternalArguments{
+		const FSelectStateInternalArguments NewInternalArgs = FSelectStateInternalArguments{
 			.MissingActiveStates = bShouldCreateNewState ? TArrayView<const FActiveState>() : InternalArgs.MissingActiveStates.Mid(1),
 			.MissingSourceFrameID = FActiveFrameID(),
 			.MissingSourceStates = TArrayView<const FActiveState>(),
@@ -6785,12 +6798,12 @@ bool FStateTreeExecutionContext::SelectStateInternal(
 		if (NextState.Type == EStateTreeStateType::Linked)
 		{
 			// MissingStatesToReachTarget can include a linked state and the frame needs to be constructed (if needed).
-			bSucceededToSelectState = SelectStateInternal_Linked(Args, NewInternalArgs, OutSelectionResult, NextFrame.StateTree, NextState, bShouldCreateNewState);
+			bSucceededToSelectState = SelectStateInternal_Linked(Args, NewInternalArgs, OutSelectionResult, NextFrame->StateTree, NextState, bShouldCreateNewState);
 		}
 		else if (NextState.Type == EStateTreeStateType::LinkedAsset)
 		{
 			// MissingStatesToReachTarget can include a linked asset state and the frame needs to be constructed (if needed).
-			bSucceededToSelectState = SelectStateInternal_LinkedAsset(Args, NewInternalArgs, OutSelectionResult, NextFrame.StateTree, NextState, NextLinkedStateAsset, bShouldCreateNewState);
+			bSucceededToSelectState = SelectStateInternal_LinkedAsset(Args, NewInternalArgs, OutSelectionResult, NextFrame->StateTree, NextState, NextLinkedStateAsset, bShouldCreateNewState);
 		}
 		else if (!bIsTargetState)
 		{
@@ -6807,19 +6820,19 @@ bool FStateTreeExecutionContext::SelectStateInternal(
 				bSucceededToSelectState = true;
 				break;
 			case EStateTreeStateSelectionBehavior::TrySelectChildrenInOrder:
-				bSucceededToSelectState = SelectStateInternal_TrySelectChildrenInOrder(Args, NewInternalArgs, OutSelectionResult, NextFrame.StateTree, NextState, NextStateHandle);
+				bSucceededToSelectState = SelectStateInternal_TrySelectChildrenInOrder(Args, NewInternalArgs, OutSelectionResult, NextFrame->StateTree, NextState, NextStateHandle);
 				break;
 			case EStateTreeStateSelectionBehavior::TrySelectChildrenAtRandom:
-				bSucceededToSelectState = SelectStateInternal_TrySelectChildrenAtRandom(Args, NewInternalArgs, OutSelectionResult, NextFrame.StateTree, NextState, NextStateHandle);
+				bSucceededToSelectState = SelectStateInternal_TrySelectChildrenAtRandom(Args, NewInternalArgs, OutSelectionResult, NextFrame->StateTree, NextState, NextStateHandle);
 				break;
 			case EStateTreeStateSelectionBehavior::TrySelectChildrenWithHighestUtility:
-				bSucceededToSelectState = SelectStateInternal_TrySelectChildrenWithHighestUtility(Args, NewInternalArgs, OutSelectionResult, NextFrame.StateTree, NextState, NextStateHandle);
+				bSucceededToSelectState = SelectStateInternal_TrySelectChildrenWithHighestUtility(Args, NewInternalArgs, OutSelectionResult, NextFrame->StateTree, NextState, NextStateHandle);
 				break;
 			case EStateTreeStateSelectionBehavior::TrySelectChildrenAtRandomWeightedByUtility:
-				bSucceededToSelectState = SelectStateInternal_TrySelectChildrenAtRandomWeightedByUtility(Args, NewInternalArgs, OutSelectionResult, NextFrame.StateTree, NextState, NextStateHandle);
+				bSucceededToSelectState = SelectStateInternal_TrySelectChildrenAtRandomWeightedByUtility(Args, NewInternalArgs, OutSelectionResult, NextFrame->StateTree, NextState, NextStateHandle);
 				break;
 			case EStateTreeStateSelectionBehavior::TryFollowTransitions:
-				bSucceededToSelectState = SelectStateInternal_TryFollowTransitions(Args, NewInternalArgs, OutSelectionResult, NextFrame.StateTree, NextState, NextStateHandle);
+				bSucceededToSelectState = SelectStateInternal_TryFollowTransitions(Args, NewInternalArgs, OutSelectionResult, NextFrame->StateTree, NextState, NextStateHandle);
 				break;
 			}
 		}
