[StateTree] Fix bug with transition when a newly added root transition fails and the frame needs to be pop from the stack.

diff --git a/Engine/Plugins/Runtime/StateTree/Source/StateTreeModule/Private/StateTreeExecutionContext.cpp b/Engine/Plugins/Runtime/StateTree/Source/StateTreeModule/Private/StateTreeExecutionContext.cpp
--- a/Engine/Plugins/Runtime/StateTree/Source/StateTreeModule/Private/StateTreeExecutionContext.cpp
+++ b/Engine/Plugins/Runtime/StateTree/Source/StateTreeModule/Private/StateTreeExecutionContext.cpp
@@ -6446,6 +6446,30 @@ bool FStateTreeExecutionContext::SelectStateInternal(
 		return false;
 	}
 
+	struct FCleanUpOnExit
+	{
+		FCleanUpOnExit(const TSharedRef<FSelectStateResult>& InSelectionResult)
+			: SelectionResult(InSelectionResult)
+		{ }
+		~FCleanUpOnExit()
+		{
+			if (bStateAdded && bSucceededToSelectState == false)
+			{
+				SelectionResult->SelectedStates.Pop();
+			}
+			if (bFrameAdded && bSucceededToSelectState == false)
+			{
+				SelectionResult->SelectedFrames.Pop();
+			}
+		}
+
+		const TSharedRef<FSelectStateResult>& SelectionResult;
+		bool bSucceededToSelectState = false;
+		bool bFrameAdded = false;
+		bool bStateAdded = false;
+	};
+	FCleanUpOnExit OnExitScope = FCleanUpOnExit(OutSelectionResult);
+
 	// Does it need a new root tree. It can create a new frame or reused an existing one.
 	bool bNewFrameCreated = false;
 	{
@@ -6512,17 +6536,20 @@ bool FStateTreeExecutionContext::SelectStateInternal(
 				if (OutSelectionResult->SelectedFrames.Last() != NextFrameID)
 				{
 					OutSelectionResult->SelectedFrames.Add(NextFrameID);
+					OnExitScope.bFrameAdded = true;
 				}
 			}
 			else
 			{
 				OutSelectionResult->SelectedFrames.Add(NextFrameID);
+				OnExitScope.bFrameAdded = true;
 			}
 
 			TOptional<bool> bSelected = SelectStateFromSourceInternal(Args, InternalArgs, OutSelectionResult, *MissingFramePtr, NextState, NextStateHandle, bNewFrameCreated);
 			if (bSelected.IsSet())
 			{
-				return bSelected.GetValue();
+				OnExitScope.bSucceededToSelectState = bSelected.GetValue();
+				return OnExitScope.bSucceededToSelectState;
 			}
 		}
 	}
@@ -6538,11 +6565,11 @@ bool FStateTreeExecutionContext::SelectStateInternal(
 			NextFrame = FindExecutionFrame(NextFrameID, MakeConstArrayView(Exec.ActiveFrames), MakeConstArrayView(OutSelectionResult->TemporaryFrames));
 			check(NextFrame);
 
-			if (OutSelectionResult->SelectedFrames.Num() > 1)
-			{
-				NextParentFrame = FindExecutionFrame(OutSelectionResult->SelectedFrames.Last(1), MakeConstArrayView(Exec.ActiveFrames), MakeConstArrayView(OutSelectionResult->TemporaryFrames));
-				check(NextParentFrame);
-			}
+	if (OutSelectionResult->SelectedFrames.Num() > 1)
+	{
+		NextParentFrame = FindExecutionFrame(OutSelectionResult->SelectedFrames.Last(1), MakeConstArrayView(Exec.ActiveFrames), MakeConstArrayView(OutSelectionResult->TemporaryFrames));
+		check(NextParentFrame);
+	}
 		};
 	CacheNextFrame();
 
@@ -6694,18 +6721,11 @@ bool FStateTreeExecutionContext::SelectStateInternal(
 		bShouldCreateNewState = !bIsNextTargetStateInActiveStates;
 	}
 
-	bool bSucceededToSelectState = false;
-	ON_SCOPE_EXIT{
-		if (bSucceededToSelectState == false)
-		{
-			OutSelectionResult->SelectedStates.Pop();
-		}
-	};
-
 	// Add the state to the selected list. It will be removed if the selection fails.
 	if (bShouldCreateNewState)
 	{
 		OutSelectionResult->SelectedStates.Emplace(NextFrame->FrameID, FActiveStateID(Storage.GenerateUniqueId()), NextStateHandle);
+		OnExitScope.bStateAdded = true;
 
 		int32 StateInFrameCounter = 0;
 		for (int32 Index = OutSelectionResult->SelectedStates.Num() - 1; Index >= 0; --Index)
@@ -6733,6 +6753,7 @@ bool FStateTreeExecutionContext::SelectStateInternal(
 	{
 		check(InternalArgs.MissingActiveStates.Num() > 0); // bIsNextTargetStateInActiveStates is false
 		OutSelectionResult->SelectedStates.Add(InternalArgs.MissingActiveStates[0]);
+		OnExitScope.bStateAdded = true;
 	}
 
 	// Set the target, since this is recursive (can be called for linked state), it only set once the original target is found.
@@ -6798,18 +6819,18 @@ bool FStateTreeExecutionContext::SelectStateInternal(
 		if (NextState.Type == EStateTreeStateType::Linked)
 		{
 			// MissingStatesToReachTarget can include a linked state and the frame needs to be constructed (if needed).
-			bSucceededToSelectState = SelectStateInternal_Linked(Args, NewInternalArgs, OutSelectionResult, NextFrame->StateTree, NextState, bShouldCreateNewState);
+			OnExitScope.bSucceededToSelectState = SelectStateInternal_Linked(Args, NewInternalArgs, OutSelectionResult, NextFrame->StateTree, NextState, bShouldCreateNewState);
 		}
 		else if (NextState.Type == EStateTreeStateType::LinkedAsset)
 		{
 			// MissingStatesToReachTarget can include a linked asset state and the frame needs to be constructed (if needed).
-			bSucceededToSelectState = SelectStateInternal_LinkedAsset(Args, NewInternalArgs, OutSelectionResult, NextFrame->StateTree, NextState, NextLinkedStateAsset, bShouldCreateNewState);
+			OnExitScope.bSucceededToSelectState = SelectStateInternal_LinkedAsset(Args, NewInternalArgs, OutSelectionResult, NextFrame->StateTree, NextState, NextLinkedStateAsset, bShouldCreateNewState);
 		}
 		else if (!bIsTargetState)
 		{
 			check(NewInternalArgs.MissingStatesToReachTarget.Num() > 0);
 			// Next child state is already known. Passing TransitionEvent further, states selected directly by transition can use it.
-			bSucceededToSelectState = SelectStateInternal(Args, NewInternalArgs, OutSelectionResult);
+			OnExitScope.bSucceededToSelectState = SelectStateInternal(Args, NewInternalArgs, OutSelectionResult);
 		}
 		else if (Args.Behavior != ESelectStateBehavior::Forced)
 		{
@@ -6817,22 +6838,22 @@ bool FStateTreeExecutionContext::SelectStateInternal(
 			{
 			case EStateTreeStateSelectionBehavior::TryEnterState:
 				UE_STATETREE_DEBUG_STATE_EVENT(this, NextStateHandle, EStateTreeTraceEventType::OnStateSelected);
-				bSucceededToSelectState = true;
+				OnExitScope.bSucceededToSelectState = true;
 				break;
 			case EStateTreeStateSelectionBehavior::TrySelectChildrenInOrder:
-				bSucceededToSelectState = SelectStateInternal_TrySelectChildrenInOrder(Args, NewInternalArgs, OutSelectionResult, NextFrame->StateTree, NextState, NextStateHandle);
+				OnExitScope.bSucceededToSelectState = SelectStateInternal_TrySelectChildrenInOrder(Args, NewInternalArgs, OutSelectionResult, NextFrame->StateTree, NextState, NextStateHandle);
 				break;
 			case EStateTreeStateSelectionBehavior::TrySelectChildrenAtRandom:
-				bSucceededToSelectState = SelectStateInternal_TrySelectChildrenAtRandom(Args, NewInternalArgs, OutSelectionResult, NextFrame->StateTree, NextState, NextStateHandle);
+				OnExitScope.bSucceededToSelectState = SelectStateInternal_TrySelectChildrenAtRandom(Args, NewInternalArgs, OutSelectionResult, NextFrame->StateTree, NextState, NextStateHandle);
 				break;
 			case EStateTreeStateSelectionBehavior::TrySelectChildrenWithHighestUtility:
-				bSucceededToSelectState = SelectStateInternal_TrySelectChildrenWithHighestUtility(Args, NewInternalArgs, OutSelectionResult, NextFrame->StateTree, NextState, NextStateHandle);
+				OnExitScope.bSucceededToSelectState = SelectStateInternal_TrySelectChildrenWithHighestUtility(Args, NewInternalArgs, OutSelectionResult, NextFrame->StateTree, NextState, NextStateHandle);
 				break;
 			case EStateTreeStateSelectionBehavior::TrySelectChildrenAtRandomWeightedByUtility:
-				bSucceededToSelectState = SelectStateInternal_TrySelectChildrenAtRandomWeightedByUtility(Args, NewInternalArgs, OutSelectionResult, NextFrame->StateTree, NextState, NextStateHandle);
+				OnExitScope.bSucceededToSelectState = SelectStateInternal_TrySelectChildrenAtRandomWeightedByUtility(Args, NewInternalArgs, OutSelectionResult, NextFrame->StateTree, NextState, NextStateHandle);
 				break;
 			case EStateTreeStateSelectionBehavior::TryFollowTransitions:
-				bSucceededToSelectState = SelectStateInternal_TryFollowTransitions(Args, NewInternalArgs, OutSelectionResult, NextFrame->StateTree, NextState, NextStateHandle);
+				OnExitScope.bSucceededToSelectState = SelectStateInternal_TryFollowTransitions(Args, NewInternalArgs, OutSelectionResult, NextFrame->StateTree, NextState, NextStateHandle);
 				break;
 			}
 		}
@@ -6840,17 +6861,17 @@ bool FStateTreeExecutionContext::SelectStateInternal(
 		{
 			if (NewInternalArgs.MissingStatesToReachTarget.Num() == 0)
 			{
-				bSucceededToSelectState = true;
+				OnExitScope.bSucceededToSelectState = true;
 			}
 			else
 			{
 				// Continue the force selection. Next child state is already known.
-				bSucceededToSelectState = SelectStateInternal(Args, NewInternalArgs, OutSelectionResult);
+				OnExitScope.bSucceededToSelectState = SelectStateInternal(Args, NewInternalArgs, OutSelectionResult);
 			}
 		}
 
 
-		if (bSucceededToSelectState)
+		if (OnExitScope.bSucceededToSelectState)
 		{
 			break;
 		}
@@ -6860,7 +6881,7 @@ bool FStateTreeExecutionContext::SelectStateInternal(
 		}
 	}
 
-	return bSucceededToSelectState;
+	return OnExitScope.bSucceededToSelectState;
 }
 
 TOptional<bool> FStateTreeExecutionContext::SelectStateFromSourceInternal(
