[StateTree] Fix bug with transition when a newly added root transition fails and the frame needs to be pop from the stack.

diff --git a/Engine/Plugins/Runtime/StateTree/Source/StateTreeModule/Private/StateTreeExecutionContext.cpp b/Engine/Plugins/Runtime/StateTree/Source/StateTreeModule/Private/StateTreeExecutionContext.cpp
--- a/Engine/Plugins/Runtime/StateTree/Source/StateTreeModule/Private/StateTreeExecutionContext.cpp
+++ b/Engine/Plugins/Runtime/StateTree/Source/StateTreeModule/Private/StateTreeExecutionContext.cpp
@@ -6497,6 +6497,7 @@ bool FStateTreeExecutionContext::SelectStateInternal(
 				const FExecutionFrameKey FrameKey = FExecutionFrameKey(NextStateContext.StateTree, NextStateContext.StateHandle);
 				const FStateTreeExecutionFrame& NewFrame = OutSelectionResult->MakeAndAddTemporaryFrameWithNewRoot(FActiveFrameID(Storage.GenerateUniqueId()), FrameKey, *PreviousFramePtr);
 				OutSelectionResult->SelectedFrames.Add(NewFrame.FrameID);
+				OnExitScope.bFrameAdded = true;
 				bNewFrameCreated = true;
 			}
 		}
