[StateTree] Rename FExecutionFrameKey to FExecutionFrameHandle

#rnx

diff --git a/Engine/Plugins/Runtime/StateTree/Source/StateTreeModule/Private/StateTreeExecutionContext.cpp b/Engine/Plugins/Runtime/StateTree/Source/StateTreeModule/Private/StateTreeExecutionContext.cpp
--- a/Engine/Plugins/Runtime/StateTree/Source/StateTreeModule/Private/StateTreeExecutionContext.cpp
+++ b/Engine/Plugins/Runtime/StateTree/Source/StateTreeModule/Private/StateTreeExecutionContext.cpp
@@ -129,9 +129,10 @@ namespace UE::StateTree::ExecutionContext::Private
 			});
 	}
 
-	const FCompactStateTreeFrame* FindStateTreeFrame(const FExecutionFrameKey& LinkStateFrameKey)
+	const FCompactStateTreeFrame* FindStateTreeFrame(const FExecutionFrameHandle& FrameHandle)
 	{
-		const FCompactStateTreeFrame* FrameInfo = LinkStateFrameKey.StateTree->GetFrameFromHandle(LinkStateFrameKey.RootState);
+		check(FrameHandle.IsValid());
+		const FCompactStateTreeFrame* FrameInfo = FrameHandle.GetStateTree()->GetFrameFromHandle(FrameHandle.GetRootState());
 		ensureAlwaysMsgf(FrameInfo, TEXT("The compiled data is invalid. It should contains the information for the new root frame."));
 		return FrameInfo;
 	}
@@ -912,16 +913,18 @@ void FStateTreeMinimalExecutionContext::ScheduleNextTick(UE::StateTree::ETickRea
  */
 FStateTreeExecutionFrame& FStateTreeExecutionContext::FSelectStateResult::MakeAndAddTemporaryFrame(
 	const UE::StateTree::FActiveFrameID FrameID,
-	const UE::StateTree::FExecutionFrameKey& FrameKey,
+	const UE::StateTree::FExecutionFrameHandle& FrameHandle,
 	const bool bIsGlobalFrame)
 {
+	check(FrameHandle.IsValid());
+
 	FStateTreeExecutionFrame& ExecFrame = TemporaryFrames.AddDefaulted_GetRef();
 	ExecFrame.FrameID = FrameID;
-	ExecFrame.StateTree = FrameKey.StateTree;
-	ExecFrame.RootState = FrameKey.RootState;
+	ExecFrame.StateTree = FrameHandle.GetStateTree();
+	ExecFrame.RootState = FrameHandle.GetRootState();
 	ExecFrame.bIsGlobalFrame = bIsGlobalFrame;
 
-	const FCompactStateTreeFrame* TargetFrame = UE::StateTree::ExecutionContext::Private::FindStateTreeFrame(FrameKey);
+	const FCompactStateTreeFrame* TargetFrame = UE::StateTree::ExecutionContext::Private::FindStateTreeFrame(FrameHandle);
 	ensureMsgf(TargetFrame, TEXT("The compiled data is invalid. It should contains the information for the new root frame."));
 	ExecFrame.ActiveTasksStatus = TargetFrame ? FStateTreeTasksCompletionStatus(*TargetFrame) : FStateTreeTasksCompletionStatus();
 
@@ -930,11 +933,13 @@ FStateTreeExecutionFrame& FStateTreeExecutionContext::FSelectStateResult::MakeAn
 
 FStateTreeExecutionFrame& FStateTreeExecutionContext::FSelectStateResult::MakeAndAddTemporaryFrameWithNewRoot(
 	const UE::StateTree::FActiveFrameID FrameID,
-	const UE::StateTree::FExecutionFrameKey& FrameKey,
+	const UE::StateTree::FExecutionFrameHandle& FrameHandle,
 	FStateTreeExecutionFrame& OtherFrame)
 {
+	check(FrameHandle.IsValid());
+
 	constexpr bool bIsGlobalFrame = true;
-	FStateTreeExecutionFrame& ExecFrame = MakeAndAddTemporaryFrame(FrameID, FrameKey, bIsGlobalFrame);
+	FStateTreeExecutionFrame& ExecFrame = MakeAndAddTemporaryFrame(FrameID, FrameHandle, bIsGlobalFrame);
 	ExecFrame.ExternalDataBaseIndex = OtherFrame.ExternalDataBaseIndex;
 	ExecFrame.GlobalParameterDataHandle = OtherFrame.GlobalParameterDataHandle;
 	ExecFrame.GlobalInstanceIndexBase = OtherFrame.GlobalInstanceIndexBase;
@@ -1416,9 +1421,9 @@ EStateTreeRunStatus FStateTreeExecutionContext::Start(FStartParameters Parameter
 	TGuardValue<TSharedPtr<FSelectStateResult>> TemporaryStorageScope(CurrentlyProcessedTemporaryStorage, SelectStateResult);
 
 	const FActiveFrameID InitFrameID = FActiveFrameID(Storage.GenerateUniqueId());
-	const FExecutionFrameKey InitFrameKey = FExecutionFrameKey(&RootStateTree, FStateTreeStateHandle::Root);
+	const FExecutionFrameHandle InitFrameHandle = FExecutionFrameHandle(&RootStateTree, FStateTreeStateHandle::Root);
 	constexpr bool bIsGlobalFrame = true;
-	FStateTreeExecutionFrame& InitFrame = SelectStateResult->MakeAndAddTemporaryFrame(InitFrameID, InitFrameKey, bIsGlobalFrame);
+	FStateTreeExecutionFrame& InitFrame = SelectStateResult->MakeAndAddTemporaryFrame(InitFrameID, InitFrameHandle, bIsGlobalFrame);
 	InitFrame.GlobalParameterDataHandle = FStateTreeDataHandle(EStateTreeDataSourceType::GlobalParameterData);
 
 	if (!CollectActiveExternalData(SelectStateResult->TemporaryFrames))
@@ -1452,10 +1457,10 @@ EStateTreeRunStatus FStateTreeExecutionContext::Start(FStartParameters Parameter
 		Exec.LastTickStatus = EStateTreeRunStatus::Unset;
 
 		FSelectStateArguments SelectStateArgs;
-		SelectStateArgs.SourceState = FActiveState(InitFrameID, FActiveStateID(), InitFrameKey.RootState);
-		SelectStateArgs.TargetState = FStateHandleContext(InitFrameKey.StateTree, InitFrameKey.RootState);
+		SelectStateArgs.SourceState = FActiveState(InitFrameID, FActiveStateID(), InitFrameHandle.GetRootState());
+		SelectStateArgs.TargetState = FStateHandleContext(InitFrameHandle.GetStateTree(), InitFrameHandle.GetRootState());
 		SelectStateArgs.Behavior = ESelectStateBehavior::StateTransition;
-		SelectStateArgs.SelectionRules = InitFrameKey.StateTree->StateSelectionRules;
+		SelectStateArgs.SelectionRules = InitFrameHandle.GetStateTree()->StateSelectionRules;
 		if (SelectState(SelectStateArgs, SelectStateResult))
 		{
 			check(SelectStateResult->SelectedStates.Num() > 0);
@@ -1475,7 +1480,7 @@ EStateTreeRunStatus FStateTreeExecutionContext::Start(FStartParameters Parameter
 			{
 				// Enter state tasks can fail/succeed, treat it same as tick.
 				FStateTreeTransitionResult Transition;
-				Transition.TargetState = InitFrameKey.RootState;
+				Transition.TargetState = InitFrameHandle.GetRootState();
 				Transition.CurrentRunStatus = Exec.LastTickStatus;
 				const EStateTreeRunStatus LastTickStatus = EnterState(SelectStateResult, Transition);
 
@@ -6180,8 +6185,8 @@ bool FStateTreeExecutionContext::SelectState(const FSelectStateArguments& Args,
 	}
 	check(!PathToTargetState.IsEmpty());
 
-	const FExecutionFrameKey TargetFrameKey = FExecutionFrameKey(Args.TargetState.StateTree, PathToTargetState[0].StateHandle);
-	const FCompactStateTreeFrame* TargetFrame = TargetFrameKey.StateTree->GetFrameFromHandle(TargetFrameKey.RootState);
+	const FExecutionFrameHandle TargetFrameHandle = FExecutionFrameHandle(Args.TargetState.StateTree, PathToTargetState[0].StateHandle);
+	const FCompactStateTreeFrame* TargetFrame = TargetFrameHandle.GetStateTree()->GetFrameFromHandle(TargetFrameHandle.GetRootState());
 	checkf(TargetFrame, TEXT("The frame was not compiled."));
 
 	// Build the source path from the active path. Includes all the states from the previous frames.
@@ -6245,7 +6250,7 @@ bool FStateTreeExecutionContext::SelectState(const FSelectStateArguments& Args,
 	//Look in the SourceStates because you can select in a different frame up the tree but not down the tree.
 	FActiveFrameID TargetFrameID;
 	{
-		if (SourceExecFrame->HasRoot(TargetFrameKey))
+		if (SourceExecFrame->HasRoot(TargetFrameHandle))
 		{
 			TargetFrameID = SourceExecFrame->FrameID;
 
@@ -6284,10 +6289,10 @@ bool FStateTreeExecutionContext::SelectState(const FSelectStateArguments& Args,
 						{
 							return false;
 						}
-						if (ProcessFrame->StateTree == TargetFrameKey.StateTree)
+						if (ProcessFrame->StateTree == TargetFrameHandle.GetStateTree())
 						{
 							bFoundStateTree = true;
-							if (ProcessFrame->RootState == TargetFrameKey.RootState)
+							if (ProcessFrame->RootState == TargetFrameHandle.GetRootState())
 							{
 								bFoundRootState = true;
 							}
@@ -6494,8 +6499,8 @@ bool FStateTreeExecutionContext::SelectStateInternal(
 
 			if (PreviousSelectedState.GetStateHandle() != NextState.Parent)
 			{
-				const FExecutionFrameKey FrameKey = FExecutionFrameKey(NextStateContext.StateTree, NextStateContext.StateHandle);
-				const FStateTreeExecutionFrame& NewFrame = OutSelectionResult->MakeAndAddTemporaryFrameWithNewRoot(FActiveFrameID(Storage.GenerateUniqueId()), FrameKey, *PreviousFramePtr);
+				const FExecutionFrameHandle FrameHandle = FExecutionFrameHandle(NextStateContext.StateTree, NextStateContext.StateHandle);
+				const FStateTreeExecutionFrame& NewFrame = OutSelectionResult->MakeAndAddTemporaryFrameWithNewRoot(FActiveFrameID(Storage.GenerateUniqueId()), FrameHandle, *PreviousFramePtr);
 				OutSelectionResult->SelectedFrames.Add(NewFrame.FrameID);
 				OnExitScope.bFrameAdded = true;
 				bNewFrameCreated = true;
@@ -6517,16 +6522,16 @@ bool FStateTreeExecutionContext::SelectStateInternal(
 			{
 				if (!NextState.Parent.IsValid() && MissingFramePtr->RootState != NextStateHandle)
 				{
-					const FExecutionFrameKey FrameKey = FExecutionFrameKey(NextStateContext.StateTree, NextStateContext.StateHandle);
-					const FStateTreeExecutionFrame& NewFrame = OutSelectionResult->MakeAndAddTemporaryFrameWithNewRoot(FActiveFrameID(Storage.GenerateUniqueId()), FrameKey, *MissingFramePtr);
+					const FExecutionFrameHandle FrameHandle = FExecutionFrameHandle(NextStateContext.StateTree, NextStateContext.StateHandle);
+					const FStateTreeExecutionFrame& NewFrame = OutSelectionResult->MakeAndAddTemporaryFrameWithNewRoot(FActiveFrameID(Storage.GenerateUniqueId()), FrameHandle, *MissingFramePtr);
 					NextFrameID = NewFrame.FrameID;
 					bNewFrameCreated = true;
 				}
 			}
 			else if (MissingFramePtr->RootState != NextStateHandle)
 			{
-				const FExecutionFrameKey FrameKey = FExecutionFrameKey(NextStateContext.StateTree, NextStateContext.StateHandle);
-				const FStateTreeExecutionFrame& NewFrame = OutSelectionResult->MakeAndAddTemporaryFrameWithNewRoot(FActiveFrameID(Storage.GenerateUniqueId()), FrameKey, *MissingFramePtr);
+				const FExecutionFrameHandle FrameHandle = FExecutionFrameHandle(NextStateContext.StateTree, NextStateContext.StateHandle);
+				const FStateTreeExecutionFrame& NewFrame = OutSelectionResult->MakeAndAddTemporaryFrameWithNewRoot(FActiveFrameID(Storage.GenerateUniqueId()), FrameHandle, *MissingFramePtr);
 				NextFrameID = NewFrame.FrameID;
 				bNewFrameCreated = true;
 			}
@@ -6975,27 +6980,29 @@ TOptional<bool> FStateTreeExecutionContext::SelectStateFromSourceInternal(
 
 namespace UE::StateTree::ExecutionContext::Private
 {
-	bool PreventRecursionCheck(const FExecutionFrameKey& LinkStateFrameKey,
+	bool PreventRecursionCheck(const FExecutionFrameHandle& LinkStateFrameHandle,
 		TArrayView<const FActiveFrameID> SelectedFrames,
 		TArrayView<const FStateTreeExecutionFrame> ActiveFrames,
 		TArrayView<const FStateTreeExecutionFrame> TemporaryFrames)
 	{
 		// Check and prevent recursion.
 		return SelectedFrames.ContainsByPredicate(
-			[&LinkStateFrameKey, &ActiveFrames, &TemporaryFrames](const FActiveFrameID& FrameID)
+			[&LinkStateFrameHandle, &ActiveFrames, &TemporaryFrames](const FActiveFrameID& FrameID)
 			{
 				const FStateTreeExecutionFrame* Frame = FindExecutionFrame(FrameID, ActiveFrames, TemporaryFrames);
-				return Frame && Frame->HasRoot(LinkStateFrameKey);
+				return Frame && Frame->HasRoot(LinkStateFrameHandle);
 			});
 	}
 
 	FStateTreeExecutionFrame* SelectedFrameLinkedFrame(FStateTreeExecutionState& Exec,
 		const bool bShouldCreateNewState,
 		TArrayView<const UE::StateTree::FActiveState> MatchingActiveStates,
-		const FExecutionFrameKey& LinkStateFrameKey,
+		const FExecutionFrameHandle& LinkStateFrameHandle,
 		const EStateTreeStateSelectionRules StateSelectionRules
 	)
 	{
+		check(LinkStateFrameHandle.IsValid());
+
 		FStateTreeExecutionFrame* Result = nullptr;
 		if (!bShouldCreateNewState)
 		{
@@ -7005,7 +7012,7 @@ namespace UE::StateTree::ExecutionContext::Private
 				const FActiveFrameID ExistingFrameID = MatchingActiveStates[0].GetFrameID();
 				Result = Exec.FindActiveFrame(ExistingFrameID);
 				check(Result);
-				if (Result->HasRoot(LinkStateFrameKey))
+				if (Result->HasRoot(LinkStateFrameHandle))
 				{
 					if (EnumHasAllFlags(StateSelectionRules, EStateTreeStateSelectionRules::CompletedTransitionStatesCreateNewStates)
 						&& Result->ActiveTasksStatus.GetStatus(Result->StateTree).IsCompleted()
@@ -7045,15 +7052,15 @@ bool FStateTreeExecutionContext::SelectStateInternal_Linked(
 		return false;
 	}
 
-	const FExecutionFrameKey LinkStateFrameKey = FExecutionFrameKey(NextStateTree, TargetState.LinkedState);
+	const FExecutionFrameHandle LinkStateFrameHandle = FExecutionFrameHandle(NextStateTree, TargetState.LinkedState);
 
 	const bool bHasMissingState = InternalArgs.MissingStatesToReachTarget.Num() > 0 && Args.Behavior == ESelectStateBehavior::Forced;
 	if (bHasMissingState)
 	{
 		// In a force transition, the root could be different from what is expected.
 		//Ex: a previous transition go to root, then another transition go to another top level state (new root)
-		if (InternalArgs.MissingStatesToReachTarget[0].StateTree != LinkStateFrameKey.StateTree
-			|| InternalArgs.MissingStatesToReachTarget[0].StateHandle != LinkStateFrameKey.RootState)
+		if (InternalArgs.MissingStatesToReachTarget[0].StateTree != LinkStateFrameHandle.GetStateTree()
+			|| InternalArgs.MissingStatesToReachTarget[0].StateHandle != LinkStateFrameHandle.GetRootState())
 		{
 			STATETREE_LOG(Error, TEXT("%hs: The missing state is not from the same state tree. '%s' using StateTree '%s'."),
 				__FUNCTION__,
@@ -7064,11 +7071,11 @@ bool FStateTreeExecutionContext::SelectStateInternal_Linked(
 		}
 	}
 
-	if (PreventRecursionCheck(LinkStateFrameKey, OutSelectionResult->SelectedFrames, MakeConstArrayView(Exec.ActiveFrames), MakeConstArrayView(OutSelectionResult->TemporaryFrames)))
+	if (PreventRecursionCheck(LinkStateFrameHandle, OutSelectionResult->SelectedFrames, MakeConstArrayView(Exec.ActiveFrames), MakeConstArrayView(OutSelectionResult->TemporaryFrames)))
 	{
 		STATETREE_LOG(Error, TEXT("%hs: Trying to recursively enter subtree '%s' from '%s'. '%s' using StateTree '%s'."),
 			__FUNCTION__,
-			*GetSafeStateName(LinkStateFrameKey.StateTree, LinkStateFrameKey.RootState),
+			*GetSafeStateName(LinkStateFrameHandle.GetStateTree(), LinkStateFrameHandle.GetRootState()),
 			*GetStateStatusString(Exec),
 			*Owner.GetName(),
 			*NextStateTree->GetFullName()
@@ -7076,12 +7083,12 @@ bool FStateTreeExecutionContext::SelectStateInternal_Linked(
 		return false;
 	}
 
-	const FCompactStateTreeFrame* LinkStateTreeFrame = FindStateTreeFrame(LinkStateFrameKey);
+	const FCompactStateTreeFrame* LinkStateTreeFrame = FindStateTreeFrame(LinkStateFrameHandle);
 	if (LinkStateTreeFrame == nullptr)
 	{
 		STATETREE_LOG(Error, TEXT("%hs: The frame '%s' from '%s' does not exist. '%s' using StateTree '%s'."),
 			__FUNCTION__,
-			*GetSafeStateName(LinkStateFrameKey.StateTree, LinkStateFrameKey.RootState),
+			*GetSafeStateName(LinkStateFrameHandle.GetStateTree(), LinkStateFrameHandle.GetRootState()),
 			*GetStateStatusString(Exec),
 			*Owner.GetName(),
 			*NextStateTree->GetFullName()
@@ -7090,8 +7097,8 @@ bool FStateTreeExecutionContext::SelectStateInternal_Linked(
 	}
 
 	// Do we have an existing frame.
-	FStateTreeExecutionFrame* SelectedFrame = SelectedFrameLinkedFrame(Exec, bShouldCreateNewState, InternalArgs.MissingActiveStates, LinkStateFrameKey, Args.SelectionRules);
-	if (SelectedFrame && !SelectedFrame->HasRoot(LinkStateFrameKey))
+	FStateTreeExecutionFrame* SelectedFrame = SelectedFrameLinkedFrame(Exec, bShouldCreateNewState, InternalArgs.MissingActiveStates, LinkStateFrameHandle, Args.SelectionRules);
+	if (SelectedFrame && !SelectedFrame->HasRoot(LinkStateFrameHandle))
 	{
 		if (Args.Behavior == ESelectStateBehavior::Forced)
 		{
@@ -7101,7 +7108,7 @@ bool FStateTreeExecutionContext::SelectStateInternal_Linked(
 		{
 			STATETREE_LOG(Error, TEXT("%hs: The frame '%s' from '%s' does not have the same root as the active frame. '%s' using StateTree '%s'."),
 				__FUNCTION__,
-				*GetSafeStateName(LinkStateFrameKey.StateTree, LinkStateFrameKey.RootState),
+				*GetSafeStateName(LinkStateFrameHandle.GetStateTree(), LinkStateFrameHandle.GetRootState()),
 				*GetStateStatusString(Exec),
 				*Owner.GetName(),
 				*NextStateTree->GetFullName()
@@ -7126,7 +7133,7 @@ bool FStateTreeExecutionContext::SelectStateInternal_Linked(
 		}
 
 		constexpr bool bIsGlobalFrame = false;
-		FStateTreeExecutionFrame& NewFrame = OutSelectionResult->MakeAndAddTemporaryFrame(FActiveFrameID(Storage.GenerateUniqueId()), LinkStateFrameKey, bIsGlobalFrame);
+		FStateTreeExecutionFrame& NewFrame = OutSelectionResult->MakeAndAddTemporaryFrame(FActiveFrameID(Storage.GenerateUniqueId()), LinkStateFrameHandle, bIsGlobalFrame);
 		NewFrame.ExternalDataBaseIndex = ExternalDataBaseIndex;
 		NewFrame.GlobalParameterDataHandle = GlobalParameterDataHandle;
 		NewFrame.GlobalInstanceIndexBase = GlobalInstanceIndexBase;
@@ -7138,7 +7145,7 @@ bool FStateTreeExecutionContext::SelectStateInternal_Linked(
 	OutSelectionResult->SelectedFrames.Add(SelectedFrame->FrameID);
 
 	// Select the root state of the new frame.
-	const FStateHandleContext RootState = FStateHandleContext(LinkStateFrameKey.StateTree, LinkStateFrameKey.RootState);
+	const FStateHandleContext RootState = FStateHandleContext(LinkStateFrameHandle.GetStateTree(), LinkStateFrameHandle.GetRootState());
 	const FSelectStateInternalArguments NewInternalArgs = FSelectStateInternalArguments{
 		.MissingActiveStates = bIsNewFrame ? TArrayView<const FActiveState>() : InternalArgs.MissingActiveStates,
 		.MissingSourceFrameID = SelectedFrame->FrameID,
@@ -7204,7 +7211,7 @@ bool FStateTreeExecutionContext::SelectStateInternal_LinkedAsset(
 		NextLinkedStateRoot = InternalArgs.MissingStatesToReachTarget[0].StateHandle;
 	}
 
-	const FExecutionFrameKey LinkStateFrameKey = FExecutionFrameKey(NextLinkedStateAsset, NextLinkedStateRoot);
+	const FExecutionFrameHandle LinkStateFrameHandle = FExecutionFrameHandle(NextLinkedStateAsset, NextLinkedStateRoot);
 
 	// The linked state tree should have compatible context requirements.
 	if (!NextLinkedStateAsset->HasCompatibleContextData(RootStateTree)
@@ -7213,7 +7220,7 @@ bool FStateTreeExecutionContext::SelectStateInternal_LinkedAsset(
 		STATETREE_LOG(Error, TEXT("%hs: The linked State Tree '%s' does not have compatible schema, trying to select state %s from '%s'. '%s' using StateTree '%s'."),
 			__FUNCTION__,
 			*GetFullNameSafe(NextLinkedStateAsset),
-			*GetSafeStateName(LinkStateFrameKey.StateTree, LinkStateFrameKey.RootState),
+			*GetSafeStateName(LinkStateFrameHandle.GetStateTree(), LinkStateFrameHandle.GetRootState()),
 			*GetStateStatusString(Exec),
 			*Owner.GetName(),
 			*NextStateTree->GetFullName()
@@ -7221,11 +7228,11 @@ bool FStateTreeExecutionContext::SelectStateInternal_LinkedAsset(
 		return false;
 	}
 
-	if (PreventRecursionCheck(LinkStateFrameKey, OutSelectionResult->SelectedFrames, MakeConstArrayView(Exec.ActiveFrames), MakeConstArrayView(OutSelectionResult->TemporaryFrames)))
+	if (PreventRecursionCheck(LinkStateFrameHandle, OutSelectionResult->SelectedFrames, MakeConstArrayView(Exec.ActiveFrames), MakeConstArrayView(OutSelectionResult->TemporaryFrames)))
 	{
 		STATETREE_LOG(Error, TEXT("%hs: Trying to recursively enter subtree '%s' from '%s'. '%s' using StateTree '%s'."),
 			__FUNCTION__,
-			*GetSafeStateName(LinkStateFrameKey.StateTree, LinkStateFrameKey.RootState),
+			*GetSafeStateName(LinkStateFrameHandle.GetStateTree(), LinkStateFrameHandle.GetRootState()),
 			*GetStateStatusString(Exec),
 			*Owner.GetName(),
 			*NextStateTree->GetFullName()
@@ -7233,12 +7240,12 @@ bool FStateTreeExecutionContext::SelectStateInternal_LinkedAsset(
 		return false;
 	}
 
-	const FCompactStateTreeFrame* LinkStateTreeFrame = FindStateTreeFrame(LinkStateFrameKey);
+	const FCompactStateTreeFrame* LinkStateTreeFrame = FindStateTreeFrame(LinkStateFrameHandle);
 	if (LinkStateTreeFrame == nullptr)
 	{
 		STATETREE_LOG(Error, TEXT("%hs: The frame '%s' from '%s' does not exist. '%s' using StateTree '%s'."),
 			__FUNCTION__,
-			*GetSafeStateName(LinkStateFrameKey.StateTree, LinkStateFrameKey.RootState),
+			*GetSafeStateName(LinkStateFrameHandle.GetStateTree(), LinkStateFrameHandle.GetRootState()),
 			*GetStateStatusString(Exec),
 			*Owner.GetName(),
 			*NextStateTree->GetFullName()
@@ -7248,9 +7255,9 @@ bool FStateTreeExecutionContext::SelectStateInternal_LinkedAsset(
 
 	// Do we have an existing frame.
 	// Do not use the transition override selection rules. A transition outside the frame should not impact the current frame.
-	const EStateTreeStateSelectionRules StateSelectionRules = LinkStateFrameKey.StateTree->GetStateSelectionRules();
-	FStateTreeExecutionFrame* SelectedFrame = SelectedFrameLinkedFrame(Exec, bShouldCreateNewState, InternalArgs.MissingActiveStates, LinkStateFrameKey, StateSelectionRules);
-	if (SelectedFrame && !SelectedFrame->HasRoot(LinkStateFrameKey))
+	const EStateTreeStateSelectionRules StateSelectionRules = LinkStateFrameHandle.GetStateTree()->GetStateSelectionRules();
+	FStateTreeExecutionFrame* SelectedFrame = SelectedFrameLinkedFrame(Exec, bShouldCreateNewState, InternalArgs.MissingActiveStates, LinkStateFrameHandle, StateSelectionRules);
+	if (SelectedFrame && !SelectedFrame->HasRoot(LinkStateFrameHandle))
 	{
 		if (Args.Behavior == ESelectStateBehavior::Forced)
 		{
@@ -7260,7 +7267,7 @@ bool FStateTreeExecutionContext::SelectStateInternal_LinkedAsset(
 		{
 			STATETREE_LOG(Error, TEXT("%hs: The frame '%s' from '%s' does not have the same root as the active frame. '%s' using StateTree '%s'."),
 				__FUNCTION__,
-				*GetSafeStateName(LinkStateFrameKey.StateTree, LinkStateFrameKey.RootState),
+				*GetSafeStateName(LinkStateFrameHandle.GetStateTree(), LinkStateFrameHandle.GetRootState()),
 				*GetStateStatusString(Exec),
 				*Owner.GetName(),
 				*NextStateTree->GetFullName()
@@ -7273,12 +7280,12 @@ bool FStateTreeExecutionContext::SelectStateInternal_LinkedAsset(
 	if (bIsNewFrame)
 	{
 		// Collect external data if needed
-		const FStateTreeIndex16 ExternalDataBaseIndex = CollectExternalData(LinkStateFrameKey.StateTree);
+		const FStateTreeIndex16 ExternalDataBaseIndex = CollectExternalData(LinkStateFrameHandle.GetStateTree());
 		if (!ExternalDataBaseIndex.IsValid())
 		{
 			STATETREE_LOG(VeryVerbose, TEXT("%hs: Cannot select state '%s' because failed to collect external data for nested tree '%s' from '%s'. '%s' using StateTree '%s'."),
 				__FUNCTION__,
-				*GetSafeStateName(LinkStateFrameKey.StateTree, LinkStateFrameKey.RootState),
+				*GetSafeStateName(LinkStateFrameHandle.GetStateTree(), LinkStateFrameHandle.GetRootState()),
 				*GetFullNameSafe(NextLinkedStateAsset),
 				*GetStateStatusString(Exec),
 				*Owner.GetName(),
@@ -7287,12 +7294,12 @@ bool FStateTreeExecutionContext::SelectStateInternal_LinkedAsset(
 			return false;
 		}
 
-		FStateTreeExecutionFrame& NewFrame = OutSelectionResult->MakeAndAddTemporaryFrame(FActiveFrameID(Storage.GenerateUniqueId()), LinkStateFrameKey, true);
+		FStateTreeExecutionFrame& NewFrame = OutSelectionResult->MakeAndAddTemporaryFrame(FActiveFrameID(Storage.GenerateUniqueId()), LinkStateFrameHandle, true);
 		NewFrame.ExternalDataBaseIndex = ExternalDataBaseIndex;
 		// Pass the linked state's parameters as global parameters to the linked asset.
 		NewFrame.GlobalParameterDataHandle = NextState.ParameterDataHandle;
 		// The state parameters will be from the root state.
-		const FCompactStateTreeState& NewFrameRootState = LinkStateFrameKey.StateTree->States[NewFrame.RootState.Index];
+		const FCompactStateTreeState& NewFrameRootState = LinkStateFrameHandle.GetStateTree()->States[NewFrame.RootState.Index];
 		NewFrame.StateParameterDataHandle = NewFrameRootState.ParameterDataHandle;
 
 		SelectedFrame = &NewFrame;
@@ -7313,7 +7320,7 @@ bool FStateTreeExecutionContext::SelectStateInternal_LinkedAsset(
 		{
 			STATETREE_LOG(VeryVerbose, TEXT("%hs: Cannot select state '%s' because cannot start nested tree's '%s' global tasks and evaluators. '%s' using StateTree '%s'."),
 				__FUNCTION__,
-				*GetSafeStateName(LinkStateFrameKey.StateTree, LinkStateFrameKey.RootState),
+				*GetSafeStateName(LinkStateFrameHandle.GetStateTree(), LinkStateFrameHandle.GetRootState()),
 				*GetFullNameSafe(NextLinkedStateAsset),
 				*Owner.GetName(),
 				*NextStateTree->GetFullName()
@@ -7328,7 +7335,7 @@ bool FStateTreeExecutionContext::SelectStateInternal_LinkedAsset(
 	}
 
 	// Select the root state of the new frame.
-	const FStateHandleContext RootState = FStateHandleContext(LinkStateFrameKey.StateTree, LinkStateFrameKey.RootState);
+	const FStateHandleContext RootState = FStateHandleContext(LinkStateFrameHandle.GetStateTree(), LinkStateFrameHandle.GetRootState());
 	FSelectStateArguments NewArgs = Args;
 	NewArgs.SelectionRules = StateSelectionRules;
 	const FSelectStateInternalArguments NewInternalArgs = FSelectStateInternalArguments{
diff --git a/Engine/Plugins/Runtime/StateTree/Source/StateTreeModule/Public/StateTreeExecutionContext.h b/Engine/Plugins/Runtime/StateTree/Source/StateTreeModule/Public/StateTreeExecutionContext.h
--- a/Engine/Plugins/Runtime/StateTree/Source/StateTreeModule/Public/StateTreeExecutionContext.h
+++ b/Engine/Plugins/Runtime/StateTree/Source/StateTreeModule/Public/StateTreeExecutionContext.h
@@ -796,8 +796,8 @@ protected:
 		UE::StateTree::FActiveState TargetState;
 
 		/** @return a new temporary frame */
-		UE_API FStateTreeExecutionFrame& MakeAndAddTemporaryFrame(UE::StateTree::FActiveFrameID FrameID, const UE::StateTree::FExecutionFrameKey& FrameKey, bool bIsGlobalFrame);
-		UE_API FStateTreeExecutionFrame& MakeAndAddTemporaryFrameWithNewRoot(UE::StateTree::FActiveFrameID FrameID, const UE::StateTree::FExecutionFrameKey& FrameKey, FStateTreeExecutionFrame& OtherFrame);
+		UE_API FStateTreeExecutionFrame& MakeAndAddTemporaryFrame(UE::StateTree::FActiveFrameID FrameID, const UE::StateTree::FExecutionFrameHandle& FrameHandle, bool bIsGlobalFrame);
+		UE_API FStateTreeExecutionFrame& MakeAndAddTemporaryFrameWithNewRoot(UE::StateTree::FActiveFrameID FrameID, const UE::StateTree::FExecutionFrameHandle& FrameHandle, FStateTreeExecutionFrame& OtherFrame);
 
 		//~ITemporaryStorage
 		virtual FFrameAndParent GetExecutionFrame(UE::StateTree::FActiveFrameID ID) override;
diff --git a/Engine/Plugins/Runtime/StateTree/Source/StateTreeModule/Public/StateTreeExecutionTypes.h b/Engine/Plugins/Runtime/StateTree/Source/StateTreeModule/Public/StateTreeExecutionTypes.h
--- a/Engine/Plugins/Runtime/StateTree/Source/StateTreeModule/Public/StateTreeExecutionTypes.h
+++ b/Engine/Plugins/Runtime/StateTree/Source/StateTreeModule/Public/StateTreeExecutionTypes.h
@@ -939,15 +939,39 @@ private:
 
 namespace UE::StateTree
 {
-	/** Helper that identifies an execution frame with the root state handle and its tree asset. */
-	struct FExecutionFrameKey
+	/** Helper that identifies an execution frame with the root state and its tree asset. */
+	USTRUCT()
+	struct FExecutionFrameHandle
 	{
-		FExecutionFrameKey(TNotNull<const UStateTree*> InStateTree, FStateTreeStateHandle InRootState)
+		GENERATED_BODY()
+
+		FExecutionFrameHandle() = default;
+		FExecutionFrameHandle(TNotNull<const UStateTree*> InStateTree, FStateTreeStateHandle InRootState)
 			: StateTree(InStateTree)
 			, RootState(InRootState)
 		{
 		}
-		TNotNull<const UStateTree*> StateTree;
+
+		bool IsValid() const
+		{
+			return RootState.IsValid() && StateTree != nullptr;
+		}
+
+		const UStateTree* GetStateTree() const
+		{
+			return StateTree;
+		}
+
+		FStateTreeStateHandle GetRootState() const
+		{
+			return RootState;
+		}
+
+private:
+		UPROPERTY()
+		TObjectPtr<const UStateTree> StateTree;
+
+		UPROPERTY()
 		FStateTreeStateHandle RootState;
 	};
 }
@@ -981,9 +1005,9 @@ PRAGMA_ENABLE_DEPRECATION_WARNINGS
 	}
 
 	/** @return wherever the frame uses the state tree and root state. */
-	bool HasRoot(const UE::StateTree::FExecutionFrameKey& FrameKey) const
+	bool HasRoot(const UE::StateTree::FExecutionFrameHandle& FrameHandle) const
 	{
-		return StateTree == FrameKey.StateTree && RootState == FrameKey.RootState;
+		return StateTree == FrameHandle.GetStateTree() && RootState == FrameHandle.GetRootState();
 	}
 
 	/** @return wherever the frame uses the state tree and root state. */
