[StateTree] Add ExecutionRuntimeData. A node can request data to be always available for the lifetime of the state tree instance. It's optional, and only nodes that request it will get that feature. The data is kept alive; be careful about garbage collection issues. To request data, a node extends the GetExecutionRuntimeDataType function. It can access the data via the execution context GetExecutionRuntimeData() function.
Update the editor node on load. If the instance data type or the execution runtime data type changes, then the node is recreated with the new data type.

diff --git a/Engine/Plugins/Runtime/StateTree/Source/StateTreeEditorModule/Private/AsyncStateTreeDiff.cpp b/Engine/Plugins/Runtime/StateTree/Source/StateTreeEditorModule/Private/AsyncStateTreeDiff.cpp
--- a/Engine/Plugins/Runtime/StateTree/Source/StateTreeEditorModule/Private/AsyncStateTreeDiff.cpp
+++ b/Engine/Plugins/Runtime/StateTree/Source/StateTreeEditorModule/Private/AsyncStateTreeDiff.cpp
@@ -42,8 +42,10 @@ static bool AreObjectsEqual(const UObject* ObjectA, const UObject* ObjectB)
 static bool AreNodesEqual(const FStateTreeEditorNode& NodeA, const FStateTreeEditorNode& NodeB)
 {
 	return AreObjectsEqual(NodeA.InstanceObject.Get(), NodeB.InstanceObject.Get())
+		&& AreObjectsEqual(NodeA.ExecutionRuntimeDataObject.Get(), NodeB.ExecutionRuntimeDataObject.Get())
 		&& NodeA.Node.Identical(&NodeB.Node, PPF_DeepComparison | PPF_ForDiff)
 		&& NodeA.Instance.Identical(&NodeB.Instance, PPF_DeepComparison | PPF_ForDiff)
+		&& NodeA.ExecutionRuntimeData.Identical(&NodeB.ExecutionRuntimeData, PPF_DeepComparison | PPF_ForDiff)
 		&& NodeA.ExpressionIndent == NodeB.ExpressionIndent
 		&& NodeA.ExpressionOperand == NodeB.ExpressionOperand;
 }
diff --git a/Engine/Plugins/Runtime/StateTree/Source/StateTreeEditorModule/Private/Customizations/StateTreeEditorNodeDetails.cpp b/Engine/Plugins/Runtime/StateTree/Source/StateTreeEditorModule/Private/Customizations/StateTreeEditorNodeDetails.cpp
--- a/Engine/Plugins/Runtime/StateTree/Source/StateTreeEditorModule/Private/Customizations/StateTreeEditorNodeDetails.cpp
+++ b/Engine/Plugins/Runtime/StateTree/Source/StateTreeEditorModule/Private/Customizations/StateTreeEditorNodeDetails.cpp
@@ -361,6 +361,8 @@ void FStateTreeEditorNodeDetails::CustomizeHeader(TSharedRef<class IPropertyHand
 	NodeProperty = StructProperty->GetChildHandle(GET_MEMBER_NAME_CHECKED(FStateTreeEditorNode, Node));
 	InstanceProperty = StructProperty->GetChildHandle(GET_MEMBER_NAME_CHECKED(FStateTreeEditorNode, Instance));
 	InstanceObjectProperty = StructProperty->GetChildHandle(GET_MEMBER_NAME_CHECKED(FStateTreeEditorNode, InstanceObject));
+	ExecutionRuntimeDataProperty = StructProperty->GetChildHandle(GET_MEMBER_NAME_CHECKED(FStateTreeEditorNode, ExecutionRuntimeData));
+	ExecutionRuntimeDataObjectProperty = StructProperty->GetChildHandle(GET_MEMBER_NAME_CHECKED(FStateTreeEditorNode, ExecutionRuntimeDataObject));
 	IDProperty = StructProperty->GetChildHandle(GET_MEMBER_NAME_CHECKED(FStateTreeEditorNode, ID));
 
 	IndentProperty = StructProperty->GetChildHandle(GET_MEMBER_NAME_CHECKED(FStateTreeEditorNode, ExpressionIndent));
@@ -368,6 +370,9 @@ void FStateTreeEditorNodeDetails::CustomizeHeader(TSharedRef<class IPropertyHand
 
 	check(NodeProperty.IsValid());
 	check(InstanceProperty.IsValid());
+	check(InstanceObjectProperty.IsValid());
+	check(ExecutionRuntimeDataProperty.IsValid());
+	check(ExecutionRuntimeDataObjectProperty.IsValid());
 	check(IDProperty.IsValid());
 	check(IndentProperty.IsValid());
 	check(OperandProperty.IsValid());
@@ -981,6 +986,50 @@ void FStateTreeEditorNodeDetails::CustomizeChildren(TSharedRef<class IPropertyHa
 	}
 	
 	UStateTreeEditorData* EditorDataPtr = EditorData.Get();
+	auto AddObjectInstance = [this, &StructBuilder, EditorNodePtr, EditorDataPtr](TSharedPtr<IPropertyHandle> ValueProperty)
+		{
+			if (ValueProperty.IsValid())
+			{
+				uint32 NumChildren = 0;
+				ValueProperty->GetNumChildren(NumChildren);
+
+				// Find visible child properties and sort them so in order: Context, Input, Param, Output.
+				struct FSortedChild
+				{
+					TSharedPtr<IPropertyHandle> PropertyHandle;
+					EStateTreePropertyUsage Usage = EStateTreePropertyUsage::Invalid;
+				};
+
+				TArray<FSortedChild> SortedChildren;
+				for (uint32 Index = 0; Index < NumChildren; Index++)
+				{
+					if (TSharedPtr<IPropertyHandle> ChildHandle = ValueProperty->GetChildHandle(Index); ChildHandle.IsValid())
+					{
+						FSortedChild Child;
+						Child.PropertyHandle = ChildHandle;
+						Child.Usage = UE::StateTree::GetUsageFromMetaData(Child.PropertyHandle->GetProperty());
+
+						// If the property is set to one of these usages, display it even if it is not edit on instance.
+						// It is a common mistake to forget to set the "eye" on these properties it and wonder why it does not show up.
+						const bool bShouldShowByUsage = Child.Usage == EStateTreePropertyUsage::Input || Child.Usage == EStateTreePropertyUsage::Output || Child.Usage == EStateTreePropertyUsage::Context;
+        				const bool bIsEditable = !Child.PropertyHandle->GetProperty()->HasAllPropertyFlags(CPF_DisableEditOnInstance);
+
+						if (bShouldShowByUsage || bIsEditable)
+						{
+							SortedChildren.Add(Child);
+						}
+					}
+				}
+
+				SortedChildren.StableSort([](const FSortedChild& LHS, const FSortedChild& RHS) { return LHS.Usage < RHS.Usage; });
+
+				for (FSortedChild& Child : SortedChildren)
+				{
+					IDetailPropertyRow& ChildRow = StructBuilder.AddProperty(Child.PropertyHandle.ToSharedRef());
+					UE::StateTreeEditor::Internal::ModifyRow(ChildRow, EditorNodePtr->ID, EditorDataPtr);
+				}
+			}
+		};
 
 	// Node
 	TSharedRef<FBindableNodeInstanceDetails> NodeDetails = MakeShareable(new FBindableNodeInstanceDetails(NodeProperty, EditorNodePtr->GetNodeID(), EditorDataPtr));
@@ -993,47 +1042,17 @@ void FStateTreeEditorNodeDetails::CustomizeChildren(TSharedRef<class IPropertyHa
 	// InstanceObject
 	// Get the actual UObject from the pointer.
 	TSharedPtr<IPropertyHandle> InstanceObjectValueProperty = GetInstancedObjectValueHandle(InstanceObjectProperty);
-	if (InstanceObjectValueProperty.IsValid())
-	{
-		uint32 NumChildren = 0;
-		InstanceObjectValueProperty->GetNumChildren(NumChildren);
+	AddObjectInstance(InstanceObjectValueProperty);
 
-		// Find visible child properties and sort them so in order: Context, Input, Param, Output.
-		struct FSortedChild
-		{
-			TSharedPtr<IPropertyHandle> PropertyHandle;
-			EStateTreePropertyUsage Usage = EStateTreePropertyUsage::Invalid;
-		};
+	// ExecutionRuntime Instance
+	TSharedRef<FBindableNodeInstanceDetails> ExecutionRuntimeDataDetails = MakeShareable(new FBindableNodeInstanceDetails(ExecutionRuntimeDataProperty, EditorNodePtr->ID, EditorDataPtr));
+	ExecutionRuntimeDataProperty->SetInstanceMetaData(UE::PropertyBinding::MetaDataNoBindingName, TEXT("true"));
+	StructBuilder.AddCustomBuilder(ExecutionRuntimeDataDetails);
 
-		TArray<FSortedChild> SortedChildren;
-		for (uint32 Index = 0; Index < NumChildren; Index++)
-		{
-			if (TSharedPtr<IPropertyHandle> ChildHandle = InstanceObjectValueProperty->GetChildHandle(Index); ChildHandle.IsValid())
-			{
-				FSortedChild Child;
-				Child.PropertyHandle = ChildHandle;
-				Child.Usage = UE::StateTree::GetUsageFromMetaData(Child.PropertyHandle->GetProperty());
-
-				// If the property is set to one of these usages, display it even if it is not edit on instance.
-				// It is a common mistake to forget to set the "eye" on these properties it and wonder why it does not show up.
-				const bool bShouldShowByUsage = Child.Usage == EStateTreePropertyUsage::Input || Child.Usage == EStateTreePropertyUsage::Output || Child.Usage == EStateTreePropertyUsage::Context;
-        		const bool bIsEditable = !Child.PropertyHandle->GetProperty()->HasAllPropertyFlags(CPF_DisableEditOnInstance);
-
-				if (bShouldShowByUsage || bIsEditable)
-				{
-					SortedChildren.Add(Child);
-				}
-			}
-		}
-
-		SortedChildren.StableSort([](const FSortedChild& LHS, const FSortedChild& RHS) { return LHS.Usage < RHS.Usage; });
-
-		for (FSortedChild& Child : SortedChildren)
-		{
-			IDetailPropertyRow& ChildRow = StructBuilder.AddProperty(Child.PropertyHandle.ToSharedRef());
-			UE::StateTreeEditor::Internal::ModifyRow(ChildRow, EditorNodePtr->ID, EditorDataPtr);
-		}
-	}
+	// ExecutionRuntime Instance Object
+	TSharedPtr<IPropertyHandle> ExecutionRuntimeDataObjectValueProperty = GetInstancedObjectValueHandle(ExecutionRuntimeDataObjectProperty);
+	ExecutionRuntimeDataObjectProperty->SetInstanceMetaData(UE::PropertyBinding::MetaDataNoBindingName, TEXT("true"));
+	AddObjectInstance(ExecutionRuntimeDataObjectValueProperty);
 }
 
 TSharedPtr<IPropertyHandle> FStateTreeEditorNodeDetails::GetInstancedObjectValueHandle(TSharedPtr<IPropertyHandle> PropertyHandle)
diff --git a/Engine/Plugins/Runtime/StateTree/Source/StateTreeEditorModule/Private/Customizations/StateTreeEditorNodeDetails.h b/Engine/Plugins/Runtime/StateTree/Source/StateTreeEditorModule/Private/Customizations/StateTreeEditorNodeDetails.h
--- a/Engine/Plugins/Runtime/StateTree/Source/StateTreeEditorModule/Private/Customizations/StateTreeEditorNodeDetails.h
+++ b/Engine/Plugins/Runtime/StateTree/Source/StateTreeEditorModule/Private/Customizations/StateTreeEditorNodeDetails.h
@@ -146,6 +146,8 @@ private:
 	TSharedPtr<IPropertyHandle> NodeProperty;
 	TSharedPtr<IPropertyHandle> InstanceProperty;
 	TSharedPtr<IPropertyHandle> InstanceObjectProperty;
+	TSharedPtr<IPropertyHandle> ExecutionRuntimeDataProperty;
+	TSharedPtr<IPropertyHandle> ExecutionRuntimeDataObjectProperty;
 	TSharedPtr<IPropertyHandle> IDProperty;
 
 	TSharedPtr<IPropertyHandle> IndentProperty;
diff --git a/Engine/Plugins/Runtime/StateTree/Source/StateTreeEditorModule/Private/Customizations/StateTreeEditorNodeUtils.cpp b/Engine/Plugins/Runtime/StateTree/Source/StateTreeEditorModule/Private/Customizations/StateTreeEditorNodeUtils.cpp
--- a/Engine/Plugins/Runtime/StateTree/Source/StateTreeEditorModule/Private/Customizations/StateTreeEditorNodeUtils.cpp
+++ b/Engine/Plugins/Runtime/StateTree/Source/StateTreeEditorModule/Private/Customizations/StateTreeEditorNodeUtils.cpp
@@ -313,34 +313,55 @@ void GetNodeBaseScriptStructAndClass(const TSharedPtr<IPropertyHandle>& StructPr
 	OutBaseClass = UClass::TryFindTypeSlow<UClass>(BaseClassName);
 }
 
+struct FNodeStructView
+{
+	const UStruct* Struct = nullptr;
+	void* Memory = nullptr;
+	operator bool() const
+	{
+		return Struct != nullptr && Memory != nullptr;
+	}
+};
 
 struct FNodeRetainPropertyData
 {
-	FStateTreeNodeBase* NodeBase = nullptr;
-	const UScriptStruct* NodeBaseStruct = nullptr;
-	const UStruct* InstanceStruct = nullptr;
-	void* InstanceData = nullptr;
+	FNodeStructView Node; // FStateTreeNodeBase
+	FNodeStructView InstanceData; // the instance data for the node
+	FNodeStructView ExecutionRuntimeData; // the execution runtime data for the node (if it exists)
 };
 
 FNodeRetainPropertyData GetNodeData(FStateTreeEditorNode& EditorNode)
 {
 	FNodeRetainPropertyData Data;
-	Data.NodeBase = EditorNode.Node.GetMutablePtr<FStateTreeNodeBase>();
-
-	if (Data.NodeBase)
+	if (FStateTreeNodeBase* NodeBase = EditorNode.Node.GetMutablePtr<FStateTreeNodeBase>(); NodeBase != nullptr)
 	{
-		Data.NodeBaseStruct = EditorNode.Node.GetScriptStruct();
-		if (const UStruct* InstanceDataType = Data.NodeBase->GetInstanceDataType())
+		Data.Node.Struct = EditorNode.Node.GetScriptStruct();
+		Data.Node.Memory = NodeBase;
+
+		if (const UStruct* InstanceDataType = NodeBase->GetInstanceDataType())
 		{
 			if (InstanceDataType->IsA<UScriptStruct>())
 			{
-				Data.InstanceStruct = EditorNode.Instance.GetScriptStruct();
-				Data.InstanceData = EditorNode.Instance.GetMutableMemory();
+				Data.InstanceData.Struct = EditorNode.Instance.GetScriptStruct();
+				Data.InstanceData.Memory = EditorNode.Instance.GetMutableMemory();
 			}
 			else if (InstanceDataType->IsA<UClass>())
 			{
-				Data.InstanceStruct = EditorNode.InstanceObject.GetClass();
-				Data.InstanceData = EditorNode.InstanceObject;
+				Data.InstanceData.Struct = EditorNode.InstanceObject.GetClass();
+				Data.InstanceData.Memory = EditorNode.InstanceObject;
+			}
+		}
+		if (const UStruct* ExecutionRuntimeInstanceDataType = NodeBase->GetExecutionRuntimeDataType())
+		{
+			if (ExecutionRuntimeInstanceDataType->IsA<UScriptStruct>())
+			{
+				Data.ExecutionRuntimeData.Struct = EditorNode.ExecutionRuntimeData.GetScriptStruct();
+				Data.ExecutionRuntimeData.Memory = EditorNode.ExecutionRuntimeData.GetMutableMemory();
+			}
+			else if (ExecutionRuntimeInstanceDataType->IsA<UClass>())
+			{
+				Data.ExecutionRuntimeData.Struct = EditorNode.ExecutionRuntimeDataObject.GetClass();
+				Data.ExecutionRuntimeData.Memory = EditorNode.ExecutionRuntimeDataObject;
 			}
 		}
 	}
@@ -398,44 +419,66 @@ void CopyPropertyValues(const UStruct* OldStruct, const void* OldData, const USt
 	}
 }
 
+void CopyPropertyValues(const FNodeStructView Old, const FNodeStructView New)
+{
+	CopyPropertyValues(Old.Struct, Old.Memory, New.Struct, New.Memory);
+}
+
+
 void RetainProperties(FStateTreeEditorNode& OldNode, FStateTreeEditorNode& NewNode)
 {
 	const FNodeRetainPropertyData OldNodeData = GetNodeData(OldNode);
 	const FNodeRetainPropertyData NewNodeData = GetNodeData(NewNode);
 
-	if (OldNodeData.NodeBase && NewNodeData.NodeBase)
+	// node to ...
+	if (OldNodeData.Node)
 	{
-		// Copy node -> node
-		CopyPropertyValues(
-			OldNodeData.NodeBaseStruct, OldNodeData.NodeBase,
-			NewNodeData.NodeBaseStruct, NewNodeData.NodeBase
-		);
-
-		if (OldNodeData.InstanceStruct && OldNodeData.InstanceData)
+		if (NewNodeData.Node)
 		{
-			// Copy instance data -> node
-			CopyPropertyValues(
-				OldNodeData.InstanceStruct, OldNodeData.InstanceData,
-				NewNodeData.NodeBaseStruct, NewNodeData.NodeBase
-			);
+			CopyPropertyValues(OldNodeData.Node, NewNodeData.Node);
+		}
+		if (NewNodeData.InstanceData)
+		{
+			CopyPropertyValues(OldNodeData.Node, NewNodeData.InstanceData);
+		}
+		if (NewNodeData.ExecutionRuntimeData)
+		{
+			CopyPropertyValues(OldNodeData.Node, NewNodeData.ExecutionRuntimeData);
+		}
+	}
 
-			if (NewNodeData.InstanceStruct && NewNodeData.InstanceData)
-			{
-				// Copy instance data -> instance data
-				CopyPropertyValues(
-					OldNodeData.InstanceStruct, OldNodeData.InstanceData,
-					NewNodeData.InstanceStruct, NewNodeData.InstanceData
-				);
-			}
+	// instance data to ... 
+	if (OldNodeData.InstanceData)
+	{
+		if (NewNodeData.Node)
+		{
+			CopyPropertyValues(OldNodeData.InstanceData, NewNodeData.Node);
+		}
+		if (NewNodeData.InstanceData)
+		{
+			CopyPropertyValues(OldNodeData.InstanceData, NewNodeData.InstanceData);
+		}
+		if (NewNodeData.ExecutionRuntimeData)
+		{
+			CopyPropertyValues(OldNodeData.InstanceData, NewNodeData.ExecutionRuntimeData);
 		}
+	}
 
-		if (NewNodeData.InstanceStruct && NewNodeData.InstanceData)
+
+	// execution runtime data to ... 
+	if (OldNodeData.ExecutionRuntimeData)
+	{
+		if (NewNodeData.Node)
 		{
-			// Copy node -> instance data
-			CopyPropertyValues(
-				OldNodeData.NodeBaseStruct, OldNodeData.NodeBase,
-				NewNodeData.InstanceStruct, NewNodeData.InstanceData
-			);
+			CopyPropertyValues(OldNodeData.ExecutionRuntimeData, NewNodeData.Node);
+		}
+		if (NewNodeData.InstanceData)
+		{
+			CopyPropertyValues(OldNodeData.ExecutionRuntimeData, NewNodeData.InstanceData);
+		}
+		if (NewNodeData.ExecutionRuntimeData)
+		{
+			CopyPropertyValues(OldNodeData.ExecutionRuntimeData, NewNodeData.ExecutionRuntimeData);
 		}
 	}
 }
@@ -472,52 +515,25 @@ void SetNodeTypeStruct(const TSharedPtr<IPropertyHandle>& StructProperty, const
 					Node->Node.InitializeAs(InStruct);
 					
 					// Generate new name and instantiate instance data.
-					if (InStruct->IsChildOf(FStateTreeTaskBase::StaticStruct()))
-					{
-						FStateTreeTaskBase& Task = Node->Node.GetMutable<FStateTreeTaskBase>();
-						if (const UScriptStruct* InstanceType = Cast<const UScriptStruct>(Task.GetInstanceDataType()))
-						{
-							Node->Instance.InitializeAs(InstanceType);
-						}
-						else if (const UClass* InstanceClass = Cast<const UClass>(Task.GetInstanceDataType()))
-						{
-							Node->InstanceObject = NewObject<UObject>(Outer, InstanceClass);
-						}
-					}
-					else if (InStruct->IsChildOf(FStateTreeEvaluatorBase::StaticStruct()))
+					if (InStruct->IsChildOf(FStateTreeNodeBase::StaticStruct()))
 					{
-						FStateTreeEvaluatorBase& Eval = Node->Node.GetMutable<FStateTreeEvaluatorBase>();
-						if (const UScriptStruct* InstanceType = Cast<const UScriptStruct>(Eval.GetInstanceDataType()))
+						FStateTreeNodeBase& NodeBase = Node->Node.GetMutable<FStateTreeNodeBase>();
+						if (const UScriptStruct* InstanceType = Cast<const UScriptStruct>(NodeBase.GetInstanceDataType()))
 						{
 							Node->Instance.InitializeAs(InstanceType);
 						}
-						else if (const UClass* InstanceClass = Cast<const UClass>(Eval.GetInstanceDataType()))
+						else if (const UClass* InstanceClass = Cast<const UClass>(NodeBase.GetInstanceDataType()))
 						{
 							Node->InstanceObject = NewObject<UObject>(Outer, InstanceClass);
 						}
-					}
-					else if (InStruct->IsChildOf(FStateTreeConditionBase::StaticStruct()))
-					{
-						FStateTreeConditionBase& Cond = Node->Node.GetMutable<FStateTreeConditionBase>();
-						if (const UScriptStruct* InstanceType = Cast<const UScriptStruct>(Cond.GetInstanceDataType()))
-						{
-							Node->Instance.InitializeAs(InstanceType);
-						}
-						else if (const UClass* InstanceClass = Cast<const UClass>(Cond.GetInstanceDataType()))
-						{
-							Node->InstanceObject = NewObject<UObject>(Outer, InstanceClass);
-						}
-					}
-					else if (InStruct->IsChildOf(FStateTreeConsiderationBase::StaticStruct()))
-					{
-						FStateTreeConsiderationBase& Consideration = Node->Node.GetMutable<FStateTreeConsiderationBase>();
-						if (const UScriptStruct* InstanceType = Cast<const UScriptStruct>(Consideration.GetInstanceDataType()))
+
+						if (const UScriptStruct* InstanceType = Cast<const UScriptStruct>(NodeBase.GetExecutionRuntimeDataType()))
 						{
-							Node->Instance.InitializeAs(InstanceType);
+							Node->ExecutionRuntimeData.InitializeAs(InstanceType);
 						}
-						else if (const UClass* InstanceClass = Cast<const UClass>(Consideration.GetInstanceDataType()))
+						else if (const UClass* InstanceClass = Cast<const UClass>(NodeBase.GetExecutionRuntimeDataType()))
 						{
-							Node->InstanceObject = NewObject<UObject>(Outer, InstanceClass);
+							Node->ExecutionRuntimeDataObject = NewObject<UObject>(Outer, InstanceClass);
 						}
 					}
 
@@ -670,10 +686,12 @@ void ConditionalUpdateNodeInstanceData(FStateTreeEditorNode& EditorNode, UObject
 	}
 
 	const UStruct* CurrentType = EditorNode.GetInstance().GetStruct();
+	const UStruct* ExecutionRuntimeCurrentType = EditorNode.GetExecutionRuntimeData().GetStruct();
 	const UStruct* DesiredType = Node->GetInstanceDataType();
+	const UStruct* ExecutionRuntimeDesiredType = Node->GetExecutionRuntimeDataType();
 
 	// Nothing to upgrade. Instance Data Type is unchanged
-	if (CurrentType == DesiredType)
+	if (CurrentType == DesiredType && ExecutionRuntimeCurrentType == ExecutionRuntimeDesiredType)
 	{
 		return;
 	}
@@ -682,6 +700,8 @@ void ConditionalUpdateNodeInstanceData(FStateTreeEditorNode& EditorNode, UObject
 
 	EditorNode.Instance.Reset();
 	EditorNode.InstanceObject = nullptr;
+	EditorNode.ExecutionRuntimeData.Reset();
+	EditorNode.ExecutionRuntimeDataObject = nullptr;
 
 	if (const UScriptStruct* InstanceType = Cast<UScriptStruct>(DesiredType))
 	{
@@ -691,6 +711,14 @@ void ConditionalUpdateNodeInstanceData(FStateTreeEditorNode& EditorNode, UObject
 	{
 		EditorNode.InstanceObject = NewObject<UObject>(&InstanceOuter, InstanceClass);
 	}
+	if (const UScriptStruct* InstanceType = Cast<UScriptStruct>(ExecutionRuntimeDesiredType))
+	{
+		EditorNode.ExecutionRuntimeData.InitializeAs(InstanceType);
+	}
+	else if (const UClass* InstanceClass = Cast<UClass>(ExecutionRuntimeDesiredType))
+	{
+		EditorNode.ExecutionRuntimeDataObject = NewObject<UObject>(&InstanceOuter, InstanceClass);
+	}
 
 	RetainProperties(OldEditorNode, EditorNode);
 
@@ -704,6 +732,14 @@ void ConditionalUpdateNodeInstanceData(FStateTreeEditorNode& EditorNode, UObject
 	{
 		UE::StateTreeEditor::EditorNodeUtils::InstantiateStructSubobjects(InstanceOuter, EditorNode.Instance);
 	}
+	if (EditorNode.ExecutionRuntimeDataObject)
+	{
+		EditorNode.ExecutionRuntimeDataObject = DuplicateObject(EditorNode.ExecutionRuntimeDataObject, &InstanceOuter);
+	}
+	else
+	{
+		UE::StateTreeEditor::EditorNodeUtils::InstantiateStructSubobjects(InstanceOuter, EditorNode.ExecutionRuntimeData);
+	}
 }
 
 void OnArrayNodePicked(const UStruct* InStruct, TSharedPtr<SComboButton> PickerCombo, TSharedPtr<IPropertyHandle> ArrayPropertyHandle, TSharedRef<IPropertyUtilities> PropUtils)
diff --git a/Engine/Plugins/Runtime/StateTree/Source/StateTreeEditorModule/Private/StateTreeCompiler.cpp b/Engine/Plugins/Runtime/StateTree/Source/StateTreeEditorModule/Private/StateTreeCompiler.cpp
--- a/Engine/Plugins/Runtime/StateTree/Source/StateTreeEditorModule/Private/StateTreeCompiler.cpp
+++ b/Engine/Plugins/Runtime/StateTree/Source/StateTreeEditorModule/Private/StateTreeCompiler.cpp
@@ -540,6 +540,7 @@ bool FStateTreeCompiler::Compile(TNotNull<UStateTree*> InStateTree)
 	StateTree->DefaultInstanceData.Init(*StateTree, InstanceStructs, FStateTreeInstanceData::FAddArgs{ .bDuplicateWrappedObject = false });
 	StateTree->SharedInstanceData.Init(*StateTree, SharedInstanceStructs, FStateTreeInstanceData::FAddArgs{ .bDuplicateWrappedObject = false });
 	StateTree->DefaultEvaluationScopeInstanceData.Init(StateTree, EvaluationScopeStructs, UE::StateTree::InstanceData::FInstanceContainer::FAddArgs{ .bDuplicateWrappedObject = false });
+	StateTree->DefaultExecutionRuntimeData.Init(StateTree, ExecutionRuntimeStructs, UE::StateTree::InstanceData::FInstanceContainer::FAddArgs{ .bDuplicateWrappedObject = false });
 
 	// Store the new compiled dispatchers.
 	EditorData->CompiledDispatchers = BindingsCompiler.GetCompiledDelegateDispatchers();
@@ -785,7 +786,7 @@ bool FStateTreeCompiler::CreateStateRecursive(UStateTreeState& State, const FSta
 		return false;
 	}
 	StateTree->States[StateIdx].ChildrenEnd = uint16(ChildrenEnd); // Not using CompactState here because the array may have changed.
-	
+
 	// create sub frame info
 	if (!Parent.IsValid())
 	{
@@ -1894,7 +1895,7 @@ FInstancedStruct* FStateTreeCompiler::CreateNodeWithSharedInstanceData(UStateTre
 	return CreateNode(State, EditorNode, InstanceDesc, DataHandle, SharedInstanceStructs);
 }
 
-TOptional<FStateTreeDataView> FStateTreeCompiler::CreateNodeInstanceData(const FStateTreeEditorNode& EditorNode, FStateTreeNodeBase& Node, FStateTreeBindableStructDesc& InstanceDesc, const FStateTreeDataHandle DataHandle, TArray<FInstancedStruct>& InstancedStructContainer)
+TOptional<FStateTreeDataView> FStateTreeCompiler::CreateNodeInstanceData(const FStateTreeEditorNode& EditorNode, FStateTreeNodeBase& Node, FStateTreeBindableStructDesc& StructDesc, const FStateTreeDataHandle DataHandle, TArray<FInstancedStruct>& InstancedStructContainer)
 {
 	FStateTreeDataView InstanceDataView;
 
@@ -1903,47 +1904,120 @@ TOptional<FStateTreeDataView> FStateTreeCompiler::CreateNodeInstanceData(const F
 
 	if (EditorNode.Instance.IsValid())
 	{
-		// Struct Instance
-		const int32 InstanceIndex = InstancedStructContainer.Add(EditorNode.Instance);
-		InstantiateStructSubobjects(InstancedStructContainer[InstanceIndex]);
+		if (ensure(EditorNode.Instance.GetScriptStruct() == Node.GetInstanceDataType()))
+		{
+			// Struct Instance
+			const int32 InstanceIndex = InstancedStructContainer.Add(EditorNode.Instance);
+			InstantiateStructSubobjects(InstancedStructContainer[InstanceIndex]);
 
-		// Create binding source struct descriptor.
-		InstanceDesc.Struct = EditorNode.Instance.GetScriptStruct();
+			// Create binding source struct descriptor.
+			StructDesc.Struct = EditorNode.Instance.GetScriptStruct();
 
-		if (const auto Validation = UE::StateTree::Compiler::IsValidIndex16(InstanceIndex); Validation.DidFail())
+			if (const auto Validation = UE::StateTree::Compiler::IsValidIndex16(InstanceIndex); Validation.DidFail())
+			{
+				Validation.Log(Log, TEXT("InstanceIndex"), StructDesc);
+				return {};
+			}
+			Node.InstanceTemplateIndex = FStateTreeIndex16(InstanceIndex);
+			Node.InstanceDataHandle = DataHandle;
+			InstanceDataView = FStateTreeDataView(InstancedStructContainer[InstanceIndex]);
+		}
+		else
 		{
-			Validation.Log(Log, TEXT("InstanceIndex"), InstanceDesc);
+			Log.Reportf(EMessageSeverity::Error, StructDesc, TEXT("The instance data type does not match."));
 			return {};
 		}
-		Node.InstanceTemplateIndex = FStateTreeIndex16(InstanceIndex);
-		Node.InstanceDataHandle = DataHandle;
-		InstanceDataView = FStateTreeDataView(InstancedStructContainer[InstanceIndex]);
 	}
-	else
+	else if (EditorNode.InstanceObject != nullptr)
 	{
-		// Object Instance
-		check(EditorNode.InstanceObject != nullptr);
-		UObject* Instance = UE::StateTree::Compiler::DuplicateInstanceObject(Log, InstanceDesc, EditorNode.ID, EditorNode.InstanceObject, StateTree);
+		if (ensure(EditorNode.InstanceObject->GetClass() == Node.GetInstanceDataType()))
+		{
+			UObject* Instance = UE::StateTree::Compiler::DuplicateInstanceObject(Log, StructDesc, EditorNode.ID, EditorNode.InstanceObject, StateTree);
 
-		FInstancedStruct Wrapper;
-		Wrapper.InitializeAs<FStateTreeInstanceObjectWrapper>(Instance);
-		const int32 InstanceIndex = InstancedStructContainer.Add(MoveTemp(Wrapper));
+			FInstancedStruct Wrapper;
+			Wrapper.InitializeAs<FStateTreeInstanceObjectWrapper>(Instance);
+			const int32 InstanceIndex = InstancedStructContainer.Add(MoveTemp(Wrapper));
 
-		// Create binding source struct descriptor.
-		InstanceDesc.Struct = Instance->GetClass();
+			// Create binding source struct descriptor.
+			StructDesc.Struct = Instance->GetClass();
 
-		if (const auto Validation = UE::StateTree::Compiler::IsValidIndex16(InstanceIndex); Validation.DidFail())
+			if (const auto Validation = UE::StateTree::Compiler::IsValidIndex16(InstanceIndex); Validation.DidFail())
+			{
+				Validation.Log(Log, TEXT("InstanceIndex"), StructDesc);
+				return {};
+			}
+			Node.InstanceTemplateIndex = FStateTreeIndex16(InstanceIndex);
+			Node.InstanceDataHandle = DataHandle.ToObjectSource();
+			InstanceDataView = FStateTreeDataView(Instance);
+		}
+		else
 		{
-			Validation.Log(Log, TEXT("InstanceIndex"), InstanceDesc);
+			Log.Reportf(EMessageSeverity::Error, StructDesc, TEXT("The instance data type does not match."));
 			return {};
 		}
-		Node.InstanceTemplateIndex = FStateTreeIndex16(InstanceIndex);
-		Node.InstanceDataHandle = DataHandle.ToObjectSource();
-		InstanceDataView = FStateTreeDataView(Instance);
+	}
+	else if (Node.GetInstanceDataType() != nullptr)
+	{
+		Log.Reportf(EMessageSeverity::Error, StructDesc, TEXT("The instance data is missing."));
+		return {};
+	}
+
+	StructDesc.DataHandle = Node.InstanceDataHandle;
+
+	if (EditorNode.ExecutionRuntimeData.IsValid() || EditorNode.ExecutionRuntimeDataObject != nullptr)
+	{
+		if (EditorNode.ExecutionRuntimeData.IsValid())
+		{
+			if (ensure(EditorNode.ExecutionRuntimeData.GetScriptStruct() == Node.GetExecutionRuntimeDataType()))
+			{
+				// Struct Instance
+				const int32 InstanceIndex = ExecutionRuntimeStructs.Add(EditorNode.ExecutionRuntimeData);
+				InstantiateStructSubobjects(ExecutionRuntimeStructs[InstanceIndex]);
+
+				if (const auto Validation = UE::StateTree::Compiler::IsValidIndex16(InstanceIndex); Validation.DidFail())
+				{
+					Validation.Log(Log, TEXT("ExecutionRuntime Index"), StructDesc);
+					return {};
+				}
+				Node.ExecutionRuntimeTemplateIndex = FStateTreeIndex16(InstanceIndex);
+			}
+			else
+			{
+				Log.Reportf(EMessageSeverity::Error, StructDesc, TEXT("The execution runtime data type does not match."));
+				return {};
+			}
+		}
+		else if (EditorNode.ExecutionRuntimeDataObject != nullptr)
+		{
+			if (ensure(EditorNode.ExecutionRuntimeDataObject->GetClass() == Node.GetExecutionRuntimeDataType()))
+			{
+				// Object Instance
+				UObject* Instance = UE::StateTree::Compiler::DuplicateInstanceObject(Log, StructDesc, EditorNode.ID, EditorNode.ExecutionRuntimeDataObject, StateTree);
+
+				FInstancedStruct Wrapper;
+				Wrapper.InitializeAs<FStateTreeInstanceObjectWrapper>(Instance);
+				const int32 InstanceIndex = ExecutionRuntimeStructs.Add(MoveTemp(Wrapper));
+
+				if (const auto Validation = UE::StateTree::Compiler::IsValidIndex16(InstanceIndex); Validation.DidFail())
+				{
+					Validation.Log(Log, TEXT("ExecutionRuntime Index"), StructDesc);
+					return {};
+				}
+				Node.ExecutionRuntimeTemplateIndex = FStateTreeIndex16(InstanceIndex);
+			}
+			else
+			{
+				Log.Reportf(EMessageSeverity::Error, StructDesc, TEXT("The execution runtime data type does not match."));
+				return {};
+			}
+		}
+	}
+	else if (Node.GetExecutionRuntimeDataType() != nullptr)
+	{
+		Log.Reportf(EMessageSeverity::Error, StructDesc, TEXT("The execution runtime data is missing."));
+		return {};
 	}
 
-	InstanceDesc.DataHandle = Node.InstanceDataHandle;
-	
 	return TOptional<FStateTreeDataView>(InstanceDataView);
 }
 
diff --git a/Engine/Plugins/Runtime/StateTree/Source/StateTreeEditorModule/Private/StateTreeEditorData.cpp b/Engine/Plugins/Runtime/StateTree/Source/StateTreeEditorModule/Private/StateTreeEditorData.cpp
--- a/Engine/Plugins/Runtime/StateTree/Source/StateTreeEditorModule/Private/StateTreeEditorData.cpp
+++ b/Engine/Plugins/Runtime/StateTree/Source/StateTreeEditorModule/Private/StateTreeEditorData.cpp
@@ -742,26 +742,85 @@ void UStateTreeEditorData::ReparentStates()
 	});
 }
 
-void UStateTreeEditorData::FixObjectInstance(TSet<UObject*>& SeenObjects, UObject& Outer, FStateTreeEditorNode& Node)
+void UStateTreeEditorData::FixObjectInstances(TSet<UObject*>& SeenObjects, UObject& Outer, FStateTreeEditorNode& Node)
 {
-	if (Node.InstanceObject)
-	{
-		// Found a duplicate reference to an object, make unique copy.
-		if (SeenObjects.Contains(Node.InstanceObject))
+	auto NewInstances = [this, &Node, &Outer](const UStruct* Struct, const UStruct* ExecutionRuntimeStruct)
 		{
-			UE_LOG(LogStateTreeEditor, Log, TEXT("%s: Making duplicate node instance %s unique."), *GetFullName(), *GetNameSafe(Node.InstanceObject));
-			Node.InstanceObject = DuplicateObject(Node.InstanceObject, &Outer);
-		}
-		else
+			if (Struct && !Node.Instance.IsValid() && Node.InstanceObject == nullptr)
+			{
+				if (const UScriptStruct* InstanceType = Cast<const UScriptStruct>(Struct))
+				{
+					Node.Instance.InitializeAs(InstanceType);
+				}
+				else if (const UClass* InstanceClass = Cast<const UClass>(Struct))
+				{
+					Node.InstanceObject = NewObject<UObject>(&Outer, InstanceClass);
+				}
+			}
+			if (ExecutionRuntimeStruct && !Node.ExecutionRuntimeData.IsValid() && Node.ExecutionRuntimeDataObject == nullptr)
+			{
+				if (const UScriptStruct* InstanceType = Cast<const UScriptStruct>(ExecutionRuntimeStruct))
+				{
+					Node.ExecutionRuntimeData.InitializeAs(InstanceType);
+				}
+				else if (const UClass* InstanceClass = Cast<const UClass>(ExecutionRuntimeStruct))
+				{
+					Node.ExecutionRuntimeDataObject = NewObject<UObject>(&Outer, InstanceClass);
+				}
+			}
+		};
+	auto FixInstance = [this, &SeenObjects, &Outer](FInstancedStruct& Instance, const UStruct* Struct)
 		{
-			// Make sure the instance object is property outered.
-			if (Node.InstanceObject->GetOuter() != &Outer)
+			if (Instance.IsValid())
 			{
-				UE_LOG(LogStateTreeEditor, Log, TEXT("%s: Fixing outer on node instance %s."), *GetFullName(), *GetNameSafe(Node.InstanceObject));
-				Node.InstanceObject->Rename(nullptr, &Outer, REN_DontCreateRedirectors | REN_DoNotDirty);
+				if (Instance.GetScriptStruct() != Struct)
+				{
+					Instance.Reset();
+				}
 			}
-		}
-		SeenObjects.Add(Node.InstanceObject);
+		};
+	auto FixInstanceObject = [this, &SeenObjects, &Outer](TObjectPtr<UObject>& InstanceObject, const UStruct* Struct)
+		{
+			if (InstanceObject)
+			{
+				if (InstanceObject->GetClass() != Struct)
+				{
+					InstanceObject = nullptr;
+				}
+				// Found a duplicate reference to an object, make unique copy.
+				else if(SeenObjects.Contains(InstanceObject))				
+				{
+					UE_LOG(LogStateTreeEditor, Log, TEXT("%s: Making duplicate node instance %s unique."), *GetFullName(), *GetNameSafe(InstanceObject));
+					InstanceObject = DuplicateObject(InstanceObject, &Outer);
+				}
+				else
+				{
+					// Make sure the instance object is property outered.
+					if (InstanceObject->GetOuter() != &Outer)
+					{
+						UE_LOG(LogStateTreeEditor, Log, TEXT("%s: Fixing outer on node instance %s."), *GetFullName(), *GetNameSafe(InstanceObject));
+						InstanceObject->Rename(nullptr, &Outer, REN_DontCreateRedirectors | REN_DoNotDirty);
+					}
+				}
+				SeenObjects.Add(InstanceObject);
+			}
+		};
+
+	const FStateTreeNodeBase* NodeBase = Node.Node.GetPtr<FStateTreeNodeBase>();
+	if (NodeBase)
+	{
+		FixInstance(Node.Instance, NodeBase->GetInstanceDataType());
+		FixInstanceObject(Node.InstanceObject, NodeBase->GetInstanceDataType());
+		FixInstance(Node.ExecutionRuntimeData, NodeBase->GetExecutionRuntimeDataType());
+		FixInstanceObject(Node.ExecutionRuntimeDataObject, NodeBase->GetExecutionRuntimeDataType());
+		NewInstances(NodeBase->GetInstanceDataType(), NodeBase->GetExecutionRuntimeDataType());
+	}
+	else
+	{
+		Node.Instance.Reset();
+		Node.InstanceObject = nullptr;
+		Node.ExecutionRuntimeData.Reset();
+		Node.ExecutionRuntimeDataObject = nullptr;
 	}
 };
 
@@ -780,16 +839,16 @@ void UStateTreeEditorData::FixObjectNodes()
 		// Enter conditions
 		for (FStateTreeEditorNode& Node : State.EnterConditions)
 		{
-			TreeData->FixObjectInstance(SeenObjects, State, Node);
+			TreeData->FixObjectInstances(SeenObjects, State, Node);
 		}
 		
 		// Tasks
 		for (FStateTreeEditorNode& Node : State.Tasks)
 		{
-			TreeData->FixObjectInstance(SeenObjects, State, Node);
+			TreeData->FixObjectInstances(SeenObjects, State, Node);
 		}
 
-		TreeData->FixObjectInstance(SeenObjects, State, State.SingleTask);
+		TreeData->FixObjectInstances(SeenObjects, State, State.SingleTask);
 
 
 		// Transitions
@@ -797,7 +856,7 @@ void UStateTreeEditorData::FixObjectNodes()
 		{
 			for (FStateTreeEditorNode& Node : Transition.Conditions)
 			{
-				TreeData->FixObjectInstance(SeenObjects, State, Node);
+				TreeData->FixObjectInstances(SeenObjects, State, Node);
 			}
 		}
 		
@@ -806,12 +865,12 @@ void UStateTreeEditorData::FixObjectNodes()
 
 	for (FStateTreeEditorNode& Node : Evaluators)
 	{
-		FixObjectInstance(SeenObjects, *this, Node);
+		FixObjectInstances(SeenObjects, *this, Node);
 	}
 
 	for (FStateTreeEditorNode& Node : GlobalTasks)
 	{
-		FixObjectInstance(SeenObjects, *this, Node);
+		FixObjectInstances(SeenObjects, *this, Node);
 	}
 }
 
diff --git a/Engine/Plugins/Runtime/StateTree/Source/StateTreeEditorModule/Private/StateTreeEditorPropertyBindings.cpp b/Engine/Plugins/Runtime/StateTree/Source/StateTreeEditorModule/Private/StateTreeEditorPropertyBindings.cpp
--- a/Engine/Plugins/Runtime/StateTree/Source/StateTreeEditorModule/Private/StateTreeEditorPropertyBindings.cpp
+++ b/Engine/Plugins/Runtime/StateTree/Source/StateTreeEditorModule/Private/StateTreeEditorPropertyBindings.cpp
@@ -33,6 +33,10 @@ FPropertyBindingPath FStateTreeEditorPropertyBindings::AddFunctionBinding(const
 	{
 		PropertyFunction.Instance.InitializeAs(InstanceType);
 	}
+	if (const UScriptStruct* InstanceType = Cast<const UScriptStruct>(Function.GetExecutionRuntimeDataType()))
+	{
+		PropertyFunction.ExecutionRuntimeData.InitializeAs(InstanceType);
+	}
 
 	Super::RemoveBindings(TargetPath, ESearchMode::Exact);
 	FPropertyBindingPath SourcePath = FPropertyBindingPath(NodeID, InSourcePathSegments);
diff --git a/Engine/Plugins/Runtime/StateTree/Source/StateTreeEditorModule/Private/StateTreePropertyHelpers.cpp b/Engine/Plugins/Runtime/StateTree/Source/StateTreeEditorModule/Private/StateTreePropertyHelpers.cpp
--- a/Engine/Plugins/Runtime/StateTree/Source/StateTreeEditorModule/Private/StateTreePropertyHelpers.cpp
+++ b/Engine/Plugins/Runtime/StateTree/Source/StateTreeEditorModule/Private/StateTreePropertyHelpers.cpp
@@ -72,6 +72,20 @@ namespace Internal
 							StateTreeNode->PostEditInstanceDataChangeChainProperty(PropertyChangedEvent, FStateTreeDataView(InEditorNode.InstanceObject));
 						}
 					}
+					else if (EditorNodeChildMember->GetFName() == GET_MEMBER_NAME_CHECKED(FStateTreeEditorNode, ExecutionRuntimeData))
+					{
+						if (InEditorNode.ExecutionRuntimeData.IsValid())
+						{
+							StateTreeNode->PostEditInstanceDataChangeChainProperty(PropertyChangedEvent, FStateTreeDataView(InEditorNode.ExecutionRuntimeData));
+						}
+					}
+					else if (EditorNodeChildMember->GetFName() == GET_MEMBER_NAME_CHECKED(FStateTreeEditorNode, ExecutionRuntimeDataObject))
+					{
+						if (InEditorNode.ExecutionRuntimeDataObject)
+						{
+							StateTreeNode->PostEditInstanceDataChangeChainProperty(PropertyChangedEvent, FStateTreeDataView(InEditorNode.ExecutionRuntimeDataObject));
+						}
+					}
 				}
 			}
 		}
@@ -157,12 +171,20 @@ void DispatchPostEditToNodes(UObject& Owner, FPropertyChangedChainEvent& InPrope
 				if (CurrentPropNode)
 				{
 					const FName EditorNodeChildMemberName = CurrentPropNode->GetValue()->GetFName();
-					if (EditorNodeChildMemberName == GET_MEMBER_NAME_CHECKED(FStateTreeEditorNode, Instance) || EditorNodeChildMemberName == GET_MEMBER_NAME_CHECKED(FStateTreeEditorNode, InstanceObject))
+					if (EditorNodeChildMemberName == GET_MEMBER_NAME_CHECKED(FStateTreeEditorNode, Instance)
+						|| EditorNodeChildMemberName == GET_MEMBER_NAME_CHECKED(FStateTreeEditorNode, InstanceObject))
 					{
 						CurrentAddress = static_cast<uint8*>(LastEditorNode->GetInstance().GetMutableMemory());
 						CurrentPropNode = CurrentPropNode->GetNextNode();
 						continue;
 					}
+					if (EditorNodeChildMemberName == GET_MEMBER_NAME_CHECKED(FStateTreeEditorNode, ExecutionRuntimeData)
+						|| EditorNodeChildMemberName == GET_MEMBER_NAME_CHECKED(FStateTreeEditorNode, ExecutionRuntimeDataObject))
+					{
+						CurrentAddress = static_cast<uint8*>(LastEditorNode->GetExecutionRuntimeData().GetMutableMemory());
+						CurrentPropNode = CurrentPropNode->GetNextNode();
+						continue;
+					}
 				}
 
 				break;
diff --git a/Engine/Plugins/Runtime/StateTree/Source/StateTreeEditorModule/Private/StateTreeScopedEditorDataFixer.cpp b/Engine/Plugins/Runtime/StateTree/Source/StateTreeEditorModule/Private/StateTreeScopedEditorDataFixer.cpp
--- a/Engine/Plugins/Runtime/StateTree/Source/StateTreeEditorModule/Private/StateTreeScopedEditorDataFixer.cpp
+++ b/Engine/Plugins/Runtime/StateTree/Source/StateTreeEditorModule/Private/StateTreeScopedEditorDataFixer.cpp
@@ -28,6 +28,15 @@ FScopedEditorDataFixer::~FScopedEditorDataFixer()
 			{
 				EditorNodeUtils::InstantiateStructSubobjects(*NodeOwner, EditorNode.Instance);
 			}
+
+			if (EditorNode.ExecutionRuntimeDataObject)
+			{
+				EditorNode.ExecutionRuntimeDataObject = DuplicateObject(EditorNode.ExecutionRuntimeDataObject, NodeOwner);
+			}
+			else if (EditorNode.ExecutionRuntimeData.IsValid())
+			{
+				EditorNodeUtils::InstantiateStructSubobjects(*NodeOwner, EditorNode.ExecutionRuntimeData);
+			}
 		};
 
 	FStateTreeEditorPropertyBindings* Bindings = EditorData->GetPropertyEditorBindings();
diff --git a/Engine/Plugins/Runtime/StateTree/Source/StateTreeEditorModule/Private/StateTreeViewModel.cpp b/Engine/Plugins/Runtime/StateTree/Source/StateTreeEditorModule/Private/StateTreeViewModel.cpp
--- a/Engine/Plugins/Runtime/StateTree/Source/StateTreeEditorModule/Private/StateTreeViewModel.cpp
+++ b/Engine/Plugins/Runtime/StateTree/Source/StateTreeEditorModule/Private/StateTreeViewModel.cpp
@@ -141,6 +141,14 @@ namespace UE::StateTree::Editor
 			{
 				CollectStateLinks(Node.InstanceObject->GetClass(), Node.InstanceObject, Links);
 			}
+			if (Node.ExecutionRuntimeData.IsValid())
+			{
+				CollectStateLinks(Node.ExecutionRuntimeData.GetScriptStruct(), Node.ExecutionRuntimeData.GetMutableMemory(), Links);
+			}
+			if (Node.ExecutionRuntimeDataObject)
+			{
+				CollectStateLinks(Node.ExecutionRuntimeDataObject->GetClass(), Node.ExecutionRuntimeDataObject, Links);
+			}
 		}
 	}
 
diff --git a/Engine/Plugins/Runtime/StateTree/Source/StateTreeEditorModule/Public/StateTreeCompiler.h b/Engine/Plugins/Runtime/StateTree/Source/StateTreeEditorModule/Public/StateTreeCompiler.h
--- a/Engine/Plugins/Runtime/StateTree/Source/StateTreeEditorModule/Public/StateTreeCompiler.h
+++ b/Engine/Plugins/Runtime/StateTree/Source/StateTreeEditorModule/Public/StateTreeCompiler.h
@@ -124,6 +124,7 @@ private:
 	TArray<FInstancedStruct> InstanceStructs;
 	TArray<FInstancedStruct> SharedInstanceStructs;
 	TArray<FInstancedStruct> EvaluationScopeStructs;
+	TArray<FInstancedStruct> ExecutionRuntimeStructs;
 	
 	/** Cached result of MakeCompletionTasksMask for global tasks. Indicates where state tasks should start. */
 	int32 GlobalTaskEndBit = 0;
diff --git a/Engine/Plugins/Runtime/StateTree/Source/StateTreeEditorModule/Public/StateTreeEditorData.h b/Engine/Plugins/Runtime/StateTree/Source/StateTreeEditorModule/Public/StateTreeEditorData.h
--- a/Engine/Plugins/Runtime/StateTree/Source/StateTreeEditorModule/Public/StateTreeEditorData.h
+++ b/Engine/Plugins/Runtime/StateTree/Source/StateTreeEditorModule/Public/StateTreeEditorData.h
@@ -249,6 +249,10 @@ public:
 		{
 			EditorNode.Instance.InitializeAs(InstanceType);
 		}
+		if (const UScriptStruct* InstanceType = Cast<const UScriptStruct>(Eval.GetExecutionRuntimeDataType()))
+		{
+			EditorNode.ExecutionRuntimeData.InitializeAs(InstanceType);
+		}
 		return static_cast<TStateTreeEditorNode<T>&>(EditorNode);
 	}
 
@@ -267,6 +271,10 @@ public:
 		{
 			EditorNode.Instance.InitializeAs(InstanceType);
 		}
+		if (const UScriptStruct* InstanceType = Cast<const UScriptStruct>(Task.GetExecutionRuntimeDataType()))
+		{
+			EditorNode.ExecutionRuntimeData.InitializeAs(InstanceType);
+		}
 		return static_cast<TStateTreeEditorNode<T>&>(EditorNode);
 	}
 
@@ -341,7 +349,7 @@ PRAGMA_ENABLE_DEPRECATION_WARNINGS
 private:
 	UE_API virtual EStateTreeVisitor EnumerateBindablePropertyFunctionNodes(TFunctionRef<EStateTreeVisitor(const UScriptStruct* NodeStruct, const FStateTreeBindableStructDesc& Desc, const FStateTreeDataView Value)> InFunc) const override;
 
-	UE_API void FixObjectInstance(TSet<UObject*>& SeenObjects, UObject& Outer, FStateTreeEditorNode& Node);
+	UE_API void FixObjectInstances(TSet<UObject*>& SeenObjects, UObject& Outer, FStateTreeEditorNode& Node);
 	UE_API void FixObjectNodes();
 	UE_API void FixDuplicateIDs();
 	UE_API void DuplicateIDs();
diff --git a/Engine/Plugins/Runtime/StateTree/Source/StateTreeEditorModule/Public/StateTreeEditorNode.h b/Engine/Plugins/Runtime/StateTree/Source/StateTreeEditorModule/Public/StateTreeEditorNode.h
--- a/Engine/Plugins/Runtime/StateTree/Source/StateTreeEditorModule/Public/StateTreeEditorNode.h
+++ b/Engine/Plugins/Runtime/StateTree/Source/StateTreeEditorModule/Public/StateTreeEditorNode.h
@@ -31,6 +31,8 @@ struct FStateTreeEditorNode
 		Node.Reset();
 		Instance.Reset();
 		InstanceObject = nullptr;
+		ExecutionRuntimeData.Reset();
+		ExecutionRuntimeDataObject = nullptr;
 		ID = FGuid();
 	}
 
@@ -69,6 +71,16 @@ struct FStateTreeEditorNode
 		return InstanceObject ? FStateTreeDataView(InstanceObject) : FStateTreeDataView(Instance);
 	}
 
+	FStateTreeDataView GetExecutionRuntimeData() const
+	{
+		return ExecutionRuntimeDataObject ? FStateTreeDataView(ExecutionRuntimeDataObject) : FStateTreeDataView(const_cast<FInstancedStruct&>(ExecutionRuntimeData));
+	}
+
+	FStateTreeDataView GetExecutionRuntimeData()
+	{
+		return ExecutionRuntimeDataObject ? FStateTreeDataView(ExecutionRuntimeDataObject) : FStateTreeDataView(ExecutionRuntimeData);
+	}
+
 	UPROPERTY(EditDefaultsOnly, Category = Node)
 	FInstancedStruct Node;
 
@@ -78,9 +90,13 @@ struct FStateTreeEditorNode
 	UPROPERTY(EditDefaultsOnly, Instanced, Category = Node)
 	TObjectPtr<UObject> InstanceObject = nullptr;
 
-	/**
-	 * ID for the node instance.
-	 */
+	UPROPERTY(EditDefaultsOnly, Category = Node)
+	FInstancedStruct ExecutionRuntimeData;
+
+	UPROPERTY(EditDefaultsOnly, Instanced, Category = Node)
+	TObjectPtr<UObject> ExecutionRuntimeDataObject = nullptr;
+
+	/** ID for the node instance. */
 	UPROPERTY(EditDefaultsOnly, Category = Node)
 	FGuid ID;
 
@@ -97,6 +113,7 @@ struct TStateTreeEditorNode : public FStateTreeEditorNode
 	using NodeType = T;
 	inline T& GetNode() { return Node.template GetMutable<T>(); }
 	inline typename T::FInstanceDataType& GetInstanceData() { return Instance.template GetMutable<typename T::FInstanceDataType>(); }
+	inline typename T::FExecutionRuntimeDataType& GetExecutionRuntimeData() { return ExecutionRuntimeData.template GetMutable<typename T::FExecutionRuntimeDataType>(); }
 };
 
 #undef UE_API
diff --git a/Engine/Plugins/Runtime/StateTree/Source/StateTreeEditorModule/Public/StateTreeState.h b/Engine/Plugins/Runtime/StateTree/Source/StateTreeEditorModule/Public/StateTreeState.h
--- a/Engine/Plugins/Runtime/StateTree/Source/StateTreeEditorModule/Public/StateTreeState.h
+++ b/Engine/Plugins/Runtime/StateTree/Source/StateTreeEditorModule/Public/StateTreeState.h
@@ -103,6 +103,10 @@ PRAGMA_ENABLE_DEPRECATION_WARNINGS
 		{
 			CondNode.Instance.InitializeAs(InstanceType);
 		}
+		if (const UScriptStruct* InstanceType = Cast<const UScriptStruct>(Node.GetExecutionRuntimeDataType()))
+		{
+			CondNode.ExecutionRuntimeData.InitializeAs(InstanceType);
+		}
 		return static_cast<TStateTreeEditorNode<T>&>(CondNode);
 	}
 
@@ -275,6 +279,10 @@ public:
 		{
 			CondNode.Instance.InitializeAs(InstanceType);
 		}
+		if (const UScriptStruct* InstanceType = Cast<const UScriptStruct>(Node.GetExecutionRuntimeDataType()))
+		{
+			CondNode.ExecutionRuntimeData.InitializeAs(InstanceType);
+		}
 		return static_cast<TStateTreeEditorNode<T>&>(CondNode);
 	}
 
@@ -293,6 +301,10 @@ public:
 		{
 			TaskItem.Instance.InitializeAs(InstanceType);
 		}
+		if (const UScriptStruct* InstanceType = Cast<const UScriptStruct>(Node.GetExecutionRuntimeDataType()))
+		{
+			TaskItem.ExecutionRuntimeData.InitializeAs(InstanceType);
+		}
 		return static_cast<TStateTreeEditorNode<T>&>(TaskItem);
 	}
 
diff --git a/Engine/Plugins/Runtime/StateTree/Source/StateTreeModule/Private/StateTree.cpp b/Engine/Plugins/Runtime/StateTree/Source/StateTreeModule/Private/StateTree.cpp
--- a/Engine/Plugins/Runtime/StateTree/Source/StateTreeModule/Private/StateTree.cpp
+++ b/Engine/Plugins/Runtime/StateTree/Source/StateTreeModule/Private/StateTree.cpp
@@ -244,6 +244,7 @@ void UStateTree::ResetCompiled()
 	Nodes.Reset();
 	DefaultInstanceData.Reset();
 	DefaultEvaluationScopeInstanceData.Reset();
+	DefaultExecutionRuntimeData.Reset();
 	SharedInstanceData.Reset();
 	ContextDataDescs.Reset();
 	PropertyBindings.Reset();
@@ -377,6 +378,17 @@ void UStateTree::PostLoad()
 		FStructView NodeView = Nodes[NodeIndex];
 		if (FStateTreeNodeBase* Node = NodeView.GetPtr<FStateTreeNodeBase>())
 		{
+			auto PostLoadInstance = [&Node]<typename T>(T & Container, FStateTreeIndex16 Index)
+			{
+				if (Container.IsObject(Index.Get()))
+				{
+					Node->PostLoad(Container.GetMutableObject(Index.Get()));
+				}
+				else
+				{
+					Node->PostLoad(Container.GetMutableStruct(Index.Get()));
+				}
+			};
 			if (Node->InstanceTemplateIndex.IsValid())
 			{
 				const bool bIsUsingSharedInstanceData = Node->InstanceDataHandle.GetSource() == EStateTreeDataSourceType::SharedInstanceData
@@ -385,26 +397,15 @@ void UStateTree::PostLoad()
 					|| Node->InstanceDataHandle.GetSource() == EStateTreeDataSourceType::EvaluationScopeInstanceDataObject;
 				if (bIsUsingEvaluationScopeInstanceData)
 				{
-					if (DefaultEvaluationScopeInstanceData.IsObject(Node->InstanceTemplateIndex.Get()))
-					{
-						Node->PostLoad(DefaultEvaluationScopeInstanceData.GetMutableObject(Node->InstanceTemplateIndex.Get()));
-					}
-					else
-					{
-						Node->PostLoad(DefaultEvaluationScopeInstanceData.GetMutableStruct(Node->InstanceTemplateIndex.Get()));
-					}
+					PostLoadInstance(DefaultEvaluationScopeInstanceData, Node->InstanceTemplateIndex);
+				}
+				else if (bIsUsingSharedInstanceData)
+				{
+					PostLoadInstance(SharedInstanceData, Node->InstanceTemplateIndex);
 				}
 				else
 				{
-					FStateTreeInstanceData& SourceInstanceData = bIsUsingSharedInstanceData ? SharedInstanceData : DefaultInstanceData;
-					if (SourceInstanceData.IsObject(Node->InstanceTemplateIndex.Get()))
-					{
-						Node->PostLoad(SourceInstanceData.GetMutableObject(Node->InstanceTemplateIndex.Get()));
-					}
-					else
-					{
-						Node->PostLoad(SourceInstanceData.GetMutableStruct(Node->InstanceTemplateIndex.Get()));
-					}
+					PostLoadInstance(DefaultInstanceData, Node->InstanceTemplateIndex);
 				}
 			}
 		}
@@ -502,115 +503,135 @@ bool UStateTree::ValidateInstanceData()
 		const FStateTreeNodeBase* Node = NodeView.GetPtr<const FStateTreeNodeBase>();
 		if (Node && Node->InstanceTemplateIndex.IsValid())
 		{
-			const UStruct* CurrentInstanceDataType = nullptr;
-			{
-				 const bool bIsUsingSharedInstanceData = Node->InstanceDataHandle.GetSource() == EStateTreeDataSourceType::SharedInstanceData
-					|| Node->InstanceDataHandle.GetSource() == EStateTreeDataSourceType::SharedInstanceDataObject;
-				 const bool bIsUsingEvaluationScopeInstanceData = Node->InstanceDataHandle.GetSource() == EStateTreeDataSourceType::EvaluationScopeInstanceData
-					 || Node->InstanceDataHandle.GetSource() == EStateTreeDataSourceType::EvaluationScopeInstanceDataObject;
-				if (bIsUsingEvaluationScopeInstanceData)
-				{
-					if (DefaultEvaluationScopeInstanceData.IsObject(Node->InstanceTemplateIndex.Get()))
-					{
-						const UObject* InstanceObject = DefaultEvaluationScopeInstanceData.GetObject(Node->InstanceTemplateIndex.Get());
-						CurrentInstanceDataType = InstanceObject ? InstanceObject->GetClass() : nullptr;
-					}
-					else
-					{
-						CurrentInstanceDataType = DefaultEvaluationScopeInstanceData.GetStruct(Node->InstanceTemplateIndex.Get()).GetScriptStruct();
-					}
-				}
-				else
+			auto TestInstanceData = [this, Node, Index, &bResult](const UStruct* CurrentInstanceDataType, const UStruct* DesiredInstanceDataType)
 				{
-					const FStateTreeInstanceData& SourceInstanceData = bIsUsingSharedInstanceData ? SharedInstanceData : DefaultInstanceData;
-					if (SourceInstanceData.IsObject(Node->InstanceTemplateIndex.Get()))
+					if (CurrentInstanceDataType == nullptr)
 					{
-						const UObject* InstanceObject = SourceInstanceData.GetObject(Node->InstanceTemplateIndex.Get());
-						CurrentInstanceDataType = InstanceObject ? InstanceObject->GetClass() : nullptr;
+						UE_LOG(LogStateTree, Error, TEXT("%s: node (%d) '%s' with name '%s' failed. Missing instance value, possibly due to Blueprint class or C++ class/struct template deletion.")
+							, *GetPathName()
+							, Index
+							, *WriteToString<64>(Node->StaticStruct()->GetFName())
+							, *WriteToString<128>(Node->Name)
+							);
+
+						bResult = false;
+						return;
 					}
-					else
+
+					auto HasNewerVersionExists = [](TNotNull<const UStruct*> InstanceDataType)
 					{
-						CurrentInstanceDataType = SourceInstanceData.GetStruct(Node->InstanceTemplateIndex.Get()).GetScriptStruct();
-					}
-				}
-			}
+						// Is the class/scriptstruct a blueprint that got replaced by another class.
+						bool bHasNewerVersionExists = InstanceDataType->HasAnyFlags(RF_NewerVersionExists);
+						if (!bHasNewerVersionExists)
+						{
+							if (const UClass* InstanceDataClass = Cast<UClass>(InstanceDataType))
+							{
+								bHasNewerVersionExists = InstanceDataClass->HasAnyClassFlags(CLASS_NewerVersionExists);
+							}
+							else if (const UScriptStruct* InstanceDataStruct = Cast<UScriptStruct>(InstanceDataType))
+							{
+								bHasNewerVersionExists = (InstanceDataStruct->StructFlags & STRUCT_NewerVersionExists) != 0;
+							}
+						}
+						return bHasNewerVersionExists;
+					};
 
-			if (CurrentInstanceDataType == nullptr)
-			{
-				UE_LOG(LogStateTree, Error, TEXT("%s: node (%d) '%s' with name '%s' failed. Missing instance value, possibly due to Blueprint class or C++ class/struct template deletion.")
-					, *GetPathName()
-					, Index
-					, *WriteToString<64>(Node->StaticStruct()->GetFName())
-					, *WriteToString<128>(Node->Name)
-					);
+					if (HasNewerVersionExists(CurrentInstanceDataType))
+					{
+						bool bLogError = true;
+#if WITH_EDITOR
+						OutOfDateStructs.Add(CurrentInstanceDataType);
+						bLogError = false;
+#endif
 
-				bResult = false;
-				continue;
-			}
+						if (bLogError)
+						{
+							UE_LOG(LogStateTree, Error, TEXT("%s: node '%s' failed. The source Instance Data type '%s' has a newer version."), *GetPathName(), *WriteToString<64>(Node->StaticStruct()->GetFName()), *WriteToString<64>(CurrentInstanceDataType->GetFName()));
+						}
 
-			auto HasNewerVersionExists = [](const UStruct* InstanceDataType)
-			{
-				// Is the class/scriptstruct a blueprint that got replaced by another class.
-				bool bHasNewerVersionExists = InstanceDataType->HasAnyFlags(RF_NewerVersionExists);
-				if (!bHasNewerVersionExists)
-				{
-					if (const UClass* InstanceDataClass = Cast<UClass>(InstanceDataType))
-					{
-						bHasNewerVersionExists = InstanceDataClass->HasAnyClassFlags(CLASS_NewerVersionExists);
+						bResult = false;
 					}
-					else if (const UScriptStruct* InstanceDataStruct = Cast<UScriptStruct>(InstanceDataType))
+
 					{
-						bHasNewerVersionExists = (InstanceDataStruct->StructFlags & STRUCT_NewerVersionExists) != 0;
-					}
-				}
-				return bHasNewerVersionExists;
-			};
+						// The FMyInstance::StaticStruct doesn't get a notification like the other objects when reinstanced.
+						const bool bDesiredHasNewerVersion = HasNewerVersionExists(DesiredInstanceDataType);
 
-			if (HasNewerVersionExists(CurrentInstanceDataType))
-			{
-				bool bLogError = true;
+						// Use strict testing so that the users will have option to initialize data mismatch if the type changes (even if potentially compatible).
+						if (CurrentInstanceDataType != DesiredInstanceDataType
+							&& !bDesiredHasNewerVersion)
+						{
+							bool bLogError = true;
 #if WITH_EDITOR
-				OutOfDateStructs.Add(CurrentInstanceDataType);
-				bLogError = false;
+							const UClass* CurrentInstanceDataClass = Cast<UClass>(CurrentInstanceDataType);
+							const UClass* DesiredInstanceDataClass = Cast<UClass>(DesiredInstanceDataType);
+							if (CurrentInstanceDataClass && DesiredInstanceDataClass)
+							{
+								// Because of the loading order. It's possible that the OnObjectsReinstanced did not completed.
+								if (CurrentInstanceDataClass->ClassGeneratedBy == DesiredInstanceDataClass->ClassGeneratedBy)
+								{
+									OutOfDateStructs.Add(CurrentInstanceDataType);
+									bLogError = false;
+								}
+							}
 #endif
-
-				if (bLogError)
-				{
-					UE_LOG(LogStateTree, Error, TEXT("%s: node '%s' failed. The source Instance Data type '%s' has a newer version."), *GetPathName(), *WriteToString<64>(Node->StaticStruct()->GetFName()), *WriteToString<64>(CurrentInstanceDataType->GetFName()));
-				}
-
-				bResult = false;
-			}
+							if (bLogError)
+							{
+								UE_LOG(LogStateTree, Error, TEXT("%s: node '%s' failed. The source Instance Data type '%s' does not match '%s'"), *GetPathName(), *WriteToString<64>(Node->StaticStruct()->GetFName()), *GetNameSafe(CurrentInstanceDataType), *GetNameSafe(DesiredInstanceDataType));
+							}
+							bResult = false;
+						}
+					}
+				};
 
 			{
-				const UStruct* DesiredInstanceDataType = Node->GetInstanceDataType();
-				// The FMyInstance::StaticStruct doesn't get a notification like the other objects when reinstanced.
-				const bool bDesiredHasNewerVersion = HasNewerVersionExists(DesiredInstanceDataType);
-
-				// Use strict testing so that the users will have option to initialize data mismatch if the type changes (even if potentially compatible).
-				if (CurrentInstanceDataType != DesiredInstanceDataType
-					&& !bDesiredHasNewerVersion)
+				const UStruct* CurrentInstanceDataType = nullptr;
 				{
-					bool bLogError = true;
-#if WITH_EDITOR
-					const UClass* CurrentInstanceDataClass = Cast<UClass>(CurrentInstanceDataType);
-					const UClass* DesiredInstanceDataClass = Cast<UClass>(DesiredInstanceDataType);
-					if (CurrentInstanceDataClass && DesiredInstanceDataClass)
+					const bool bIsUsingSharedInstanceData = Node->InstanceDataHandle.GetSource() == EStateTreeDataSourceType::SharedInstanceData
+						|| Node->InstanceDataHandle.GetSource() == EStateTreeDataSourceType::SharedInstanceDataObject;
+					const bool bIsUsingEvaluationScopeInstanceData = Node->InstanceDataHandle.GetSource() == EStateTreeDataSourceType::EvaluationScopeInstanceData
+						|| Node->InstanceDataHandle.GetSource() == EStateTreeDataSourceType::EvaluationScopeInstanceDataObject;
+					if (bIsUsingEvaluationScopeInstanceData)
 					{
-						// Because of the loading order. It's possible that the OnObjectsReinstanced did not completed.
-						if (CurrentInstanceDataClass->ClassGeneratedBy == DesiredInstanceDataClass->ClassGeneratedBy)
+						if (DefaultEvaluationScopeInstanceData.IsObject(Node->InstanceTemplateIndex.Get()))
+						{
+							const UObject* InstanceObject = DefaultEvaluationScopeInstanceData.GetObject(Node->InstanceTemplateIndex.Get());
+							CurrentInstanceDataType = InstanceObject ? InstanceObject->GetClass() : nullptr;
+						}
+						else
 						{
-							OutOfDateStructs.Add(CurrentInstanceDataType);
-							bLogError = false;
+							CurrentInstanceDataType = DefaultEvaluationScopeInstanceData.GetStruct(Node->InstanceTemplateIndex.Get()).GetScriptStruct();
 						}
 					}
-#endif
-					if (bLogError)
+					else
 					{
-						UE_LOG(LogStateTree, Error, TEXT("%s: node '%s' failed. The source Instance Data type '%s' does not match '%s'"), *GetPathName(), *WriteToString<64>(Node->StaticStruct()->GetFName()), *GetNameSafe(CurrentInstanceDataType), *GetNameSafe(DesiredInstanceDataType));
+						const FStateTreeInstanceData& SourceInstanceData = bIsUsingSharedInstanceData ? SharedInstanceData : DefaultInstanceData;
+						if (SourceInstanceData.IsObject(Node->InstanceTemplateIndex.Get()))
+						{
+							const UObject* InstanceObject = SourceInstanceData.GetObject(Node->InstanceTemplateIndex.Get());
+							CurrentInstanceDataType = InstanceObject ? InstanceObject->GetClass() : nullptr;
+						}
+						else
+						{
+							CurrentInstanceDataType = SourceInstanceData.GetStruct(Node->InstanceTemplateIndex.Get()).GetScriptStruct();
+						}
 					}
-					bResult = false;
 				}
+				TestInstanceData(CurrentInstanceDataType, Node->GetInstanceDataType());
+			}
+
+			if (Node->ExecutionRuntimeTemplateIndex.IsValid())
+			{
+				const UStruct* CurrentInstanceDataType = nullptr;
+				if (GetDefaultExecutionRuntimeData().IsObject(Node->ExecutionRuntimeTemplateIndex.Get()))
+				{
+					const UObject* InstanceObject = GetDefaultExecutionRuntimeData().GetObject(Node->ExecutionRuntimeTemplateIndex.Get());
+					CurrentInstanceDataType = InstanceObject ? InstanceObject->GetClass() : nullptr;
+				}
+				else
+				{
+					CurrentInstanceDataType = GetDefaultExecutionRuntimeData().GetStruct(Node->ExecutionRuntimeTemplateIndex.Get()).GetScriptStruct();
+				}
+				TestInstanceData(CurrentInstanceDataType, Node->GetExecutionRuntimeDataType());
 			}
 		}
 	}
@@ -712,6 +733,12 @@ bool UStateTree::Link()
 		UE_LOG(LogStateTree, Error, TEXT("%s: State Tree asset was not properly loaded (missing evaluation scope instance data). See log for loading failures, or recompile the StateTree asset."), *GetPathName());
 		return false;
 	}
+
+	if (!GetDefaultExecutionRuntimeData().AreAllInstancesValid())
+	{
+		UE_LOG(LogStateTree, Error, TEXT("%s: State Tree asset was not properly loaded (missing execution runtime data). See log for loading failures, or recompile the StateTree asset."), *GetPathName());
+		return false;
+	}
 	
 	if (!PatchBindings())
 	{
@@ -1469,53 +1496,35 @@ FString UStateTree::DebugInternalLayoutAsString() const
 		PrintViewNameSafe(Index, Node);
 	}
 
-	// Instance InstanceData data (e.g. tasks)
-	DebugString.Appendf(TEXT("\nInstance Data(%d)\n"), DefaultInstanceData.Num());
-	for (int32 Index = 0; Index < DefaultInstanceData.Num(); Index++)
-	{
-		if (DefaultInstanceData.IsObject(Index))
+	auto PrintInstanceData = [&DebugString, &PrintObjectNameSafe, &PrintViewNameSafe]<typename T>(T& Container, const FStringView Name)
 		{
-			const UObject* Data = DefaultInstanceData.GetObject(Index);
-			PrintObjectNameSafe(Index, Data);
-		}
-		else
-		{
-			const FConstStructView Data = DefaultInstanceData.GetStruct(Index);
-			PrintViewNameSafe(Index, Data);
-		}
-	}
+			DebugString.Appendf(TEXT("\n%s(%d)\n"), Name.GetData(), Container.Num());
+			for (int32 Index = 0; Index < Container.Num(); Index++)
+			{
+				if (Container.IsObject(Index))
+				{
+					const UObject* Data = Container.GetObject(Index);
+					PrintObjectNameSafe(Index, Data);
+				}
+				else
+				{
+					const FConstStructView Data = Container.GetStruct(Index);
+					PrintViewNameSafe(Index, Data);
+				}
+			}
+		};
+
+	// Instance InstanceData data (e.g. tasks)
+	PrintInstanceData(DefaultInstanceData, TEXT("Instance Data"));
 
 	// Shared Instance data (e.g. conditions/evaluators)
-	DebugString.Appendf(TEXT("\nShared Instance Data(%d)\n"), SharedInstanceData.Num());
-	for (int32 Index = 0; Index < SharedInstanceData.Num(); Index++)
-	{
-		if (SharedInstanceData.IsObject(Index))
-		{
-			const UObject* Data = SharedInstanceData.GetObject(Index);
-			PrintObjectNameSafe(Index, Data);
-		}
-		else
-		{
-			const FConstStructView Data = SharedInstanceData.GetStruct(Index);
-			PrintViewNameSafe(Index, Data);
-		}
-	}
+	PrintInstanceData(SharedInstanceData, TEXT("Shared Instance Data"));
 
-	// Execution Scope InstanceData data (e.g. tasks)
-	DebugString.Appendf(TEXT("\nExecution Scope(%d)\n"), DefaultEvaluationScopeInstanceData.Num());
-	for (int32 Index = 0; Index < DefaultEvaluationScopeInstanceData.Num(); Index++)
-	{
-		if (DefaultEvaluationScopeInstanceData.IsObject(Index))
-		{
-			const UObject* Data = DefaultEvaluationScopeInstanceData.GetObject(Index);
-			PrintObjectNameSafe(Index, Data);
-		}
-		else
-		{
-			const FConstStructView Data = DefaultEvaluationScopeInstanceData.GetStruct(Index);
-			PrintViewNameSafe(Index, Data);
-		}
-	}
+	// Evaluation Scope InstanceData data (e.g. conditions/evaluators)
+	PrintInstanceData(DefaultEvaluationScopeInstanceData, TEXT("Evaluation Scope Instance Data"));
+
+	// Execution Runtime InstanceData data (e.g. tasks)
+	PrintInstanceData(GetDefaultExecutionRuntimeData(), TEXT("Execution Runtime Instance Data"));
 
 	// External data (e.g. fragments, subsystems)
 	DebugString.Appendf(TEXT("\nExternal Data(%d)\n"), ExternalDataDescs.Num());
diff --git a/Engine/Plugins/Runtime/StateTree/Source/StateTreeModule/Private/StateTreeExecutionContext.cpp b/Engine/Plugins/Runtime/StateTree/Source/StateTreeModule/Private/StateTreeExecutionContext.cpp
--- a/Engine/Plugins/Runtime/StateTree/Source/StateTreeModule/Private/StateTreeExecutionContext.cpp
+++ b/Engine/Plugins/Runtime/StateTree/Source/StateTreeModule/Private/StateTreeExecutionContext.cpp
@@ -1424,6 +1424,7 @@ EStateTreeRunStatus FStateTreeExecutionContext::Start(FStartParameters Parameter
 	const FExecutionFrameHandle InitFrameHandle = FExecutionFrameHandle(&RootStateTree, FStateTreeStateHandle::Root);
 	constexpr bool bIsGlobalFrame = true;
 	FStateTreeExecutionFrame& InitFrame = SelectStateResult->MakeAndAddTemporaryFrame(InitFrameID, InitFrameHandle, bIsGlobalFrame);
+	InitFrame.ExecutionRuntimeIndexBase = FStateTreeIndex16(Storage.AddExecutionRuntimeData(GetOwner(), InitFrameHandle));
 	InitFrame.GlobalParameterDataHandle = FStateTreeDataHandle(EStateTreeDataSourceType::GlobalParameterData);
 
 	if (!CollectActiveExternalData(SelectStateResult->TemporaryFrames))
@@ -2278,11 +2279,7 @@ void FStateTreeExecutionContext::UpdateInstanceData(const TSharedPtr<FSelectStat
 
 				if (CurrentFrame == nullptr)
 				{
-					CurrentFrame = Args->TemporaryFrames.FindByPredicate(
-						[FrameID = SelectedState.GetFrameID()](const FStateTreeExecutionFrame& Frame)
-						{
-							return Frame.FrameID == FrameID;
-						});
+					CurrentFrame = Args->FindTemporaryFrame(SelectedState.GetFrameID());
 				}
 				check(CurrentFrame);
 				CurrentStateTree = CurrentFrame->StateTree;
@@ -2587,6 +2584,34 @@ FStateTreeDataView FStateTreeExecutionContext::GetDataView(const FStateTreeExecu
 	return GetDataView(ParentFrame, CurrentFrame, Handle);
 }
 
+FStateTreeDataView FStateTreeExecutionContext::GetExecutionRuntimeDataView() const
+{
+	check(CurrentNode && CurrentlyProcessedFrame);
+	check(CurrentlyProcessedFrame->ExecutionRuntimeIndexBase.IsValid());
+
+	if (!ensureMsgf(CurrentNode->ExecutionRuntimeTemplateIndex.IsValid(), TEXT("The node doesn't support execution runtime data.")))
+	{
+		return FStateTreeDataView();
+	}
+
+	const UE::StateTree::InstanceData::FInstanceContainer& Container = Storage.GetExecutionRuntimeData();
+	const int32 ActiveIndex = CurrentlyProcessedFrame->ExecutionRuntimeIndexBase.Get() + CurrentNode->ExecutionRuntimeTemplateIndex.Get();
+
+	if (!ensureMsgf(Container.IsValidIndex(ActiveIndex), TEXT("Invalid execution runtime data index.")))
+	{
+		return FStateTreeDataView();
+	}
+
+	if (Container.IsObject(ActiveIndex))
+	{
+		return Container.GetMutableObject(ActiveIndex);
+	}
+	else
+	{
+		return const_cast<UE::StateTree::InstanceData::FInstanceContainer&>(Container).GetMutableStruct(ActiveIndex);
+	}
+}
+
 EStateTreeRunStatus FStateTreeExecutionContext::ForceTransition(const FRecordedStateTreeTransitionResult& RecordedTransition)
 {
 	using namespace UE::StateTree;
@@ -3373,11 +3398,7 @@ EStateTreeRunStatus FStateTreeExecutionContext::EnterState(const TSharedPtr<FSel
 		if (FrameIndex == INDEX_NONE)
 		{
 			check(bIsCurrentFrameNew);
-			FStateTreeExecutionFrame* FoundNewFrame = Args.TemporaryFrames.FindByPredicate(
-				[FrameID = SelectedState.GetFrameID()](const FStateTreeExecutionFrame& Frame)
-				{
-					return Frame.FrameID == FrameID;
-				});
+			FStateTreeExecutionFrame* FoundNewFrame = Args.FindTemporaryFrame(SelectedState.GetFrameID());
 			if (!ensure(FoundNewFrame != nullptr))
 			{
 				return EStateTreeRunStatus::Failed;
@@ -6500,7 +6521,8 @@ bool FStateTreeExecutionContext::SelectStateInternal(
 			if (PreviousSelectedState.GetStateHandle() != NextState.Parent)
 			{
 				const FExecutionFrameHandle FrameHandle = FExecutionFrameHandle(NextStateContext.StateTree, NextStateContext.StateHandle);
-				const FStateTreeExecutionFrame& NewFrame = OutSelectionResult->MakeAndAddTemporaryFrameWithNewRoot(FActiveFrameID(Storage.GenerateUniqueId()), FrameHandle, *PreviousFramePtr);
+				FStateTreeExecutionFrame& NewFrame = OutSelectionResult->MakeAndAddTemporaryFrameWithNewRoot(FActiveFrameID(Storage.GenerateUniqueId()), FrameHandle, *PreviousFramePtr);
+				NewFrame.ExecutionRuntimeIndexBase = FStateTreeIndex16(Storage.AddExecutionRuntimeData(GetOwner(), FrameHandle));
 				OutSelectionResult->SelectedFrames.Add(NewFrame.FrameID);
 				OnExitScope.bFrameAdded = true;
 				bNewFrameCreated = true;
@@ -6523,7 +6545,8 @@ bool FStateTreeExecutionContext::SelectStateInternal(
 				if (!NextState.Parent.IsValid() && MissingFramePtr->RootState != NextStateHandle)
 				{
 					const FExecutionFrameHandle FrameHandle = FExecutionFrameHandle(NextStateContext.StateTree, NextStateContext.StateHandle);
-					const FStateTreeExecutionFrame& NewFrame = OutSelectionResult->MakeAndAddTemporaryFrameWithNewRoot(FActiveFrameID(Storage.GenerateUniqueId()), FrameHandle, *MissingFramePtr);
+					FStateTreeExecutionFrame& NewFrame = OutSelectionResult->MakeAndAddTemporaryFrameWithNewRoot(FActiveFrameID(Storage.GenerateUniqueId()), FrameHandle, *MissingFramePtr);
+					NewFrame.ExecutionRuntimeIndexBase = FStateTreeIndex16(Storage.AddExecutionRuntimeData(GetOwner(), FrameHandle));
 					NextFrameID = NewFrame.FrameID;
 					bNewFrameCreated = true;
 				}
@@ -6531,7 +6554,8 @@ bool FStateTreeExecutionContext::SelectStateInternal(
 			else if (MissingFramePtr->RootState != NextStateHandle)
 			{
 				const FExecutionFrameHandle FrameHandle = FExecutionFrameHandle(NextStateContext.StateTree, NextStateContext.StateHandle);
-				const FStateTreeExecutionFrame& NewFrame = OutSelectionResult->MakeAndAddTemporaryFrameWithNewRoot(FActiveFrameID(Storage.GenerateUniqueId()), FrameHandle, *MissingFramePtr);
+				FStateTreeExecutionFrame& NewFrame = OutSelectionResult->MakeAndAddTemporaryFrameWithNewRoot(FActiveFrameID(Storage.GenerateUniqueId()), FrameHandle, *MissingFramePtr);
+				NewFrame.ExecutionRuntimeIndexBase = FStateTreeIndex16(Storage.AddExecutionRuntimeData(GetOwner(), FrameHandle));
 				NextFrameID = NewFrame.FrameID;
 				bNewFrameCreated = true;
 			}
@@ -7135,9 +7159,10 @@ bool FStateTreeExecutionContext::SelectStateInternal_Linked(
 		constexpr bool bIsGlobalFrame = false;
 		FStateTreeExecutionFrame& NewFrame = OutSelectionResult->MakeAndAddTemporaryFrame(FActiveFrameID(Storage.GenerateUniqueId()), LinkStateFrameHandle, bIsGlobalFrame);
 		NewFrame.ExternalDataBaseIndex = ExternalDataBaseIndex;
-		NewFrame.GlobalParameterDataHandle = GlobalParameterDataHandle;
 		NewFrame.GlobalInstanceIndexBase = GlobalInstanceIndexBase;
+		NewFrame.ExecutionRuntimeIndexBase = FStateTreeIndex16(Storage.AddExecutionRuntimeData(GetOwner(), LinkStateFrameHandle));
 		NewFrame.StateParameterDataHandle = TargetState.ParameterDataHandle; // Temporary allocated earlier if did not exists.
+		NewFrame.GlobalParameterDataHandle = GlobalParameterDataHandle;
 
 		SelectedFrame = &NewFrame;
 	}
@@ -7156,6 +7181,10 @@ bool FStateTreeExecutionContext::SelectStateInternal_Linked(
 
 	if (!bSucceededToSelectState)
 	{
+		if (bIsNewFrame)
+		{
+			CleanFrame(Exec, NewInternalArgs.MissingSourceFrameID);
+		}
 		OutSelectionResult->SelectedFrames.Pop();
 	}
 
@@ -7296,6 +7325,7 @@ bool FStateTreeExecutionContext::SelectStateInternal_LinkedAsset(
 
 		FStateTreeExecutionFrame& NewFrame = OutSelectionResult->MakeAndAddTemporaryFrame(FActiveFrameID(Storage.GenerateUniqueId()), LinkStateFrameHandle, true);
 		NewFrame.ExternalDataBaseIndex = ExternalDataBaseIndex;
+		NewFrame.ExecutionRuntimeIndexBase = FStateTreeIndex16(Storage.AddExecutionRuntimeData(GetOwner(), LinkStateFrameHandle));
 		// Pass the linked state's parameters as global parameters to the linked asset.
 		NewFrame.GlobalParameterDataHandle = NextState.ParameterDataHandle;
 		// The state parameters will be from the root state.
diff --git a/Engine/Plugins/Runtime/StateTree/Source/StateTreeModule/Private/StateTreeInstanceContainer.cpp b/Engine/Plugins/Runtime/StateTree/Source/StateTreeModule/Private/StateTreeInstanceContainer.cpp
--- a/Engine/Plugins/Runtime/StateTree/Source/StateTreeModule/Private/StateTreeInstanceContainer.cpp
+++ b/Engine/Plugins/Runtime/StateTree/Source/StateTreeModule/Private/StateTreeInstanceContainer.cpp
@@ -3,24 +3,61 @@
 #include "StateTreeInstanceContainer.h"
 #include "StateTreeInstanceDataHelpers.h"
 
+#include "Experimental/ConcurrentLinearAllocator.h"
+
 #include UE_INLINE_GENERATED_CPP_BY_NAME(StateTreeInstanceContainer)
 
 namespace UE::StateTree::InstanceData
 {
 
-void FInstanceContainer::Init(TNotNull<UObject*> InOwner, TConstArrayView<FInstancedStruct> InStructs, FAddArgs Args)
-{
-	Private::AppendToInstanceStructContainer(InstanceStructs, InOwner, InStructs, Args.bDuplicateWrappedObject);
-}
+	void FInstanceContainer::Init(TNotNull<UObject*> InOwner, const FInstanceContainer& InStructs, FAddArgs Args)
+	{
+		InstanceStructs.Reset();
+		InstanceStructs = InStructs.InstanceStructs;
+		Private::PostAppendToInstanceStructContainer(InstanceStructs, InOwner, Args.bDuplicateWrappedObject, 0);
+	}
 
-void FInstanceContainer::Init(TNotNull<UObject*> InOwner, TConstArrayView<FConstStructView> InStructs, FAddArgs Args)
-{
-	Private::AppendToInstanceStructContainer(InstanceStructs, InOwner, InStructs, Args.bDuplicateWrappedObject);
-}
+	void FInstanceContainer::Init(TNotNull<UObject*> InOwner, TConstArrayView<FInstancedStruct> InStructs, FAddArgs Args)
+	{
+		InstanceStructs.Reset();
+		Private::AppendToInstanceStructContainer(InstanceStructs, InOwner, InStructs, Args.bDuplicateWrappedObject);
+	}
 
-bool FInstanceContainer::AreAllInstancesValid() const
-{
-	return Private::AreAllInstancesValid(InstanceStructs);
-}
+	void FInstanceContainer::Init(TNotNull<UObject*> InOwner, TConstArrayView<FConstStructView> InStructs, FAddArgs Args)
+	{
+		InstanceStructs.Reset();
+		Private::AppendToInstanceStructContainer(InstanceStructs, InOwner, InStructs, Args.bDuplicateWrappedObject);
+	}
+
+	int32 FInstanceContainer::Append(TNotNull<UObject*> InOwner, const FInstanceContainer& InStructs, FAddArgs Args)
+	{
+		TArray<FConstStructView, FNonconcurrentLinearArrayAllocator> ToAppend;
+		ToAppend.Reserve(InStructs.InstanceStructs.Num());
+		for (const FConstStructView& Struct : InStructs.InstanceStructs)
+		{
+			ToAppend.Add(Struct);
+		}
+		return Append(InOwner, MakeConstArrayView(ToAppend), Args);
+	}
+
+	int32 FInstanceContainer::Append(TNotNull<UObject*> InOwner, TConstArrayView<FInstancedStruct> InStructs, FAddArgs Args)
+	{
+		return Private::AppendToInstanceStructContainer(InstanceStructs, InOwner, InStructs, Args.bDuplicateWrappedObject);
+	}
+
+	int32 FInstanceContainer::Append(TNotNull<UObject*> InOwner, TConstArrayView<FConstStructView> InStructs, FAddArgs Args)
+	{
+		return Private::AppendToInstanceStructContainer(InstanceStructs, InOwner, InStructs, Args.bDuplicateWrappedObject);
+	}
+
+	bool FInstanceContainer::AreAllInstancesValid() const
+	{
+		return Private::AreAllInstancesValid(InstanceStructs);
+	}
+
+	int32 FInstanceContainer::GetAllocatedMemory() const
+	{
+		return Private::GetAllocatedMemory(InstanceStructs);
+	}
 
 } // namespace UE::StateTree::InstanceData
diff --git a/Engine/Plugins/Runtime/StateTree/Source/StateTreeModule/Private/StateTreeInstanceData.cpp b/Engine/Plugins/Runtime/StateTree/Source/StateTreeModule/Private/StateTreeInstanceData.cpp
--- a/Engine/Plugins/Runtime/StateTree/Source/StateTreeModule/Private/StateTreeInstanceData.cpp
+++ b/Engine/Plugins/Runtime/StateTree/Source/StateTreeModule/Private/StateTreeInstanceData.cpp
@@ -40,6 +40,22 @@ namespace UE::StateTree::InstanceData::Private
 		return true;
 	}
 
+	int32 GetAllocatedMemory(const FInstancedStructContainer& InstanceStructs)
+	{
+		int32 Size = InstanceStructs.GetAllocatedMemory();
+		for (FConstStructView Instance : InstanceStructs)
+		{
+			if (const FStateTreeInstanceObjectWrapper* Wrapper = Instance.GetPtr<const FStateTreeInstanceObjectWrapper>())
+			{
+				if (Wrapper->InstanceObject)
+				{
+					Size += Wrapper->InstanceObject->GetClass()->GetStructureSize();
+				}
+			}
+		}
+		return Size;
+	}
+
 	TNotNull<UObject*> CopyNodeInstance(TNotNull<UObject*> Instance, TNotNull<UObject*> InOwner, bool bDuplicate)
 	{
 		const UClass* InstanceClass = Instance->GetClass();
@@ -354,6 +370,8 @@ namespace UE::StateTree::InstanceData
 FStateTreeInstanceStorage::FStateTreeInstanceStorage(const FStateTreeInstanceStorage& Other)
 	: InstanceStructs(Other.InstanceStructs)
 	, ExecutionState(Other.ExecutionState)
+	, ExecutionRuntimeData(Other.ExecutionRuntimeData)
+	, ExecutionRuntimeDataInfos(Other.ExecutionRuntimeDataInfos)
 	, TemporaryInstances(Other.TemporaryInstances)
 	, EventQueue(MakeShared<FStateTreeEventQueue>(*Other.EventQueue))
 	, TransitionRequests(Other.TransitionRequests)
@@ -371,6 +389,8 @@ FStateTreeInstanceStorage::FStateTreeInstanceStorage(const FStateTreeInstanceSto
 FStateTreeInstanceStorage::FStateTreeInstanceStorage(FStateTreeInstanceStorage&& Other) noexcept
 	: InstanceStructs(MoveTemp(Other.InstanceStructs))
 	, ExecutionState(MoveTemp(Other.ExecutionState))
+	, ExecutionRuntimeData(MoveTemp(Other.ExecutionRuntimeData))
+	, ExecutionRuntimeDataInfos(MoveTemp(Other.ExecutionRuntimeDataInfos))
 	, TemporaryInstances(MoveTemp(Other.TemporaryInstances))
 	, EventQueue(Other.EventQueue)
 	, TransitionRequests(MoveTemp(Other.TransitionRequests))
@@ -393,6 +413,8 @@ FStateTreeInstanceStorage& FStateTreeInstanceStorage::operator=(const FStateTree
 {
 	InstanceStructs = Other.InstanceStructs;
 	ExecutionState = Other.ExecutionState;
+	ExecutionRuntimeData = Other.ExecutionRuntimeData;
+	ExecutionRuntimeDataInfos = Other.ExecutionRuntimeDataInfos;
 	TemporaryInstances = Other.TemporaryInstances;
 	EventQueue = MakeShared<FStateTreeEventQueue>(*Other.EventQueue);
 	TransitionRequests = Other.TransitionRequests;
@@ -415,6 +437,8 @@ FStateTreeInstanceStorage& FStateTreeInstanceStorage::operator=(FStateTreeInstan
 {
 	InstanceStructs = MoveTemp(Other.InstanceStructs);
 	ExecutionState = MoveTemp(Other.ExecutionState);
+	ExecutionRuntimeData = MoveTemp(Other.ExecutionRuntimeData);
+	ExecutionRuntimeDataInfos = MoveTemp(Other.ExecutionRuntimeDataInfos);
 	TemporaryInstances = MoveTemp(Other.TemporaryInstances);
 	EventQueue = Other.EventQueue;
 	Other.EventQueue = MakeShared<FStateTreeEventQueue>();
@@ -483,7 +507,34 @@ void FStateTreeInstanceStorage::ResetTransitionRequests()
 
 bool FStateTreeInstanceStorage::AreAllInstancesValid() const
 {
-	return UE::StateTree::InstanceData::Private::AreAllInstancesValid(InstanceStructs);
+	return UE::StateTree::InstanceData::Private::AreAllInstancesValid(InstanceStructs)
+		&& ExecutionRuntimeData.AreAllInstancesValid();
+}
+
+int32 FStateTreeInstanceStorage::AddExecutionRuntimeData(TNotNull<UObject*> Owner, const UE::StateTree::FExecutionFrameHandle FrameHandle)
+{
+	using namespace UE::StateTree;
+	using namespace UE::StateTree::InstanceData;
+
+	check(FrameHandle.IsValid());
+
+	const FObjectKey StateTreeKey = FrameHandle.GetStateTree();
+	const FExecutionRuntimeInfo* Info = ExecutionRuntimeDataInfos.FindByPredicate([StateTreeKey](const FExecutionRuntimeInfo& Other)
+		{
+			return Other.StateTree == StateTreeKey;
+		});
+	if (Info != nullptr)
+	{
+		return Info->StartIndex;
+	}
+
+	const int32 StartIndex = ExecutionRuntimeData.Append(Owner, FrameHandle.GetStateTree()->GetDefaultExecutionRuntimeData(), FInstanceContainer::FAddArgs());
+
+	FExecutionRuntimeInfo& NewInfo = ExecutionRuntimeDataInfos.AddDefaulted_GetRef();
+	NewInfo.StateTree = StateTreeKey;
+	NewInfo.StartIndex = StartIndex;
+
+	return StartIndex;
 }
 
 FStructView FStateTreeInstanceStorage::AddTemporaryInstance(UObject& InOwner, const FStateTreeExecutionFrame& Frame, const FStateTreeIndex16 OwnerNodeIndex, const FStateTreeDataHandle DataHandle, FConstStructView NewInstanceData)
@@ -587,6 +638,8 @@ void FStateTreeInstanceStorage::Reset()
 {
 	InstanceStructs.Reset();
 	ExecutionState.Reset();
+	ExecutionRuntimeData.Reset();
+	ExecutionRuntimeDataInfos.Reset();
 	TemporaryInstances.Reset();
 	if (bIsOwningEventQueue)
 	{
@@ -765,18 +818,8 @@ int32 FStateTreeInstanceData::GetEstimatedMemoryUsage() const
 	const FStateTreeInstanceStorage& Storage = GetStorage();
 	int32 Size = sizeof(FStateTreeInstanceData);
 
-	Size += Storage.InstanceStructs.GetAllocatedMemory();
-
-	for (FConstStructView Instance : Storage.InstanceStructs)
-	{
-		if (const FStateTreeInstanceObjectWrapper* Wrapper = Instance.GetPtr<const FStateTreeInstanceObjectWrapper>())
-		{
-			if (Wrapper->InstanceObject)
-			{
-				Size += Wrapper->InstanceObject->GetClass()->GetStructureSize();
-			}
-		}
-	}
+	Size += UE::StateTree::InstanceData::Private::GetAllocatedMemory(Storage.InstanceStructs);
+	Size += Storage.ExecutionRuntimeData.GetAllocatedMemory();
 
 	return Size;
 }
diff --git a/Engine/Plugins/Runtime/StateTree/Source/StateTreeModule/Private/StateTreeInstanceDataHelpers.h b/Engine/Plugins/Runtime/StateTree/Source/StateTreeModule/Private/StateTreeInstanceDataHelpers.h
--- a/Engine/Plugins/Runtime/StateTree/Source/StateTreeModule/Private/StateTreeInstanceDataHelpers.h
+++ b/Engine/Plugins/Runtime/StateTree/Source/StateTreeModule/Private/StateTreeInstanceDataHelpers.h
@@ -29,6 +29,9 @@ namespace UE::StateTree::InstanceData::Private
 	/** @return true if all instances are valid. */
 	[[nodiscard]] bool AreAllInstancesValid(const FInstancedStructContainer& InstanceStructs);
 
+	/** Returns number of bytes allocated for the array  */
+	[[nodiscard]] int32 GetAllocatedMemory(const FInstancedStructContainer& InstanceStructs);
+
 	/** Duplicates object, and tries to covert old BP classes (REINST_*) to their newer version. */
 	[[nodiscard]] TNotNull<UObject*> CopyNodeInstance(TNotNull<UObject*> Instance, TNotNull<UObject*> InOwner, bool bDuplicate);
 
@@ -37,10 +40,11 @@ namespace UE::StateTree::InstanceData::Private
 
 	/** Appends new items to the instance, and moves existing data into the allocated instances. */
 	template<typename TOtherType>
-	void AppendToInstanceStructContainer(FInstancedStructContainer& InstanceStructs, TNotNull<UObject*> InOwner, TConstArrayView<TOtherType> InStructs, bool bDuplicateWrappedObject)
+	int32 AppendToInstanceStructContainer(FInstancedStructContainer& InstanceStructs, TNotNull<UObject*> InOwner, TConstArrayView<TOtherType> InStructs, bool bDuplicateWrappedObject)
 	{
 		const int32 StartIndex = InstanceStructs.Num();
 		InstanceStructs.Append(InStructs);
 		PostAppendToInstanceStructContainer(InstanceStructs, InOwner, bDuplicateWrappedObject, StartIndex);
+		return StartIndex;
 	}
 } // namespace
diff --git a/Engine/Plugins/Runtime/StateTree/Source/StateTreeModule/Private/StateTreeTypes.cpp b/Engine/Plugins/Runtime/StateTree/Source/StateTreeModule/Private/StateTreeTypes.cpp
--- a/Engine/Plugins/Runtime/StateTree/Source/StateTreeModule/Private/StateTreeTypes.cpp
+++ b/Engine/Plugins/Runtime/StateTree/Source/StateTreeModule/Private/StateTreeTypes.cpp
@@ -149,7 +149,8 @@ bool FStateTreeDataHandle::IsObjectSource() const
 	return Source == EStateTreeDataSourceType::GlobalInstanceDataObject
 		|| Source == EStateTreeDataSourceType::ActiveInstanceDataObject
 		|| Source == EStateTreeDataSourceType::SharedInstanceDataObject
-		|| Source == EStateTreeDataSourceType::EvaluationScopeInstanceDataObject;
+		|| Source == EStateTreeDataSourceType::EvaluationScopeInstanceDataObject
+		|| Source == EStateTreeDataSourceType::ExecutionRuntimeDataObject;
 }
 
 FStateTreeDataHandle FStateTreeDataHandle::ToObjectSource() const
@@ -164,6 +165,8 @@ FStateTreeDataHandle FStateTreeDataHandle::ToObjectSource() const
 		return FStateTreeDataHandle(EStateTreeDataSourceType::SharedInstanceDataObject, Index, StateHandle);
 	case EStateTreeDataSourceType::EvaluationScopeInstanceData:
 		return FStateTreeDataHandle(EStateTreeDataSourceType::EvaluationScopeInstanceDataObject, Index, StateHandle);
+	case EStateTreeDataSourceType::ExecutionRuntimeData:
+		return FStateTreeDataHandle(EStateTreeDataSourceType::ExecutionRuntimeDataObject, Index, StateHandle);
 	default:
 		return *this;
 	}
@@ -191,6 +194,10 @@ FString FStateTreeDataHandle::Describe() const
 		return FString::Printf(TEXT("EvalScope[%d]"), Index);
 	case EStateTreeDataSourceType::EvaluationScopeInstanceDataObject:
 		return FString::Printf(TEXT("EvalScopeObj[%d]"), Index);
+	case EStateTreeDataSourceType::ExecutionRuntimeData:
+		return FString::Printf(TEXT("ExecRun[%d]"), Index);
+	case EStateTreeDataSourceType::ExecutionRuntimeDataObject:
+		return FString::Printf(TEXT("ExecRunObj[%d]"), Index);
 	case EStateTreeDataSourceType::ContextData:
 		return FString::Printf(TEXT("Context[%d]"), Index);
 	case EStateTreeDataSourceType::GlobalParameterData:
diff --git a/Engine/Plugins/Runtime/StateTree/Source/StateTreeModule/Public/StateTree.h b/Engine/Plugins/Runtime/StateTree/Source/StateTreeModule/Public/StateTree.h
--- a/Engine/Plugins/Runtime/StateTree/Source/StateTreeModule/Public/StateTree.h
+++ b/Engine/Plugins/Runtime/StateTree/Source/StateTreeModule/Public/StateTree.h
@@ -136,6 +136,12 @@ public:
 		return DefaultEvaluationScopeInstanceData;
 	}
 
+	/** @return Default execution runtime data. */
+	const UE::StateTree::InstanceData::FInstanceContainer& GetDefaultExecutionRuntimeData() const
+	{
+		return DefaultExecutionRuntimeData;
+	}
+
 	/** @return List of external data required by the state tree */
 	TConstArrayView<FStateTreeExternalDataDesc> GetExternalDataDescs() const
 	{
@@ -396,10 +402,14 @@ private:
 	UPROPERTY()
 	FStateTreeInstanceData DefaultInstanceData;
 
-	/** Default node instance data for evaluation scope (e.g. conditions, considerations, function) */
+	/** Default node instance data for evaluation scope (e.g. conditions, considerations, functions) */
 	UPROPERTY()
 	UE::StateTree::InstanceData::FInstanceContainer DefaultEvaluationScopeInstanceData;
 
+	/** Default node execution runtime data for execution runtime (e.g. tasks, conditions, considerations, functions) */
+	UPROPERTY()
+	UE::StateTree::InstanceData::FInstanceContainer DefaultExecutionRuntimeData;
+
 	/** Shared node instance data (e.g. conditions, considerations). */
 	UPROPERTY()
 	FStateTreeInstanceData SharedInstanceData;
diff --git a/Engine/Plugins/Runtime/StateTree/Source/StateTreeModule/Public/StateTreeExecutionContext.h b/Engine/Plugins/Runtime/StateTree/Source/StateTreeModule/Public/StateTreeExecutionContext.h
--- a/Engine/Plugins/Runtime/StateTree/Source/StateTreeModule/Public/StateTreeExecutionContext.h
+++ b/Engine/Plugins/Runtime/StateTree/Source/StateTreeModule/Public/StateTreeExecutionContext.h
@@ -690,6 +690,31 @@ PRAGMA_ENABLE_DEPRECATION_WARNINGS
 		return CurrentNodeInstanceData.template GetMutable<typename T::FInstanceDataType>();
 	}
 
+	/** @returns pointer to the execution runtime data of specified node. */
+	template <typename T>
+	T* GetExecutionRuntimeDataPtr(const FStateTreeNodeBase& Node) const
+	{
+		check(CurrentNodeDataHandle == Node.InstanceDataHandle && Node.InstanceDataHandle.IsValid());
+		return GetExecutionRuntimeDataView().template GetMutablePtr<T>();
+	}
+
+	/** @returns reference to the execution runtime data of specified node. */
+	template <typename T>
+	T& GetExecutionRuntimeData(const FStateTreeNodeBase& Node) const
+	{
+		check(CurrentNodeDataHandle == Node.InstanceDataHandle && Node.InstanceDataHandle.IsValid());
+		return GetExecutionRuntimeDataView().template GetMutable<T>();
+	}
+
+	/** @returns reference to the execution runtime data of specified node. Infers the execution runtime data type from the node's FExecutionRuntimeDataType. */
+	template <typename T>
+	typename T::FExecutionRuntimeDataType& GetExecutionRuntimeData(const T& Node) const
+	{
+		static_assert(TIsDerivedFrom<T, FStateTreeNodeBase>::IsDerived, "Expecting Node to derive from FStateTreeNodeBase.");
+		check(CurrentNodeDataHandle == Node.InstanceDataHandle && Node.InstanceDataHandle.IsValid());
+		return GetExecutionRuntimeDataView().template GetMutable<typename T::FExecutionRuntimeDataType>();
+	}
+
 	/** @returns reference to instance data struct that can be passed to lambdas. See TStateTreeInstanceDataStructRef for usage. */
 	template <typename T>
 	TStateTreeInstanceDataStructRef<typename T::FInstanceDataType> GetInstanceDataStructRef(const T& Node) const
@@ -799,6 +824,16 @@ protected:
 		UE_API FStateTreeExecutionFrame& MakeAndAddTemporaryFrame(UE::StateTree::FActiveFrameID FrameID, const UE::StateTree::FExecutionFrameHandle& FrameHandle, bool bIsGlobalFrame);
 		UE_API FStateTreeExecutionFrame& MakeAndAddTemporaryFrameWithNewRoot(UE::StateTree::FActiveFrameID FrameID, const UE::StateTree::FExecutionFrameHandle& FrameHandle, FStateTreeExecutionFrame& OtherFrame);
 
+		/** @return a temporary frame by ID. */
+		FStateTreeExecutionFrame* FindTemporaryFrame(UE::StateTree::FActiveFrameID FrameID)
+		{
+			return TemporaryFrames.FindByPredicate(
+				[FrameID = FrameID](const FStateTreeExecutionFrame& Frame)
+				{
+					return Frame.FrameID == FrameID;
+				});
+		}
+
 		//~ITemporaryStorage
 		virtual FFrameAndParent GetExecutionFrame(UE::StateTree::FActiveFrameID ID) override;
 		virtual UE::StateTree::FActiveState GetStateHandle(UE::StateTree::FActiveStateID ID) const override;
@@ -1354,6 +1389,9 @@ PRAGMA_ENABLE_DEPRECATION_WARNINGS
 	[[nodiscard]] UE_API FStateTreeDataView GetDataView(const FStateTreeExecutionFrame* ParentFrame, const FStateTreeExecutionFrame& CurrentFrame, const FStateTreeDataHandle Handle);
 	[[nodiscard]] UE_API FStateTreeDataView GetDataView(const FStateTreeExecutionFrame* ParentFrame, const FStateTreeExecutionFrame& CurrentFrame, const FPropertyBindingCopyInfo& CopyInfo);
 
+	/** @return data view for the execution runtime view of the current node and frame. */
+	[[nodiscard]] UE_API FStateTreeDataView GetExecutionRuntimeDataView() const;
+
 	/** @return true if handle source is valid cified handle relative to given frame. */
 	[[nodiscard]] UE_API bool IsHandleSourceValid(const FStateTreeExecutionFrame* ParentFrame, const FStateTreeExecutionFrame& CurrentFrame, const FStateTreeDataHandle Handle) const;
 	[[nodiscard]] UE_API bool IsHandleSourceValid(const FStateTreeExecutionFrame* ParentFrame, const FStateTreeExecutionFrame& CurrentFrame, const FPropertyBindingCopyInfo& CopyInfo) const;
diff --git a/Engine/Plugins/Runtime/StateTree/Source/StateTreeModule/Public/StateTreeExecutionRuntimeDataTypes.h b/Engine/Plugins/Runtime/StateTree/Source/StateTreeModule/Public/StateTreeExecutionRuntimeDataTypes.h
new file mode 100644
--- /dev/null
+++ b/Engine/Plugins/Runtime/StateTree/Source/StateTreeModule/Public/StateTreeExecutionRuntimeDataTypes.h
@@ -0,0 +1,30 @@
+// Copyright Epic Games, Inc. All Rights Reserved.
+
+#pragma once
+
+#include "StateTreeInstanceContainer.h"
+#include "StateTreeTypes.h"
+#include "UObject/ObjectKey.h"
+#include "StateTreeExecutionRuntimeDataTypes.generated.h"
+
+#define UE_API STATETREEMODULE_API
+
+class UStateTree;
+
+namespace UE::StateTree::InstanceData
+{
+	/** Helper structure that holds the execution runtime instances. */
+	USTRUCT()
+	struct FExecutionRuntimeData
+	{
+		GENERATED_BODY()
+
+		/** The state tree of the instances. */
+		TObjectKey<UStateTree> StateTree;
+
+		UPROPERTY()
+		UE::StateTree::InstanceData::FInstanceContainer Instances;
+	};
+}
+
+#undef UE_API
diff --git a/Engine/Plugins/Runtime/StateTree/Source/StateTreeModule/Public/StateTreeExecutionTypes.h b/Engine/Plugins/Runtime/StateTree/Source/StateTreeModule/Public/StateTreeExecutionTypes.h
--- a/Engine/Plugins/Runtime/StateTree/Source/StateTreeModule/Public/StateTreeExecutionTypes.h
+++ b/Engine/Plugins/Runtime/StateTree/Source/StateTreeModule/Public/StateTreeExecutionTypes.h
@@ -1032,7 +1032,7 @@ PRAGMA_ENABLE_DEPRECATION_WARNINGS
 	UPROPERTY(Transient)
 	FStateTreeTasksCompletionStatus ActiveTasksStatus;
 
-	/** Unique frame id for this frame. Can be used to identify the frame. */
+	/** Unique frame ID for this frame. Can be used to identify the frame. */
 	UE::StateTree::FActiveFrameID FrameID;
 
 	/**
@@ -1054,6 +1054,10 @@ PRAGMA_ENABLE_DEPRECATION_WARNINGS
 	UPROPERTY()
 	FStateTreeIndex16 ActiveInstanceIndexBase = FStateTreeIndex16::Invalid;
 
+	/** Index within the execution runtime data to the first execution runtime's instance data (e.g. tasks). */
+	UPROPERTY()
+	FStateTreeIndex16 ExecutionRuntimeIndexBase = FStateTreeIndex16::Invalid;
+
 	/** Handle to the state parameter data, exists in ParentFrame. */
 	UPROPERTY()
 	FStateTreeDataHandle StateParameterDataHandle = FStateTreeDataHandle::Invalid; 
diff --git a/Engine/Plugins/Runtime/StateTree/Source/StateTreeModule/Public/StateTreeInstanceContainer.h b/Engine/Plugins/Runtime/StateTree/Source/StateTreeModule/Public/StateTreeInstanceContainer.h
--- a/Engine/Plugins/Runtime/StateTree/Source/StateTreeModule/Public/StateTreeInstanceContainer.h
+++ b/Engine/Plugins/Runtime/StateTree/Source/StateTreeModule/Public/StateTreeInstanceContainer.h
@@ -40,10 +40,18 @@ namespace UE::StateTree::InstanceData
 			bool bDuplicateWrappedObject = true;
 		};
 
+		/** Initializes the array with specified items. */
+		UE_API void Init(TNotNull<UObject*> InOwner, const FInstanceContainer& InStructs, FAddArgs Args);
 		/** Initializes the array with specified items. */
 		UE_API void Init(TNotNull<UObject*> InOwner, TConstArrayView<FInstancedStruct> InStructs, FAddArgs Args);
 		/** Initializes the array with specified items. */
 		UE_API void Init(TNotNull<UObject*> InOwner, TConstArrayView<FConstStructView> InStructs, FAddArgs Args);
+		/** Append the specified items to the array. */
+		UE_API int32 Append(TNotNull<UObject*> InOwner, const FInstanceContainer& InStructs, FAddArgs Args);
+		/** Append the specified items to the array. */
+		UE_API int32 Append(TNotNull<UObject*> InOwner, TConstArrayView<FInstancedStruct> InStructs, FAddArgs Args);
+		/** Append the specified items to the array. */
+		UE_API int32 Append(TNotNull<UObject*> InOwner, TConstArrayView<FConstStructView> InStructs, FAddArgs Args);
 
 		/** @return number of items in the storage. */
 		[[nodiscard]] int32 Num() const
@@ -98,10 +106,29 @@ namespace UE::StateTree::InstanceData
 			InstanceStructs.Reset();
 		}
 
-		/** Returns number of bytes allocated for the array  */
-		[[nodiscard]] int32 GetAllocatedMemory() const
+		/** Returns the estimated memory allocated by the container. */
+		[[nodiscard]] UE_API int32 GetAllocatedMemory() const;
+
+		using FIterator = FInstancedStructContainer::FIterator;
+		using FConstIterator = FInstancedStructContainer::FConstIterator;
+
+		//~ For ranged for, do not use directly.
+		FIterator begin()
+		{
+			return InstanceStructs.begin();
+		}
+		FIterator end()
+		{
+			return InstanceStructs.end();
+		}
+
+		FConstIterator begin() const
+		{
+			return InstanceStructs.begin();
+		}
+		FConstIterator end() const
 		{
-			return InstanceStructs.GetAllocatedMemory();
+			return InstanceStructs.end();
 		}
 
 	private:
diff --git a/Engine/Plugins/Runtime/StateTree/Source/StateTreeModule/Public/StateTreeInstanceData.h b/Engine/Plugins/Runtime/StateTree/Source/StateTreeModule/Public/StateTreeInstanceData.h
--- a/Engine/Plugins/Runtime/StateTree/Source/StateTreeModule/Public/StateTreeInstanceData.h
+++ b/Engine/Plugins/Runtime/StateTree/Source/StateTreeModule/Public/StateTreeInstanceData.h
@@ -5,6 +5,7 @@
 #include "StructUtils/InstancedStructContainer.h"
 #include "Debugger/StateTreeRuntimeValidation.h"
 #include "StateTreeEvents.h"
+#include "StateTreeExecutionRuntimeDataTypes.h"
 #include "StateTreeInstanceContainer.h"
 #include "StateTreeTypes.h"
 #include "StateTreeExecutionTypes.h"
@@ -47,7 +48,7 @@ namespace UE::StateTree::InstanceData
 
 /**
  * Holds temporary instance data created during state selection.
- * The data is identified by Frame (StateTree + RootState) and DataHandle.
+ * The data is identified by Frame and DataHandle.
  */
 USTRUCT()
 struct FStateTreeTemporaryInstanceData
@@ -245,6 +246,24 @@ struct FStateTreeInstanceStorage
 		return ExecutionState;
 	}
 
+	/** @return the storage's execution runtime data for the state tree. */
+	[[nodiscard]] const UE::StateTree::InstanceData::FInstanceContainer& GetExecutionRuntimeData() const
+	{
+		return ExecutionRuntimeData;
+	}
+
+	/** @return the storage's execution runtime data for the state tree. */
+	[[nodiscard]] UE_API UE::StateTree::InstanceData::FInstanceContainer& GetExecutionRuntimeData()
+	{
+		return ExecutionRuntimeData;
+	}
+
+	/**
+	 * Add or reuse the execution runtime data for the state tree.
+	 * Return the base index for the execution runtime data.
+	 */
+	[[nodiscard]] UE_API int32 AddExecutionRuntimeData(TNotNull<UObject*> Owner, UE::StateTree::FExecutionFrameHandle FrameHandle);
+
 	/**
 	 * Adds temporary instance data associated with specified frame and data handle.
 	 * @returns mutable struct view to the instance.
@@ -282,7 +301,7 @@ struct FStateTreeInstanceStorage
 	}
 
 	/** @return a unique number used to make active frame id and active state id. */
-	UE_API uint32 GenerateUniqueId();
+	UE_API uint32 GenerateUniqueId(); //@TODO rename to GenerateUniqueID
 
 	/** Note, called by FStateTreeInstanceData. */
 	UE_API void AddStructReferencedObjects(FReferenceCollector& Collector);
@@ -306,7 +325,16 @@ struct FStateTreeInstanceStorage
 	UE_API UE::StateTree::Debug::FRuntimeValidation GetRuntimeValidation() const;
 
 protected:
-	/** Struct instances (Not transient, as we use FStateTreeInstanceData to store default values for instance data) */
+	/**
+	 * Struct for global and active instances.
+	 * The buffer format is:
+	 *  for each frames
+	 *    Global parameters, if it's a global frame.
+	 *    Global node instances, if it's a global frame. (evaluator, global tasks)
+	 *    Active state parameters
+	 *    Active node instances (tasks)
+	 * @note Not transient, as we use FStateTreeInstanceData to store default values for instance data.
+	 */
 	UPROPERTY()
 	FInstancedStructContainer InstanceStructs;
 
@@ -314,6 +342,21 @@ protected:
 	UPROPERTY(Transient)
 	FStateTreeExecutionState ExecutionState;
 
+	/**
+	 * Struct for the execution runtime data.
+	 * They stay alive until the owning execution context stops.
+	 */
+	UPROPERTY(Transient)
+	UE::StateTree::InstanceData::FInstanceContainer ExecutionRuntimeData;
+
+	/** Info to find the index where the execution runtime data starts for a specific state tree. */
+	struct FExecutionRuntimeInfo
+	{
+		FObjectKey StateTree;
+		int32 StartIndex = 0;
+	};
+	TArray<FExecutionRuntimeInfo, TInlineAllocator<1>> ExecutionRuntimeDataInfos;
+
 	/** Temporary instances */
 	UPROPERTY(Transient)
 	TArray<FStateTreeTemporaryInstanceData> TemporaryInstances;
diff --git a/Engine/Plugins/Runtime/StateTree/Source/StateTreeModule/Public/StateTreeNodeBase.h b/Engine/Plugins/Runtime/StateTree/Source/StateTreeModule/Public/StateTreeNodeBase.h
--- a/Engine/Plugins/Runtime/StateTree/Source/StateTreeModule/Public/StateTreeNodeBase.h
+++ b/Engine/Plugins/Runtime/StateTree/Source/StateTreeModule/Public/StateTreeNodeBase.h
@@ -72,10 +72,40 @@ PRAGMA_DISABLE_DEPRECATION_WARNINGS
 	FStateTreeNodeBase& operator=(FStateTreeNodeBase&&) = default;
 PRAGMA_ENABLE_DEPRECATION_WARNINGS
 
+	struct FNoInstanceDataType
+	{ };
+
+	/**
+	 * The instance data type. The implementation node should set its type.
+	 * The type should match GetInstanceDataType().
+	 */
+	using FInstanceDataType = FNoInstanceDataType;
+
+	/**
+	 * The execution runtime data type.
+	 * Execution runtime is optional. If needed, the implementation node should set its type.
+	 * The type should match GetExecutionRuntimeDataType().
+	 */
+	using FExecutionRuntimeDataType = FNoInstanceDataType;
+
 	virtual ~FStateTreeNodeBase() {}
 
 	/** @return Struct that represents the runtime data of the node. */
-	virtual const UStruct* GetInstanceDataType() const { return nullptr; };
+	virtual const UStruct* GetInstanceDataType() const
+	{
+		return nullptr;
+	}
+
+	/**
+	 * Execution runtime data is always valid between an FStateTreeExecutionContext::Start and FStateTreeExecutionContext::Stop.
+	 * If the node is no longer active but the instance (FStateTreeInstanceData) is active, the data is still valid and will persist.
+	 * @note Be careful about UObject reference, the objects are garbage collected like normal references.
+	 * @return the struct that represents the persistent runtime data of the node.
+	 */
+	virtual const UStruct* GetExecutionRuntimeDataType() const
+	{
+		return nullptr;
+	}
 
 	/**
 	 * Called when the StateTree asset is linked. Allows to resolve references to other StateTree data.
@@ -83,7 +113,10 @@ PRAGMA_ENABLE_DEPRECATION_WARNINGS
 	 * @param Linker Reference to the linker
 	 * @return true if linking succeeded. 
 	 */
-	[[nodiscard]] virtual bool Link(FStateTreeLinker& Linker) { return true; }
+	[[nodiscard]] virtual bool Link(FStateTreeLinker& Linker)
+	{
+		return true;
+	}
 
 #if WITH_EDITOR
 	/**
@@ -169,9 +202,9 @@ PRAGMA_ENABLE_DEPRECATION_WARNINGS
 #endif
 
 	/**
-	* Called after the state tree asset that contains this node is loaded from disk.
-	* @param InstanceDataView view to the instance data, can be struct or class.
-	*/
+	 * Called after the state tree asset that contains this node is loaded from disk.
+	 * @param InstanceDataView view to the instance data, can be struct or class.
+	 */
 	virtual void PostLoad(FStateTreeDataView InstanceDataView) {}
 
 	/** Name of the node. */
@@ -186,6 +219,10 @@ PRAGMA_ENABLE_DEPRECATION_WARNINGS
 	UPROPERTY()
 	FStateTreeIndex16 InstanceTemplateIndex = FStateTreeIndex16::Invalid;
 
+	/** Index of template execution runtime data for the node. */
+	UPROPERTY()
+	FStateTreeIndex16 ExecutionRuntimeTemplateIndex = FStateTreeIndex16::Invalid;
+
 	/** Data handle to access the instance data. */
 	UPROPERTY()
 	FStateTreeDataHandle InstanceDataHandle = FStateTreeDataHandle::Invalid;
diff --git a/Engine/Plugins/Runtime/StateTree/Source/StateTreeModule/Public/StateTreeTypes.h b/Engine/Plugins/Runtime/StateTree/Source/StateTreeModule/Public/StateTreeTypes.h
--- a/Engine/Plugins/Runtime/StateTree/Source/StateTreeModule/Public/StateTreeTypes.h
+++ b/Engine/Plugins/Runtime/StateTree/Source/StateTreeModule/Public/StateTreeTypes.h
@@ -291,10 +291,10 @@ enum class EStateTreeDataSourceType : uint8
 {
 	None UMETA(Hidden),
 
-	/** Global Tasks, Evaluators */
+	/** Global Tasks, Evaluators. */
 	GlobalInstanceData,
 
-	/** Global Tasks, Evaluators*/
+	/** Global Tasks, Evaluators. */
 	GlobalInstanceDataObject,
 
 	/** Active State Tasks */
@@ -315,6 +315,12 @@ enum class EStateTreeDataSourceType : uint8
 	/** Temporary data constructor, used and destroyed immediately. Used by conditions, considerations and function bindings.*/
 	EvaluationScopeInstanceDataObject,
 
+	/** Node instance data for the duration of the execution. */
+	ExecutionRuntimeData,
+
+	/** Node instance data for the duration of the execution. */
+	ExecutionRuntimeDataObject,
+
 	/** Context Data, Tree Parameters */
 	ContextData,
 
@@ -470,7 +476,7 @@ namespace UE::StateTree::InstanceData
 		int32 Alignment = 0;
 		int32 NumberOfElements = 0;
 	};
-};
+}
 
 /**
  * Time duration with random variance. Stored compactly as two uint16s, which gives time range of about 650 seconds.
diff --git a/Engine/Plugins/Runtime/StateTree/Source/StateTreeTestSuite/Private/StateTreeTestTypes.h b/Engine/Plugins/Runtime/StateTree/Source/StateTreeTestSuite/Private/StateTreeTestTypes.h
--- a/Engine/Plugins/Runtime/StateTree/Source/StateTreeTestSuite/Private/StateTreeTestTypes.h
+++ b/Engine/Plugins/Runtime/StateTree/Source/StateTreeTestSuite/Private/StateTreeTestTypes.h
@@ -318,18 +318,29 @@ struct FTestTask_PrintValueInstanceData
 	}
 };
 
+USTRUCT()
+struct FTestTask_PrintValueExecutionRuntimeData
+{
+	GENERATED_BODY()
+
+	UPROPERTY(EditAnywhere, Category = "Parameter")
+	int32 Value = 0;
+};
+
 USTRUCT()
 struct FTestTask_PrintValue : public FStateTreeTaskBase
 {
 	GENERATED_BODY()
 
 	using FInstanceDataType = FTestTask_PrintValueInstanceData;
+	using FExecutionRuntimeDataType = FTestTask_PrintValueExecutionRuntimeData;
 
 	FTestTask_PrintValue() = default;
 	FTestTask_PrintValue(const FName InName) { Name = InName; }
 	virtual ~FTestTask_PrintValue() override {}
 	
 	virtual const UStruct* GetInstanceDataType() const override { return FInstanceDataType::StaticStruct(); }
+	virtual const UStruct* GetExecutionRuntimeDataType() const override { return FExecutionRuntimeDataType::StaticStruct(); }
 
 	virtual EStateTreeRunStatus EnterState(FStateTreeExecutionContext& Context, const FStateTreeTransitionResult& Transition) const override
 	{
@@ -338,6 +349,9 @@ struct FTestTask_PrintValue : public FStateTreeTaskBase
 		TestLog.Log(Name, FString::Printf(TEXT("EnterState%d"), InstanceData.Value));
 		TestLog.Log(Name, FString::Printf(TEXT("EnterState:%s"), *InstanceData.GetArrayDescription()));
 		TestLog.Log(Name, FString::Printf(TEXT("EnterState=%s"), *StaticEnum<EStateTreeStateChangeType>()->GetNameStringByValue((int64)Transition.ChangeType)));
+
+		const FExecutionRuntimeDataType& ExecData = Context.GetExecutionRuntimeData(*this);
+		TestLog.Log(Name, FString::Printf(TEXT("EnterStateExec=%d"), ExecData.Value));
 		if (CustomEnterStateFunc)
 		{
 			CustomEnterStateFunc(Context, this);
