[StateTreeDebugger] fixed issue when stepping backward and reaching to the beginning of the recording

#rnx

diff --git a/Engine/Plugins/Runtime/StateTree/Source/StateTreeEditorModule/Private/Debugger/StateTreeRewindDebuggerTrack.cpp b/Engine/Plugins/Runtime/StateTree/Source/StateTreeEditorModule/Private/Debugger/StateTreeRewindDebuggerTrack.cpp
--- a/Engine/Plugins/Runtime/StateTree/Source/StateTreeEditorModule/Private/Debugger/StateTreeRewindDebuggerTrack.cpp
+++ b/Engine/Plugins/Runtime/StateTree/Source/StateTreeEditorModule/Private/Debugger/StateTreeRewindDebuggerTrack.cpp
@@ -183,7 +183,7 @@ bool FRewindDebuggerTrack::UpdateInternal()
 
 			// Events need to be updated for a new scrub time
 			FScrubState ScrubState(&EventCollection);
-			ScrubState.SetScrubTime(ScrubTime);
+			ScrubState.SetScrubTime(ScrubTime, /*ForceRefresh*/true);
 			EventsView->RequestRefresh(ScrubState);
 
 			// Tree view only need to be updated on new timeline data or new active states
