[State Tree] Nodes manipulation in Editor overhaul
- use FEditorDataClipboard to text serialize variable number and types of nodes and their corresponding bindings
- added copy/paste/append multiple nodes feature.
- fixed property function node value not retained in new nodes when copy/paste/duplicate.
- fixed property bindings not retained in new nodes if the old bindings were removed before pasting.
- fixed Detail Customization initlization issue when duplicating.
- now node manipulation can correctly update bindings of nodes under operation. If bindings are intra-nodes, the binding relationships between the new nodes will be correctly updated.

#DisabledValidations Preflight, foundationTestsBuild, foundationTestsRun

diff --git a/Engine/Plugins/Runtime/StateTree/Source/StateTreeEditorModule/Private/Customizations/StateTreeEditorNodeDetails.cpp b/Engine/Plugins/Runtime/StateTree/Source/StateTreeEditorModule/Private/Customizations/StateTreeEditorNodeDetails.cpp
--- a/Engine/Plugins/Runtime/StateTree/Source/StateTreeEditorModule/Private/Customizations/StateTreeEditorNodeDetails.cpp
+++ b/Engine/Plugins/Runtime/StateTree/Source/StateTreeEditorModule/Private/Customizations/StateTreeEditorNodeDetails.cpp
@@ -30,7 +30,6 @@
 #include "Blueprint/StateTreeConditionBlueprintBase.h"
 #include "Blueprint/StateTreeConsiderationBlueprintBase.h"
 #include "Styling/StyleColors.h"
-#include "HAL/PlatformApplicationMisc.h"
 #include "ScopedTransaction.h"
 #include "StateTreeEditorNodeUtils.h"
 #include "Debugger/StateTreeDebuggerUIExtensions.h"
@@ -40,9 +39,11 @@
 #include "Styling/SlateTypes.h"
 #include "TextStyleDecorator.h"
 #include "Framework/Application/SlateApplication.h"
+#include "Framework/Notifications/NotificationManager.h"
 #include "SourceCodeNavigation.h"
 #include "StateTreeDelegate.h"
-#include "StateTreeScopedEditorDataFixer.h"
+#include "StateTreeEditorDataClipboardHelpers.h"
+#include "Widgets/Notifications/SNotificationList.h"
 
 #define LOCTEXT_NAMESPACE "StateTreeEditor"
 
@@ -356,6 +357,9 @@ FStateTreeEditorNodeDetails::~FStateTreeEditorNodeDetails()
 void FStateTreeEditorNodeDetails::CustomizeHeader(TSharedRef<class IPropertyHandle> StructPropertyHandle, class FDetailWidgetRow& HeaderRow, IPropertyTypeCustomizationUtils& StructCustomizationUtils)
 {
 	StructProperty = StructPropertyHandle;
+	ParentProperty = StructProperty->GetParentHandle();
+	ParentArrayProperty = ParentProperty->AsArray();
+
 	PropUtils = StructCustomizationUtils.GetPropertyUtilities();
 
 	NodeProperty = StructProperty->GetChildHandle(GET_MEMBER_NAME_CHECKED(FStateTreeEditorNode, Node));
@@ -700,7 +704,7 @@ void FStateTreeEditorNodeDetails::CustomizeHeader(TSharedRef<class IPropertyHand
 		]
 		.OverrideResetToDefault(ResetOverride)
 		.CopyAction(FUIAction(FExecuteAction::CreateSP(this, &FStateTreeEditorNodeDetails::OnCopyNode)))
-		.PasteAction(FUIAction(FExecuteAction::CreateSP(this, &FStateTreeEditorNodeDetails::OnPasteNode)));
+		.PasteAction(FUIAction(FExecuteAction::CreateSP(this, &FStateTreeEditorNodeDetails::OnPasteNodes)));
 
 	// Task completion
 	bool bShowCompletion = true;
@@ -855,16 +859,49 @@ FReply FStateTreeEditorNodeDetails::OnRowMouseUp(const FGeometry& MyGeometry, co
 
 void FStateTreeEditorNodeDetails::OnCopyNode()
 {
-	FString Value;
-	// Use PPF_Copy so that all properties get copied.
-	if (StructProperty->GetValueAsFormattedString(Value, PPF_Copy) == FPropertyAccess::Success)
+	UStateTreeEditorData* EditorDataPtr = EditorData.Get();
+	if (!EditorDataPtr)
+	{
+		return;
+	}
+
+	void* EditorNodeAddress = nullptr;
+	if (StructProperty->GetValueData(EditorNodeAddress) == FPropertyAccess::Success)
+	{
+		UE::StateTreeEditor::FClipboardEditorData Clipboard;
+		Clipboard.Append(EditorDataPtr, TConstArrayView<FStateTreeEditorNode>(static_cast<FStateTreeEditorNode*>(EditorNodeAddress), 1));
+		UE::StateTreeEditor::ExportTextAsClipboardEditorData(Clipboard);
+	}
+}
+
+void FStateTreeEditorNodeDetails::OnCopyAllNodes()
+{
+	UStateTreeEditorData* EditorDataPtr = EditorData.Get();
+	if (!EditorDataPtr)
+	{
+		return;
+	}
+
+	if (ParentArrayProperty)
+	{
+		void* EditorNodeArrayAddress = nullptr;
+		if (ParentProperty->GetValueData(EditorNodeArrayAddress) == FPropertyAccess::Success)
+		{
+			UE::StateTreeEditor::FClipboardEditorData Clipboard;
+			Clipboard.Append(EditorDataPtr, *static_cast<TArray<FStateTreeEditorNode>*>(EditorNodeArrayAddress));
+			UE::StateTreeEditor::ExportTextAsClipboardEditorData(Clipboard);
+		}
+	}
+	else
 	{
-		FPlatformApplicationMisc::ClipboardCopy(*Value);
+		OnCopyNode();
 	}
 }
 
-void FStateTreeEditorNodeDetails::OnPasteNode()
+void FStateTreeEditorNodeDetails::OnPasteNodes()
 {
+	using namespace UE::StateTreeEditor;
+
 	TArray<UObject*> OuterObjects;
 	StructProperty->GetOuterObjects(OuterObjects);
 	if (OuterObjects.IsEmpty())
@@ -872,47 +909,150 @@ void FStateTreeEditorNodeDetails::OnPasteNode()
 		return;
 	}
 
-	FStateTreeEditorNode TempNode;
-	bool bSuccess = UE::StateTreeEditor::EditorNodeUtils::ImportTextAsNode(BaseScriptStruct.Get(), EditorData.Get(), TempNode);
+	UStateTreeEditorData* EditorDataPtr = EditorData.Get();
+	if (!EditorDataPtr)
+	{
+		return;
+	}
+
+	// In case its multi selected, we need to have a unique copy for each Object
+	TArray<FClipboardEditorData, TInlineAllocator<2>> ClipboardEditorDatas;
+	ClipboardEditorDatas.AddDefaulted(OuterObjects.Num());
 
-	if (!bSuccess)
+	for (int32 Idx = 0; Idx < OuterObjects.Num(); ++Idx)
+	{
+		const bool bSuccess = ImportTextAsClipboardEditorData
+		(BaseScriptStruct.Get(),
+			EditorData.Get(),
+			OuterObjects[Idx],
+			ClipboardEditorDatas[Idx]);
+
+		if (!bSuccess)
+		{
+			return;
+		}
+	}
+
+	// make sure each Clipboard has the same number of nodes
+	for (int32 Idx = 0; Idx < ClipboardEditorDatas.Num() - 1; ++Idx)
+	{
+		check(ClipboardEditorDatas[Idx].GetEditorNodesInBuffer().Num() == ClipboardEditorDatas[Idx + 1].GetEditorNodesInBuffer().Num());
+	}
+
+	int32 NumEditorNodesInBuffer = ClipboardEditorDatas[0].GetEditorNodesInBuffer().Num();
+
+	if (NumEditorNodesInBuffer == 0)
 	{
 		return;
 	}
 
-	TArray<void*> RawNodeData;
-	StructProperty->AccessRawData(RawNodeData);
-	if (OuterObjects.Num() == RawNodeData.Num())
+	if (!ParentArrayProperty.IsValid() && NumEditorNodesInBuffer != 1)
 	{
-		if (UStateTreeEditorData* EditorDataPtr = EditorData.Get())
+		// Node is not in an array. we can't do multi-to-one paste
+		return;
+	}
+
+	if (OuterObjects.Num() != 1 && NumEditorNodesInBuffer != 1)
+	{
+		// if multiple selected objects, and we have more than one nodes to paste into
+		// Array Handle doesn't support manipulation on multiple objects.
+		FNotificationInfo NotificationInfo(FText::GetEmpty());
+		NotificationInfo.Text = LOCTEXT("NotSupportedByMultipleObjects", "Operation is not supported for multi-selected objects");
+		NotificationInfo.ExpireDuration = 5.0f;
+		FSlateNotificationManager::Get().AddNotification(NotificationInfo);
+
+		return;
+	}
+
+	{
+		FScopedTransaction Transaction(LOCTEXT("PasteNode", "Paste Node"));
+
+		EditorDataPtr->Modify();	// we might modify the bindings on Editor Data
+		StructProperty->NotifyPreChange();
+
+		if (ParentArrayProperty)
 		{
-			FScopedTransaction Transaction(LOCTEXT("PasteNode", "Paste Node"));
+			// Paste multi nodes into one node
+			const int32 StructIndex = StructProperty->GetIndexInArray();
+			check(StructIndex != INDEX_NONE);
 
-			StructProperty->NotifyPreChange();
+			uint32 NumArrayElements = 0;
+			ParentArrayProperty->GetNumElements(NumArrayElements);
+			check(NumArrayElements > 0);	// since we already have at least one element to paste into
 
-			// we will likely change the bindings on the Editor Data or the editor node could be global, so record
-			EditorDataPtr->Modify();
+			// Insert or append uninitialized elements after the current node to match the number of nodes in the paste buffer and retain the order of elements 
+			// The first node in the buffer goes into the current node
+			const int32 IndexToInsert = StructIndex + 1;
+			const int32 NumElementsToAddOrInsert = ClipboardEditorDatas[0].GetEditorNodesInBuffer().Num() - 1;
 
+			int32 Cnt = 0;
+			if (IndexToInsert == NumArrayElements)
+			{
+				while (Cnt++ < NumElementsToAddOrInsert)
+				{
+					FPropertyHandleItemAddResult Result = ParentArrayProperty->AddItem();
+					if (Result.GetAccessResult() != FPropertyAccess::Success)
+					{
+						return;
+					}
+				}
+			}
+			else
 			{
-				UE::StateTreeEditor::FScopedEditorDataFixer DataFixer(EditorDataPtr);
+				while (Cnt++ < NumElementsToAddOrInsert)
+				{
+					FPropertyAccess::Result Result = ParentArrayProperty->Insert(IndexToInsert);
+					if (Result != FPropertyAccess::Success)
+					{
+						return;
+					}
+				}
+			}
 
-				for (int32 Index = 0; Index < RawNodeData.Num(); Index++)
+			TArray<void*> RawDatasArray;
+			ParentProperty->AccessRawData(RawDatasArray);
+			check(RawDatasArray.Num() == OuterObjects.Num());
+			for (int32 ObjIdx = 0; ObjIdx < OuterObjects.Num(); ++ObjIdx)
+			{
+				if (TArray<FStateTreeEditorNode>* EditorNodesPtr = static_cast<TArray<FStateTreeEditorNode>*>(RawDatasArray[ObjIdx]))
 				{
-					if (FStateTreeEditorNode* EditorNode = static_cast<FStateTreeEditorNode*>(RawNodeData[Index]))
+					TArrayView<FStateTreeEditorNode> EditorNodesClipboardBuffer = ClipboardEditorDatas[ObjIdx].GetEditorNodesInBuffer();
+					TArray<FStateTreeEditorNode>& EditorNodesToPasteInto = *EditorNodesPtr;
+
+					for (int32 Idx = 0; Idx < EditorNodesClipboardBuffer.Num(); ++Idx)
 					{
-						*EditorNode = TempNode;
+						EditorNodesToPasteInto[StructIndex + Idx] = MoveTemp(EditorNodesClipboardBuffer[Idx]);
+					}
 
-						constexpr bool bShouldCopyBindings = true;
-						constexpr bool bShouldReinstantiateObjects = true;
-						constexpr bool bShouldRegenerateGUID = true;
-						DataFixer.EditorNodesToFix.Emplace(OuterObjects[Index], *EditorNode, bShouldCopyBindings, bShouldReinstantiateObjects, bShouldRegenerateGUID);
+					for (FStateTreePropertyPathBinding& Binding : ClipboardEditorDatas[ObjIdx].GetBindingsInBuffer())
+					{
+						EditorDataPtr->GetPropertyEditorBindings()->AddStateTreeBinding(MoveTemp(Binding));
 					}
 				}
 			}
+		}
+		else
+		{
+			// Paste single node to a single Node
+			TArray<void*> RawDatas;
+			StructProperty->AccessRawData(RawDatas);
+			check(RawDatas.Num() == OuterObjects.Num());
+			for (int32 Idx = 0; Idx < RawDatas.Num(); ++Idx)
+			{
+				if (FStateTreeEditorNode* CurrentEditorNode = static_cast<FStateTreeEditorNode*>(RawDatas[Idx]))
+				{
+					*CurrentEditorNode = MoveTemp(ClipboardEditorDatas[Idx].GetEditorNodesInBuffer()[0]);
 
-			StructProperty->NotifyPostChange(EPropertyChangeType::ValueSet);
-			StructProperty->NotifyFinishedChangingProperties();
+					for (FStateTreePropertyPathBinding& Binding : ClipboardEditorDatas[Idx].GetBindingsInBuffer())
+					{
+						EditorDataPtr->GetPropertyEditorBindings()->AddStateTreeBinding(MoveTemp(Binding));
+					}
+				}
+			}
 		}
+
+		StructProperty->NotifyPostChange(EPropertyChangeType::ValueSet);
+		StructProperty->NotifyFinishedChangingProperties();
 	}
 
 	if (PropUtils)
@@ -1973,13 +2113,23 @@ TSharedRef<SWidget> FStateTreeEditorNodeDetails::GenerateOptionsMenu()
 			FCanExecuteAction()
 		));
 
+	// Copy all
+	MenuBuilder.AddMenuEntry(
+		LOCTEXT("CopyAllItems", "Copy all"),
+		LOCTEXT("CopyItemTooltip", "Copy all items"),
+		FSlateIcon(FAppStyle::GetAppStyleSetName(), "GenericCommands.Copy"),
+		FUIAction(
+			FExecuteAction::CreateSP(this, &FStateTreeEditorNodeDetails::OnCopyAllNodes),
+			FCanExecuteAction()
+		));
+
 	// Paste
 	MenuBuilder.AddMenuEntry(
 		LOCTEXT("PasteItem", "Paste"),
 		LOCTEXT("PasteItemTooltip", "Paste into this item"),
 		FSlateIcon(FAppStyle::GetAppStyleSetName(), "GenericCommands.Paste"),
 		FUIAction(
-			FExecuteAction::CreateSP(this, &FStateTreeEditorNodeDetails::OnPasteNode),
+			FExecuteAction::CreateSP(this, &FStateTreeEditorNodeDetails::OnPasteNodes),
 			FCanExecuteAction()
 		));
 
@@ -2034,21 +2184,17 @@ TSharedRef<SWidget> FStateTreeEditorNodeDetails::GenerateOptionsMenu()
 void FStateTreeEditorNodeDetails::OnDeleteNode() const
 {
 	const int32 Index = StructProperty->GetArrayIndex();
-	if (const TSharedPtr<IPropertyHandle> ParentHandle = StructProperty->GetParentHandle())
+	if (ParentArrayProperty)
 	{
-		if (const TSharedPtr<IPropertyHandleArray> ArrayHandle = ParentHandle->AsArray())
+		if (UStateTreeEditorData* EditorDataPtr = EditorData.Get())
 		{
-			if (UStateTreeEditorData* EditorDataPtr = EditorData.Get())
-			{
-				FScopedTransaction Transaction(LOCTEXT("DeleteNode", "Delete Node"));
+			FScopedTransaction Transaction(LOCTEXT("DeleteNode", "Delete Node"));
 
-				EditorDataPtr->Modify();
+			EditorDataPtr->Modify();
 
-				ArrayHandle->DeleteItem(Index);
+			ParentArrayProperty->DeleteItem(Index);
 
-				UE::StateTreeEditor::FScopedEditorDataFixer DataFixer(EditorDataPtr);
-				DataFixer.bRemoveInvalidBindings = true;
-			}
+			UE::StateTreeEditor::RemoveInvalidBindings(EditorDataPtr);
 		}
 	}
 }
@@ -2067,8 +2213,7 @@ void FStateTreeEditorNodeDetails::OnDeleteAllNodes() const
 
 				ArrayHandle->EmptyArray();
 
-				UE::StateTreeEditor::FScopedEditorDataFixer DataFixer(EditorDataPtr);
-				DataFixer.bRemoveInvalidBindings = true;
+				UE::StateTreeEditor::RemoveInvalidBindings(EditorDataPtr);
 			}
 		}
 	}
@@ -2083,39 +2228,64 @@ void FStateTreeEditorNodeDetails::OnDuplicateNode() const
 		return;
 	}
 
-	const int32 ArrayIndex = StructProperty->GetArrayIndex();
-	if (const TSharedPtr<IPropertyHandle> ParentHandle = StructProperty->GetParentHandle())
+	if (OuterObjects.Num() != 1)
 	{
-		if (const TSharedPtr<IPropertyHandleArray> ArrayHandle = ParentHandle->AsArray())
+		// Array Handle Manipulation doesn't support multiple selected objects
+		FNotificationInfo NotificationInfo(FText::GetEmpty());
+		NotificationInfo.Text = LOCTEXT("NotSupportedByMultipleObjects", "Operation is not supported for multi-selected objects");
+		NotificationInfo.ExpireDuration = 5.0f;
+		FSlateNotificationManager::Get().AddNotification(NotificationInfo);
+
+		return;
+	}
+
+	if (ParentArrayProperty)
+	{
+		if (UStateTreeEditorData* EditorDataPtr = EditorData.Get())
 		{
-			if (UStateTreeEditorData* EditorDataPtr = EditorData.Get())
+			void* NodePtr = nullptr;
+			if (StructProperty->GetValueData(NodePtr) == FPropertyAccess::Success)
 			{
+				UE::StateTreeEditor::FClipboardEditorData Clipboard;
+				Clipboard.Append(EditorDataPtr, TConstArrayView<FStateTreeEditorNode>(static_cast<FStateTreeEditorNode*>(NodePtr), 1));
+				//UE::StateTreeEditor::ExportTextAsClipboardEditorData()
+				Clipboard.ProcessBuffer(nullptr, EditorDataPtr, OuterObjects[0]);
+
+				if (!Clipboard.IsValid())
+				{
+					return;
+				}
+
 				FScopedTransaction Transaction(LOCTEXT("DuplicateNode", "Duplicate Node"));
 
 				// Might modify the bindings data
 				EditorDataPtr->Modify();
 
-				ArrayHandle->DuplicateItem(ArrayIndex);
-				
-				TSharedRef<IPropertyHandle> DuplicatedStructHandle = ArrayHandle->GetElement(ArrayIndex);
+				const int32 Index = StructProperty->GetArrayIndex();
+				ParentArrayProperty->Insert(Index);
 
-				UE::StateTreeEditor::FScopedEditorDataFixer DataFixer(EditorDataPtr);
-
-				TArray<void*> RawNodeData;
-				StructProperty->AccessRawData(RawNodeData);
-				for (int32 Index = 0; Index < RawNodeData.Num(); ++Index)
+				TSharedPtr<IPropertyHandle> InsertedElementHandle = ParentArrayProperty->GetElement(Index);
+				void* InsertedNodePtr = nullptr;
+				if (InsertedElementHandle->GetValueData(InsertedNodePtr) == FPropertyAccess::Success)
 				{
-					if (FStateTreeEditorNode* EditorNode = static_cast<FStateTreeEditorNode*>(RawNodeData[Index]))
+					*static_cast<FStateTreeEditorNode*>(InsertedNodePtr) = MoveTemp(Clipboard.GetEditorNodesInBuffer()[0]);
+
+					for (FStateTreePropertyPathBinding& Binding : Clipboard.GetBindingsInBuffer())
 					{
-						constexpr bool bShouldCopyBindings = true;
-						constexpr bool bShouldRegenerateGUID = true;
-						constexpr bool bShouldRenstantiateSubObjects = false; // ArrayHandle Duplication has done deep copies for us
-						DataFixer.EditorNodesToFix.Emplace(OuterObjects[Index], *EditorNode, bShouldCopyBindings, bShouldRenstantiateSubObjects, bShouldRegenerateGUID);
+						EditorDataPtr->GetPropertyEditorBindings()->AddStateTreeBinding(MoveTemp(Binding));
 					}
 				}
+
+				// We reinieitalized item nodes on ArrayProperty operations before the data is completely set up. Reinitialize.
+				if (PropUtils)
+				{
+					PropUtils->ForceRefresh();
+				}
 			}
 		}
 	}
+
+	
 }
 
 void FStateTreeEditorNodeDetails::OnRenameNode() const
diff --git a/Engine/Plugins/Runtime/StateTree/Source/StateTreeEditorModule/Private/Customizations/StateTreeEditorNodeDetails.h b/Engine/Plugins/Runtime/StateTree/Source/StateTreeEditorModule/Private/Customizations/StateTreeEditorNodeDetails.h
--- a/Engine/Plugins/Runtime/StateTree/Source/StateTreeEditorModule/Private/Customizations/StateTreeEditorNodeDetails.h
+++ b/Engine/Plugins/Runtime/StateTree/Source/StateTreeEditorModule/Private/Customizations/StateTreeEditorNodeDetails.h
@@ -3,6 +3,7 @@
 #pragma once
 
 #include "IPropertyTypeCustomization.h"
+#include "IStateTreeEditorHost.h"
 #include "Widgets/Views/STreeView.h"
 #include "Textures/SlateIcon.h"
 
@@ -13,6 +14,7 @@ class FDetailWidgetRow;
 class FStateTreeViewModel;
 class IDetailChildrenBuilder;
 class IPropertyHandle;
+class IPropertyHandleArray;
 class IPropertyUtilities;
 class SComboButton;
 class SWidget;
@@ -50,7 +52,8 @@ private:
 	bool ShouldResetToDefault(TSharedPtr<IPropertyHandle> PropertyHandle) const;
 	void ResetToDefault(TSharedPtr<IPropertyHandle> PropertyHandle);
 	void OnCopyNode();
-	void OnPasteNode();
+	void OnCopyAllNodes();
+	void OnPasteNodes();
 	TSharedPtr<IPropertyHandle> GetInstancedObjectValueHandle(TSharedPtr<IPropertyHandle> PropertyHandle);
 
 	FOptionalSize GetIndentSize() const;
@@ -153,6 +156,9 @@ private:
 	TSharedPtr<IPropertyHandle> IndentProperty;
 	TSharedPtr<IPropertyHandle> OperandProperty;
 
+	TSharedPtr<IPropertyHandle> ParentProperty;
+	TSharedPtr<IPropertyHandleArray> ParentArrayProperty;
+
 	FDelegateHandle OnBindingChangedHandle;
 	FDelegateHandle OnChangedAssetHandle;
 };
diff --git a/Engine/Plugins/Runtime/StateTree/Source/StateTreeEditorModule/Private/Customizations/StateTreeEditorNodeUtils.cpp b/Engine/Plugins/Runtime/StateTree/Source/StateTreeEditorModule/Private/Customizations/StateTreeEditorNodeUtils.cpp
--- a/Engine/Plugins/Runtime/StateTree/Source/StateTreeEditorModule/Private/Customizations/StateTreeEditorNodeUtils.cpp
+++ b/Engine/Plugins/Runtime/StateTree/Source/StateTreeEditorModule/Private/Customizations/StateTreeEditorNodeUtils.cpp
@@ -13,7 +13,6 @@
 #include "StateTreeEditorData.h"
 #include "StateTreeEditorNode.h"
 #include "StateTreeEditorSettings.h"
-#include "HAL/PlatformApplicationMisc.h"
 #include "StateTreeEditorStyle.h"
 #include "StateTreeTaskBase.h"
 #include "Blueprint/StateTreeConditionBlueprintBase.h"
@@ -24,29 +23,12 @@
 #include "Widgets/Input/SComboButton.h"
 #include "Widgets/Input/SButton.h"
 #include "Editor.h"
-#include "Framework/Notifications/NotificationManager.h"
-#include "Widgets/Notifications/SNotificationList.h"
 
 #define LOCTEXT_NAMESPACE "StateTreeEditor"
 
 namespace UE::StateTreeEditor::EditorNodeUtils
 {
 
-/** Helper class to detect if there were issues when calling ImportText() */
-class FDefaultValueImportErrorContext : public FOutputDevice
-{
-public:
-
-	int32 NumErrors = 0;
-
-	FDefaultValueImportErrorContext() = default;
-
-	virtual void Serialize(const TCHAR* V, ELogVerbosity::Type Verbosity, const class FName& Category) override
-	{
-		++NumErrors;
-	}
-};
-
 EStateTreeConditionEvaluationMode GetConditionEvaluationMode(const TSharedPtr<IPropertyHandle>& StructProperty)
 {
 	if (const FStateTreeEditorNode* Node = GetCommonNode(StructProperty))
@@ -1002,72 +984,6 @@ void MakeArrayItems(IDetailCategoryBuilder& Category, const TSharedPtr<IProperty
 	}));
 	Category.AddCustomBuilder(Builder, /*bForAdvanced*/ false);
 }
-
-bool ImportTextAsNode(const UScriptStruct* BaseScriptStruct, const UStateTreeEditorData* EditorData, FStateTreeEditorNode& OutEditorNode)
-{
-	FString PastedText;
-	FPlatformApplicationMisc::ClipboardPaste(PastedText);
-
-	if (PastedText.IsEmpty())
-	{
-		return false;
-	}
-
-	UScriptStruct* NodeScriptStruct = TBaseStructure<FStateTreeEditorNode>::Get();
-	FDefaultValueImportErrorContext ErrorPipe;
-	NodeScriptStruct->ImportText(*PastedText, &OutEditorNode, nullptr, PPF_None, &ErrorPipe, NodeScriptStruct->GetName());
-
-	const UStruct* NodeTypeStruct = OutEditorNode.Node.GetScriptStruct();
-	// Only allow valid node types for this property (e.g. do not mix task with conditions).
-	if (ErrorPipe.NumErrors > 0 || !NodeTypeStruct || !NodeTypeStruct->IsChildOf(BaseScriptStruct))
-	{
-		FNotificationInfo NotificationInfo(FText::GetEmpty());
-		NotificationInfo.Text = FText::Format(LOCTEXT("NotSupportedByType", "This property only accepts nodes of type {0}."), BaseScriptStruct->GetDisplayNameText());
-		NotificationInfo.ExpireDuration = 5.0f;
-		FSlateNotificationManager::Get().AddNotification(NotificationInfo);
-		return false;
-	}
-
-	const UStateTreeSchema* Schema = EditorData ? EditorData->Schema.Get() : nullptr;
-	if (Schema)
-	{
-		bool bNodeIsAllowed = false;
-
-		// BP nodes are identified by the instance type.
-		if (NodeTypeStruct->IsChildOf(FStateTreeBlueprintEvaluatorWrapper::StaticStruct())
-			|| NodeTypeStruct->IsChildOf(FStateTreeBlueprintTaskWrapper::StaticStruct())
-			|| NodeTypeStruct->IsChildOf(FStateTreeBlueprintConditionWrapper::StaticStruct())
-			|| NodeTypeStruct->IsChildOf(FStateTreeBlueprintConsiderationWrapper::StaticStruct()))
-		{
-			if (const FStateTreeNodeBase* Node = OutEditorNode.Node.GetPtr<FStateTreeNodeBase>())
-			{
-				NodeTypeStruct = Node->GetInstanceDataType(); // Report error with the BP node type, as that is what the user expects to see.
-				if (const UClass* InstanceClass = Cast<UClass>(NodeTypeStruct))
-				{
-					bNodeIsAllowed = Schema->IsClassAllowed(InstanceClass);
-				}
-			}
-		}
-		else
-		{
-			bNodeIsAllowed = Schema->IsStructAllowed(OutEditorNode.Node.GetScriptStruct());
-		}
-
-		if (!bNodeIsAllowed)
-		{
-
-			FNotificationInfo NotificationInfo(FText::GetEmpty());
-			NotificationInfo.Text = FText::Format(LOCTEXT("NotSupportedBySchema", "Node {0} is not supported by {1} schema."),
-												  NodeTypeStruct->GetDisplayNameText(), Schema->GetClass()->GetDisplayNameText());
-			NotificationInfo.ExpireDuration = 5.0f;
-			FSlateNotificationManager::Get().AddNotification(NotificationInfo);
-			return false;
-		}
-	}
-
-	// No schema, any types would be accepted
-	return true;
-}
 } // namespace UE::StateTreeEditor::EditorNodeUtils
 
 #undef LOCTEXT_NAMESPACE
diff --git a/Engine/Plugins/Runtime/StateTree/Source/StateTreeEditorModule/Private/Customizations/StateTreeEditorNodeUtils.h b/Engine/Plugins/Runtime/StateTree/Source/StateTreeEditorModule/Private/Customizations/StateTreeEditorNodeUtils.h
--- a/Engine/Plugins/Runtime/StateTree/Source/StateTreeEditorModule/Private/Customizations/StateTreeEditorNodeUtils.h
+++ b/Engine/Plugins/Runtime/StateTree/Source/StateTreeEditorModule/Private/Customizations/StateTreeEditorNodeUtils.h
@@ -207,13 +207,4 @@ IDetailCategoryBuilder& MakeArrayCategoryHeader(
  * @param ArrayPropertyHandle Property handle pointing to a TArray<FStateTreeEditorNode>.
  */
 void MakeArrayItems(IDetailCategoryBuilder& Category, const TSharedPtr<IPropertyHandle>& ArrayPropertyHandle);
-
-/**
- * Import the text from clipboard as an Editor Node
- * @param BaseScriptStruct the Base Struct of the Editor Node's Node
- * @param EditorData Editor Data of the Node
- * @param OutEditorNode Editor Node reference to which the text will be imported
- * @return Import succeeded or not
- */
-bool ImportTextAsNode(const UScriptStruct* BaseScriptStruct, const UStateTreeEditorData* EditorData, FStateTreeEditorNode& OutEditorNode);
-}; // UE::StateTreeEditor::EditorNodeUtils
+};
\ No newline at end of file
diff --git a/Engine/Plugins/Runtime/StateTree/Source/StateTreeEditorModule/Private/Customizations/StateTreeTransitionDetails.cpp b/Engine/Plugins/Runtime/StateTree/Source/StateTreeEditorModule/Private/Customizations/StateTreeTransitionDetails.cpp
--- a/Engine/Plugins/Runtime/StateTree/Source/StateTreeEditorModule/Private/Customizations/StateTreeTransitionDetails.cpp
+++ b/Engine/Plugins/Runtime/StateTree/Source/StateTreeEditorModule/Private/Customizations/StateTreeTransitionDetails.cpp
@@ -5,11 +5,10 @@
 #include "DetailLayoutBuilder.h"
 #include "DetailWidgetRow.h"
 #include "Framework/Application/SlateApplication.h"
-#include "HAL/PlatformApplicationMisc.h"
+#include "Framework/Notifications/NotificationManager.h"
 #include "IDetailChildrenBuilder.h"
 #include "IPropertyUtilities.h"
 #include "ScopedTransaction.h"
-#include "StateTreeBindingExtension.h"
 #include "StateTreeDescriptionHelpers.h"
 #include "StateTreeEditingSubsystem.h"
 #include "StateTreeEditor.h"
@@ -17,13 +16,14 @@
 #include "StateTreeEditorNodeUtils.h"
 #include "StateTreeEditorStyle.h"
 #include "StateTreePropertyHelpers.h"
-#include "StateTreeScopedEditorDataFixer.h"
+#include "StateTreeEditorDataClipboardHelpers.h"
 #include "StateTreeTypes.h"
 #include "TextStyleDecorator.h"
 #include "Widgets/Images/SImage.h"
 #include "Widgets/Input/SButton.h"
 #include "Widgets/Input/SComboButton.h"
 #include "Widgets/Layout/SBox.h"
+#include "Widgets/Notifications/SNotificationList.h"
 #include "Widgets/SBoxPanel.h"
 #include "Widgets/Text/SRichTextBlock.h"
 #include "Widgets/Text/STextBlock.h"
@@ -38,6 +38,9 @@ TSharedRef<IPropertyTypeCustomization> FStateTreeTransitionDetails::MakeInstance
 void FStateTreeTransitionDetails::CustomizeHeader(TSharedRef<IPropertyHandle> StructPropertyHandle, FDetailWidgetRow& HeaderRow, IPropertyTypeCustomizationUtils& StructCustomizationUtils)
 {
 	StructProperty = StructPropertyHandle;
+	ParentProperty = StructProperty->GetParentHandle();
+	ParentArrayProperty = ParentProperty->AsArray();
+
 	PropUtils = StructCustomizationUtils.GetPropertyUtilities();
 
 	// Find StateTreeEditorData associated with this panel.
@@ -151,7 +154,7 @@ void FStateTreeTransitionDetails::CustomizeHeader(TSharedRef<IPropertyHandle> St
 			]
 		]
 		.CopyAction(FUIAction(FExecuteAction::CreateSP(this, &FStateTreeTransitionDetails::OnCopyTransition)))
-		.PasteAction(FUIAction(FExecuteAction::CreateSP(this, &FStateTreeTransitionDetails::OnPasteTransition)));
+		.PasteAction(FUIAction(FExecuteAction::CreateSP(this, &FStateTreeTransitionDetails::OnPasteTransitions)));
 }
 
 void FStateTreeTransitionDetails::CustomizeChildren(TSharedRef<class IPropertyHandle> StructPropertyHandle, class IDetailChildrenBuilder& StructBuilder, IPropertyTypeCustomizationUtils& StructCustomizationUtils)
@@ -314,13 +317,23 @@ TSharedRef<SWidget> FStateTreeTransitionDetails::GenerateOptionsMenu()
 			FCanExecuteAction()
 		));
 
+	// Copy all
+	MenuBuilder.AddMenuEntry(
+		LOCTEXT("CopyAllItems", "Copy all"),
+		LOCTEXT("CopyAllItemsTooltip", "Copy all items"),
+		FSlateIcon(FAppStyle::GetAppStyleSetName(), "GenericCommands.Copy"),
+		FUIAction(
+			FExecuteAction::CreateSP(this, &FStateTreeTransitionDetails::OnCopyAllTransitions),
+			FCanExecuteAction()
+		));
+
 	// Paste
 	MenuBuilder.AddMenuEntry(
 		LOCTEXT("PasteItem", "Paste"),
 		LOCTEXT("PasteItemTooltip", "Paste into this item"),
 		FSlateIcon(FAppStyle::GetAppStyleSetName(), "GenericCommands.Paste"),
 		FUIAction(
-			FExecuteAction::CreateSP(this, &FStateTreeTransitionDetails::OnPasteTransition),
+			FExecuteAction::CreateSP(this, &FStateTreeTransitionDetails::OnPasteTransitions),
 			FCanExecuteAction()
 		));
 
@@ -344,7 +357,7 @@ TSharedRef<SWidget> FStateTreeTransitionDetails::GenerateOptionsMenu()
 			FCanExecuteAction()
 		));
 
-	// Delete
+	// Delete all
 	MenuBuilder.AddMenuEntry(
 		LOCTEXT("DeleteAllItems", "Delete all"),
 		LOCTEXT("DeleteAllItemsTooltip", "Delete all items"),
@@ -376,8 +389,7 @@ void FStateTreeTransitionDetails::OnDeleteTransition() const
 
 				ArrayHandle->DeleteItem(Index);
 
-				UE::StateTreeEditor::FScopedEditorDataFixer DataFixer(EditorData);
-				DataFixer.bRemoveInvalidBindings = true;
+				UE::StateTreeEditor::RemoveInvalidBindings(EditorData);
 			}
 		}
 	}
@@ -396,8 +408,7 @@ void FStateTreeTransitionDetails::OnDeleteAllTransitions() const
 
 				ArrayHandle->EmptyArray();
 
-				UE::StateTreeEditor::FScopedEditorDataFixer DataFixer(EditorData);
-				DataFixer.bRemoveInvalidBindings = true;
+				UE::StateTreeEditor::RemoveInvalidBindings(EditorData);
 			}
 		}
 	}
@@ -412,37 +423,58 @@ void FStateTreeTransitionDetails::OnDuplicateTransition() const
 		return;
 	}
 
-	const int32 ArrayIndex = StructProperty->GetArrayIndex();
-	if (const TSharedPtr<IPropertyHandle> ParentHandle = StructProperty->GetParentHandle())
+	if (OuterObjects.Num() != 1)
 	{
-		if (const TSharedPtr<IPropertyHandleArray> ArrayHandle = ParentHandle->AsArray())
+		// Array Handle Manipulation doesn't support multiple selected objects
+		FNotificationInfo NotificationInfo(FText::GetEmpty());
+		NotificationInfo.Text = LOCTEXT("NotSupportedByMultipleObjects", "Operation is not supported for multi-selected objects");
+		NotificationInfo.ExpireDuration = 5.0f;
+		FSlateNotificationManager::Get().AddNotification(NotificationInfo);
+
+		return;
+	}
+
+	if (ParentArrayProperty)
+	{
+		if (UStateTreeEditorData* EditorDataPtr = GetEditorData())
 		{
-			if (UStateTreeEditorData* EditorData = GetEditorData())
+			void* TransitionPtr = nullptr;
+			if (StructProperty->GetValueData(TransitionPtr) == FPropertyAccess::Success)
 			{
-				FScopedTransaction Transaction(LOCTEXT("DuplicateTransition", "Duplicate Transition"));
+				UE::StateTreeEditor::FClipboardEditorData Clipboard;
+				Clipboard.Append(EditorDataPtr, TConstArrayView<FStateTreeTransition>(static_cast<FStateTreeTransition*>(TransitionPtr), 1));
+				Clipboard.ProcessBuffer(nullptr, EditorDataPtr, OuterObjects[0]);
 
-				EditorData->Modify();
+				if (!Clipboard.IsValid())
+				{
+					return;
+				}
+
+				FScopedTransaction Transaction(LOCTEXT("DuplicateTransition", "Duplicate Transition"));
 
-				ArrayHandle->DuplicateItem(ArrayIndex);
+				// Might modify the bindings data
+				EditorDataPtr->Modify();
 
-				TSharedRef<IPropertyHandle> DuplicatedStructHandle = ArrayHandle->GetElement(ArrayIndex);
+				const int32 Index = StructProperty->GetArrayIndex();
+				ParentArrayProperty->Insert(Index);
 
+				TSharedPtr<IPropertyHandle> InsertedElementHandle = ParentArrayProperty->GetElement(Index);
+				void* InsertedTransitionPtr = nullptr;
+				if (InsertedElementHandle->GetValueData(InsertedTransitionPtr) == FPropertyAccess::Success)
 				{
-					UE::StateTreeEditor::FScopedEditorDataFixer DataFixer(EditorData);
+					*static_cast<FStateTreeTransition*>(InsertedTransitionPtr) = MoveTemp(Clipboard.GetTransitionsInBuffer()[0]);
 
-					TArray<void*> RawNodeData;
-					StructProperty->AccessRawData(RawNodeData);
-					for (int32 Index = 0; Index < RawNodeData.Num(); ++Index)
+					for (FStateTreePropertyPathBinding& Binding : Clipboard.GetBindingsInBuffer())
 					{
-						if (FStateTreeTransition* Transition = static_cast<FStateTreeTransition*>(RawNodeData[Index]))
-						{
-							constexpr bool bShouldCopyBindings = true;
-							constexpr bool bShouldReinstantiateInstanceObjects = false; // ArrayHandle has already deep copied UObjects
-							constexpr bool bShouldRegenerateGUID = true;
-							DataFixer.TransitionsToFix.Emplace(OuterObjects[Index], *Transition, bShouldCopyBindings, bShouldReinstantiateInstanceObjects, bShouldRegenerateGUID);
-						}
+						EditorDataPtr->GetPropertyEditorBindings()->AddStateTreeBinding(MoveTemp(Binding));
 					}
 				}
+
+				// We reinieitalized item nodes on ArrayProperty operations before the data is completely set up. Reinitialize.
+				if (PropUtils)
+				{
+					PropUtils->ForceRefresh();
+				}
 			}
 		}
 	}
@@ -485,11 +517,44 @@ FText FStateTreeTransitionDetails::GetDescription() const
 
 void FStateTreeTransitionDetails::OnCopyTransition() const
 {
-	FString Value;
-	// Use PPF_Copy so that all properties get copied.
-	if (StructProperty->GetValueAsFormattedString(Value, PPF_Copy) == FPropertyAccess::Success)
+	UStateTreeEditorData* EditorDataPtr = GetEditorData();
+	if (!EditorDataPtr)
 	{
-		FPlatformApplicationMisc::ClipboardCopy(*Value);
+		return;
+	}
+
+	void* TransitionAddress = nullptr;
+	if (StructProperty->GetValueData(TransitionAddress) == FPropertyAccess::Success)
+	{
+		UE::StateTreeEditor::FClipboardEditorData Clipboard;
+		Clipboard.Append(EditorDataPtr, TConstArrayView<FStateTreeTransition>(static_cast<FStateTreeTransition*>(TransitionAddress), 1));
+
+		UE::StateTreeEditor::ExportTextAsClipboardEditorData(Clipboard);
+	}
+}
+
+void FStateTreeTransitionDetails::OnCopyAllTransitions() const
+{
+	UStateTreeEditorData* EditorDataPtr = GetEditorData();
+	if (!EditorDataPtr)
+	{
+		return;
+	}
+
+	if (ParentArrayProperty)
+	{
+		void* TransitionArrayAddress = nullptr;
+		if (ParentProperty->GetValueData(TransitionArrayAddress) == FPropertyAccess::Success)
+		{
+			UE::StateTreeEditor::FClipboardEditorData Clipboard;
+			Clipboard.Append(EditorDataPtr, *static_cast<TArray<FStateTreeTransition>*>(TransitionArrayAddress));
+
+			UE::StateTreeEditor::ExportTextAsClipboardEditorData(Clipboard);
+		}
+	}
+	else
+	{
+		OnCopyTransition();
 	}
 }
 
@@ -512,8 +577,10 @@ UStateTreeEditorData* FStateTreeTransitionDetails::GetEditorData() const
 	return nullptr;
 }
 
-void FStateTreeTransitionDetails::OnPasteTransition() const
+void FStateTreeTransitionDetails::OnPasteTransitions() const
 {
+	using namespace UE::StateTreeEditor;
+
 	TArray<UObject*> OuterObjects;
 	StructProperty->GetOuterObjects(OuterObjects);
 	if (OuterObjects.IsEmpty())
@@ -521,39 +588,146 @@ void FStateTreeTransitionDetails::OnPasteTransition() const
 		return;
 	}
 
-	FString PastedText;
-	FPlatformApplicationMisc::ClipboardPaste(PastedText);
+	UStateTreeEditorData* EditorDataPtr = GetEditorData();
+	if (!EditorDataPtr)
+	{
+		return;
+	}
+
+	// In case its multi selected, we need to have a unique copy for each Object
+	TArray<FClipboardEditorData, TInlineAllocator<2>> ClipboardTransitions;
+	ClipboardTransitions.AddDefaulted(OuterObjects.Num());
 
-	if (PastedText.IsEmpty())
+	for (int32 Idx = 0; Idx < OuterObjects.Num(); ++Idx)
+	{
+		const bool bSuccess = ImportTextAsClipboardEditorData
+		(TBaseStructure<FStateTreeTransition>::Get(),
+			EditorDataPtr,
+			OuterObjects[Idx],
+			ClipboardTransitions[Idx]);
+
+		if (!bSuccess)
+		{
+			return;
+		}
+	}
+
+	// make sure each Clipboard has the same number of nodes
+	for (int32 Idx = 0; Idx < ClipboardTransitions.Num() - 1; ++Idx)
+	{
+		check(ClipboardTransitions[Idx].GetTransitionsInBuffer().Num() == ClipboardTransitions[Idx + 1].GetTransitionsInBuffer().Num());
+	}
+
+	int32 NumTransitionsInBuffer = ClipboardTransitions[0].GetTransitionsInBuffer().Num();
+
+	if (NumTransitionsInBuffer == 0)
 	{
 		return;
 	}
 
-	if (UStateTreeEditorData* EditorData = GetEditorData())
+	if (!ParentArrayProperty.IsValid() && NumTransitionsInBuffer != 1)
+	{
+		// Transition is not in an array. we can't do multi-to-one paste
+		return;
+	}
+
+	if (OuterObjects.Num() != 1 && NumTransitionsInBuffer != 1)
+	{
+		// if multiple selected objects, and we have more than one nodes to paste into
+		// Array Handle doesn't support manipulation on multiple objects.
+		FNotificationInfo NotificationInfo(FText::GetEmpty());
+		NotificationInfo.Text = LOCTEXT("NotSupportedByMultipleObjects", "Operation is not supported for multi-selected objects");
+		NotificationInfo.ExpireDuration = 5.0f;
+		FSlateNotificationManager::Get().AddNotification(NotificationInfo);
+
+		return;
+	}
+
 	{
 		FScopedTransaction Transaction(LOCTEXT("PasteTransition", "Paste Transition"));
 
-		TArray<void*> RawData;
-		StructProperty->AccessRawData(RawData);
+		EditorDataPtr->Modify();	// we might modify the bindings on Editor Data
+		StructProperty->NotifyPreChange();
 
-		// Might modify the bindings
-		EditorData->Modify();
+		if (ParentArrayProperty)
+		{
+			// Paste multi nodes into one node
+			const int32 StructIndex = StructProperty->GetIndexInArray();
+			check(StructIndex != INDEX_NONE);
+
+			uint32 NumArrayElements = 0;
+			ParentArrayProperty->GetNumElements(NumArrayElements);
+			check(NumArrayElements > 0);	// since we already have at least one element to paste into
 
-		// Make sure we instantiate new objects when setting the value.
-		// Calls NotifyPreChange for us, will also call NotifyFinishedChangingProperties
-		StructProperty->SetValueFromFormattedString(PastedText, EPropertyValueSetFlags::InstanceObjects);
+			// Insert or append uninitialized elements after the current node to match the number of nodes in the paste buffer and retain the order of elements 
+			// The first node in the buffer goes into the current node
+			const int32 IndexToInsert = StructIndex + 1;
+			const int32 NumElementsToAddOrInsert = ClipboardTransitions[0].GetTransitionsInBuffer().Num() - 1;
 
+			int32 Cnt = 0;
+			if (IndexToInsert == NumArrayElements)
+			{
+				while (Cnt++ < NumElementsToAddOrInsert)
+				{
+					FPropertyHandleItemAddResult Result = ParentArrayProperty->AddItem();
+					if (Result.GetAccessResult() != FPropertyAccess::Success)
+					{
+						return;
+					}
+				}
+			}
+			else
+			{
+				while (Cnt++ < NumElementsToAddOrInsert)
+				{
+					FPropertyAccess::Result Result = ParentArrayProperty->Insert(IndexToInsert);
+					if (Result != FPropertyAccess::Success)
+					{
+						return;
+					}
+				}
+			}
+
+			TArray<void*> RawDatasArray;
+			ParentProperty->AccessRawData(RawDatasArray);
+			check(RawDatasArray.Num() == OuterObjects.Num());
+
+			for (int32 ObjIdx = 0; ObjIdx < OuterObjects.Num(); ++ObjIdx)
+			{
+				if (TArray<FStateTreeTransition>* TransitionsPtr = static_cast<TArray<FStateTreeTransition>*>(RawDatasArray[ObjIdx]))
+				{
+					TArrayView<FStateTreeTransition> TransitionsClipboardBuffer = ClipboardTransitions[ObjIdx].GetTransitionsInBuffer();
+					TArray<FStateTreeTransition>& TransitionsToPasteInto = *TransitionsPtr;
+
+					for (int32 Idx = 0; Idx < TransitionsClipboardBuffer.Num(); ++Idx)
+					{
+						TransitionsToPasteInto[Idx + StructIndex] = MoveTemp(TransitionsClipboardBuffer[Idx]);
+					}
+
+					for (FStateTreePropertyPathBinding& Binding : ClipboardTransitions[ObjIdx].GetBindingsInBuffer())
+					{
+						EditorDataPtr->GetPropertyEditorBindings()->AddStateTreeBinding(MoveTemp(Binding));
+					}
+				}
+			}
+		}
+		else
 		{
-			UE::StateTreeEditor::FScopedEditorDataFixer DataFixer(EditorData);
+			// Paste single node to a single Node
+			TArray<void*> RawDatas;
+			StructProperty->AccessRawData(RawDatas);
+			check(RawDatas.Num() == OuterObjects.Num());
 
-			for (int32 Index = 0; Index < RawData.Num(); Index++)
+			for (int32 Idx = 0; Idx < RawDatas.Num(); ++Idx)
 			{
-				if (FStateTreeTransition* Transition = static_cast<FStateTreeTransition*>(RawData[Index]))
+				if (FStateTreeTransition* CurrentTransition = static_cast<FStateTreeTransition*>(RawDatas[Idx]))
 				{
-					constexpr bool bShouldCopyBindings = true;
-					constexpr bool bShouldReinstantiateInstanceData = false; // StructProperty has already done UObject deep copies
-					constexpr bool bShouldRegenerateGUID = true;
-					DataFixer.TransitionsToFix.Emplace(OuterObjects[Index], *Transition, bShouldCopyBindings, bShouldReinstantiateInstanceData, bShouldRegenerateGUID);
+					*CurrentTransition = MoveTemp(ClipboardTransitions[Idx].GetTransitionsInBuffer()[0]);
+
+					for (FStateTreePropertyPathBinding& Binding : ClipboardTransitions[Idx].GetBindingsInBuffer())
+					{
+						EditorDataPtr->GetPropertyEditorBindings()->AddStateTreeBinding(MoveTemp(Binding));
+					}
 				}
 			}
 		}
diff --git a/Engine/Plugins/Runtime/StateTree/Source/StateTreeEditorModule/Private/Customizations/StateTreeTransitionDetails.h b/Engine/Plugins/Runtime/StateTree/Source/StateTreeEditorModule/Private/Customizations/StateTreeTransitionDetails.h
--- a/Engine/Plugins/Runtime/StateTree/Source/StateTreeEditorModule/Private/Customizations/StateTreeTransitionDetails.h
+++ b/Engine/Plugins/Runtime/StateTree/Source/StateTreeEditorModule/Private/Customizations/StateTreeTransitionDetails.h
@@ -7,6 +7,7 @@
 
 class FStateTreeViewModel;
 class IPropertyHandle;
+class IPropertyHandleArray;
 class IPropertyUtilities;
 class IDetailChildrenBuilder;
 class FDetailWidgetRow;
@@ -39,7 +40,9 @@ private:
 	bool GetDelayTransition() const;
 
 	void OnCopyTransition() const;
-	void OnPasteTransition() const;
+	void OnCopyAllTransitions() const;
+
+	void OnPasteTransitions() const;
 
 	FSlateColor GetContentRowColor() const;
 
@@ -66,4 +69,7 @@ private:
 
 	TSharedPtr<IPropertyUtilities> PropUtils;
 	TSharedPtr<IPropertyHandle> StructProperty;
+
+	TSharedPtr<IPropertyHandle> ParentProperty;
+	TSharedPtr<IPropertyHandleArray> ParentArrayProperty;
 };
diff --git a/Engine/Plugins/Runtime/StateTree/Source/StateTreeEditorModule/Private/Customizations/Widgets/SStateTreeContextMenuButton.cpp b/Engine/Plugins/Runtime/StateTree/Source/StateTreeEditorModule/Private/Customizations/Widgets/SStateTreeContextMenuButton.cpp
--- a/Engine/Plugins/Runtime/StateTree/Source/StateTreeEditorModule/Private/Customizations/Widgets/SStateTreeContextMenuButton.cpp
+++ b/Engine/Plugins/Runtime/StateTree/Source/StateTreeEditorModule/Private/Customizations/Widgets/SStateTreeContextMenuButton.cpp
@@ -103,6 +103,22 @@ TSharedRef<SWidget> SStateTreeContextMenuButton::MakeContextMenu() const
 			})
 		));
 
+		// Copy all
+		MenuBuilder.AddMenuEntry(
+			LOCTEXT("CopyAllItems", "Copy all"),
+			LOCTEXT("CopyAllItemsTooltip", "Copy all items"),
+			FSlateIcon(FAppStyle::GetAppStyleSetName(), "GenericCommands.Copy"),
+			FUIAction(
+				FExecuteAction::CreateSPLambda(this, [this]()
+					{
+						StateTreeViewModel->CopyAllNodes(OwnerStateWeak, NodeID);
+					}),
+				FCanExecuteAction::CreateSPLambda(this, [this]()
+					{
+						return !bIsTransition || bIsStateTransition;
+					})
+			));
+
 		// Paste
 		MenuBuilder.AddMenuEntry(
 			LOCTEXT("PasteItem", "Paste"),
diff --git a/Engine/Plugins/Runtime/StateTree/Source/StateTreeEditorModule/Private/SStateTreeView.cpp b/Engine/Plugins/Runtime/StateTree/Source/StateTreeEditorModule/Private/SStateTreeView.cpp
--- a/Engine/Plugins/Runtime/StateTree/Source/StateTreeEditorModule/Private/SStateTreeView.cpp
+++ b/Engine/Plugins/Runtime/StateTree/Source/StateTreeEditorModule/Private/SStateTreeView.cpp
@@ -210,12 +210,17 @@ void SStateTreeView::BindCommands()
 	CommandList->MapAction(
 		Commands.PasteStatesAsSiblings,
 		FExecuteAction::CreateSP(this, &SStateTreeView::HandlePasteStatesAsSiblings),
-		FCanExecuteAction::CreateSP(this, &SStateTreeView::CanPaste));
+		FCanExecuteAction::CreateSP(this, &SStateTreeView::CanPasteStates));
 
 	CommandList->MapAction(
 		Commands.PasteStatesAsChildren,
 		FExecuteAction::CreateSP(this, &SStateTreeView::HandlePasteStatesAsChildren),
-		FCanExecuteAction::CreateSP(this, &SStateTreeView::CanPaste));
+		FCanExecuteAction::CreateSP(this, &SStateTreeView::CanPasteStates));
+
+	CommandList->MapAction(
+		Commands.PasteNodesToSelectedStates,
+		FExecuteAction::CreateSP(this, &SStateTreeView::HandlePasteNodesToState),
+		FCanExecuteAction::CreateSP(this, &SStateTreeView::CanPasteNodesToSelectedStates));
 
 	CommandList->MapAction(
 		Commands.DuplicateStates,
@@ -314,7 +319,7 @@ bool SStateTreeView::HasSelection() const
 	return StateTreeViewModel && StateTreeViewModel->HasSelection();
 }
 
-bool SStateTreeView::CanPaste() const
+bool SStateTreeView::CanPasteStates() const
 {
 	return StateTreeViewModel
 			&& StateTreeViewModel->HasSelection()
@@ -335,6 +340,11 @@ bool SStateTreeView::CanDisableStates() const
 			&& StateTreeViewModel->CanDisableStates();
 }
 
+bool SStateTreeView::CanPasteNodesToSelectedStates() const
+{
+	return StateTreeViewModel && StateTreeViewModel->HasSelection() && StateTreeViewModel->CanPasteNodesToSelectedStates();
+}
+
 FReply SStateTreeView::OnKeyDown(const FGeometry& MyGeometry, const FKeyEvent& InKeyEvent)
 {
 	if(CommandList->ProcessCommandBindings(InKeyEvent))
@@ -632,6 +642,7 @@ TSharedPtr<SWidget> SStateTreeView::HandleContextMenuOpening()
 		{
 			MenuBuilder.AddMenuEntry(FStateTreeEditorCommands::Get().PasteStatesAsSiblings);
 			MenuBuilder.AddMenuEntry(FStateTreeEditorCommands::Get().PasteStatesAsChildren);
+			MenuBuilder.AddMenuEntry(FStateTreeEditorCommands::Get().PasteNodesToSelectedStates);
 		}),
 		/*bInOpenSubMenuOnClick =*/false,
 		FSlateIcon(FAppStyle::GetAppStyleSetName(), "GenericCommands.Paste")
@@ -758,6 +769,14 @@ void SStateTreeView::HandleDuplicateSelectedStates()
 	}
 }
 
+void SStateTreeView::HandlePasteNodesToState()
+{
+	if (StateTreeViewModel)
+	{
+		StateTreeViewModel->PasteNodesToSelectedStates();
+	}
+}
+
 void SStateTreeView::HandleDeleteStates()
 {
 	if (StateTreeViewModel)
diff --git a/Engine/Plugins/Runtime/StateTree/Source/StateTreeEditorModule/Private/SStateTreeView.h b/Engine/Plugins/Runtime/StateTree/Source/StateTreeEditorModule/Private/SStateTreeView.h
--- a/Engine/Plugins/Runtime/StateTree/Source/StateTreeEditorModule/Private/SStateTreeView.h
+++ b/Engine/Plugins/Runtime/StateTree/Source/StateTreeEditorModule/Private/SStateTreeView.h
@@ -70,15 +70,17 @@ private:
 	void HandlePasteStatesAsSiblings();
 	void HandlePasteStatesAsChildren();
 	void HandleDuplicateSelectedStates();
+	void HandlePasteNodesToState();
 	void HandleRenameState();
 	void HandleDeleteStates();
 	void HandleEnableSelectedStates();
 	void HandleDisableSelectedStates();
 
 	bool HasSelection() const;
-	bool CanPaste() const;
+	bool CanPasteStates() const;
 	bool CanEnableStates() const;
 	bool CanDisableStates() const;
+	bool CanPasteNodesToSelectedStates() const;
 
 	void BindCommands();
 
diff --git a/Engine/Plugins/Runtime/StateTree/Source/StateTreeEditorModule/Private/StateTreeEditorCommands.cpp b/Engine/Plugins/Runtime/StateTree/Source/StateTreeEditorModule/Private/StateTreeEditorCommands.cpp
--- a/Engine/Plugins/Runtime/StateTree/Source/StateTreeEditorModule/Private/StateTreeEditorCommands.cpp
+++ b/Engine/Plugins/Runtime/StateTree/Source/StateTreeEditorModule/Private/StateTreeEditorCommands.cpp
@@ -30,6 +30,7 @@ void FStateTreeEditorCommands::RegisterCommands()
 	UI_COMMAND(CopyStates, "Copy", "Copy Selected States", EUserInterfaceActionType::Button, FInputChord(EModifierKey::Control, EKeys::C));
 	UI_COMMAND(PasteStatesAsSiblings, "Paste As Siblings", "Paste States as Siblings", EUserInterfaceActionType::Button, FInputChord(EModifierKey::Control, EKeys::V));
 	UI_COMMAND(PasteStatesAsChildren, "Paste As Children", "Paste States as Children", EUserInterfaceActionType::Button, FInputChord(EModifierKey::Control | EModifierKey::Shift, EKeys::V));
+	UI_COMMAND(PasteNodesToSelectedStates, "Paste As Nodes and Transitions", "Append Nodes and Transitions to Selected States", EUserInterfaceActionType::Button, FInputChord());
 	UI_COMMAND(DuplicateStates, "Duplicate", "Duplicate Selected States", EUserInterfaceActionType::Button, FInputChord(EModifierKey::Control, EKeys::D));
 	UI_COMMAND(DeleteStates, "Delete", "Delete Selected States", EUserInterfaceActionType::Button, FInputChord(EKeys::Platform_Delete));
 	UI_COMMAND(RenameState, "Rename", "Rename Selected State", EUserInterfaceActionType::Button, FInputChord(EKeys::F2));
diff --git a/Engine/Plugins/Runtime/StateTree/Source/StateTreeEditorModule/Private/StateTreeEditorCommands.h b/Engine/Plugins/Runtime/StateTree/Source/StateTreeEditorModule/Private/StateTreeEditorCommands.h
--- a/Engine/Plugins/Runtime/StateTree/Source/StateTreeEditorModule/Private/StateTreeEditorCommands.h
+++ b/Engine/Plugins/Runtime/StateTree/Source/StateTreeEditorModule/Private/StateTreeEditorCommands.h
@@ -30,6 +30,7 @@ public:
 	TSharedPtr<FUICommandInfo> CopyStates;
 	TSharedPtr<FUICommandInfo> PasteStatesAsSiblings;
 	TSharedPtr<FUICommandInfo> PasteStatesAsChildren;
+	TSharedPtr<FUICommandInfo> PasteNodesToSelectedStates;
 	TSharedPtr<FUICommandInfo> DuplicateStates;
 	TSharedPtr<FUICommandInfo> DeleteStates;
 	TSharedPtr<FUICommandInfo> EnableStates;
diff --git a/Engine/Plugins/Runtime/StateTree/Source/StateTreeEditorModule/Private/StateTreeEditorDataClipboardHelpers.cpp b/Engine/Plugins/Runtime/StateTree/Source/StateTreeEditorModule/Private/StateTreeEditorDataClipboardHelpers.cpp
new file mode 100644
--- /dev/null
+++ b/Engine/Plugins/Runtime/StateTree/Source/StateTreeEditorModule/Private/StateTreeEditorDataClipboardHelpers.cpp
@@ -0,0 +1,344 @@
+// Copyright Epic Games, Inc. All Rights Reserved.
+
+#pragma once
+
+#include "StateTreeEditorDataClipboardHelpers.h"
+
+#include "Customizations/StateTreeEditorNodeUtils.h"
+#include "Framework/Notifications/NotificationManager.h"
+#include "HAL/PlatformApplicationMisc.h"
+#include "StateTreeConditionBase.h"
+#include "StateTreeConsiderationBase.h"
+#include "StateTreeEditorData.h"
+#include "StateTreeTaskBase.h"
+#include "Widgets/Notifications/SNotificationList.h"
+
+#define LOCTEXT_NAMESPACE "StateTreeEditor"
+
+namespace UE::StateTreeEditor
+{
+
+/** Helper class to detect if there were issues when calling ImportText() */
+class FDefaultValueImportErrorContext : public FOutputDevice
+{
+public:
+
+	int32 NumErrors = 0;
+
+	FDefaultValueImportErrorContext() = default;
+
+	virtual void Serialize(const TCHAR* V, ELogVerbosity::Type Verbosity, const class FName& Category) override
+	{
+		++NumErrors;
+	}
+};
+
+void ExportTextAsClipboardEditorData(const FClipboardEditorData& InClipboardEditorData)
+{
+	FString Value;
+
+	// Use PPF_Copy so that all properties get copied.
+	constexpr EPropertyPortFlags PortFlags = PPF_Copy;
+	TBaseStructure<FClipboardEditorData>::Get()->ExportText(Value, &InClipboardEditorData, nullptr, nullptr, PortFlags, nullptr);
+
+	FPlatformApplicationMisc::ClipboardCopy(*Value);
+}
+
+bool ImportTextAsClipboardEditorData(const UScriptStruct* InTargetType, TNotNull<UStateTreeEditorData*> InTargetTree, TNotNull<UObject*> InOwner,
+	FClipboardEditorData& OutClipboardEditorData, bool bProcessBuffer /*true*/)
+{
+	OutClipboardEditorData.Reset();
+
+	FString PastedText;
+	FPlatformApplicationMisc::ClipboardPaste(PastedText);
+
+	if (PastedText.IsEmpty())
+	{
+		return false;
+	}
+
+	const UScriptStruct* ClipboardEditorDataScriptStruct = TBaseStructure<FClipboardEditorData>::Get();
+	FDefaultValueImportErrorContext ErrorPipe;
+	ClipboardEditorDataScriptStruct->ImportText(*PastedText, &OutClipboardEditorData, nullptr, PPF_None, &ErrorPipe, ClipboardEditorDataScriptStruct->GetName());
+
+	if (bProcessBuffer)
+	{
+		OutClipboardEditorData.ProcessBuffer(InTargetType, InTargetTree, InOwner);
+	}
+
+	return !bProcessBuffer || OutClipboardEditorData.IsValid();
+}
+
+void RemoveInvalidBindings(TNotNull<UStateTreeEditorData*> InEditorData)
+{
+	if (FStateTreeEditorPropertyBindings* Bindings = InEditorData->GetPropertyEditorBindings())
+	{
+		TMap<FGuid, const FPropertyBindingDataView> AllStructValues;
+		InEditorData->GetAllStructValues(AllStructValues);
+		Bindings->RemoveInvalidBindings(AllStructValues);
+	}
+}
+
+struct FScopedEditorDataFixer
+{
+	explicit FScopedEditorDataFixer(TNotNull<UStateTreeEditorData*> InEditorData, TNotNull<UObject*> InOwner, bool InbShouldRegenerateGUID = true,
+		bool InbShouldReinstantiateInstanceData = true)
+		: EditorData(InEditorData)
+		, Owner(InOwner)
+		, bShouldRegenerateGUID(InbShouldRegenerateGUID)
+		, bShouldReinstantiateInstanceData(InbShouldReinstantiateInstanceData)
+	{
+	}
+
+	explicit FScopedEditorDataFixer(TNotNull<UStateTreeEditorData*> InEditorData, TNotNull<UObject*> InOwner, FClipboardEditorData& InClipboardEditorData)
+		: FScopedEditorDataFixer(InEditorData, InOwner)
+	{
+
+		EditorNodesToFix = InClipboardEditorData.GetEditorNodesInBuffer();
+		TransitionsToFix = InClipboardEditorData.GetTransitionsInBuffer();
+		BindingsToFix = InClipboardEditorData.GetBindingsInBuffer();
+	}
+
+	~FScopedEditorDataFixer()
+	{
+		// OldID -> NewID
+		TMap<FGuid, FGuid> IDsMap;
+
+		auto UpdateBindings = [Self = this, &IDsMap]()
+			{
+				for (FStateTreePropertyPathBinding& Binding : Self->BindingsToFix)
+				{
+					if (FGuid* NewSourceID = IDsMap.Find(Binding.GetSourcePath().GetStructID()))
+					{
+						Binding.GetMutableSourcePath().SetStructID(*NewSourceID);
+					}
+
+					if (FGuid* NewTargetID = IDsMap.Find(Binding.GetTargetPath().GetStructID()))
+					{
+						Binding.GetMutableTargetPath().SetStructID(*NewTargetID);
+					}
+				}
+			};
+
+		auto ReinstantiateEditorNodeInstanceData = [Self = this](FStateTreeEditorNode& EditorNode)
+			{
+				EditorNodeUtils::InstantiateStructSubobjects(*Self->Owner, EditorNode.Node);
+				if (EditorNode.InstanceObject)
+				{
+					EditorNode.InstanceObject = DuplicateObject(EditorNode.InstanceObject, Self->Owner);
+				}
+				else
+				{
+					EditorNodeUtils::InstantiateStructSubobjects(*Self->Owner, EditorNode.Instance);
+				}
+			};
+
+		auto FixEditorNode = [Self = this, &IDsMap, &ReinstantiateEditorNodeInstanceData](FStateTreeEditorNode& EditorNode)
+			{
+				FGuid OldInstanceID = EditorNode.ID;
+				FGuid OldTemplateID = EditorNode.GetNodeID();
+				if (Self->bShouldRegenerateGUID)
+				{
+					EditorNode.ID = FGuid::NewGuid();
+
+					IDsMap.Add(OldInstanceID, EditorNode.ID);
+					IDsMap.Add(OldTemplateID, EditorNode.GetNodeID());
+				}
+
+				if (Self->bShouldReinstantiateInstanceData)
+				{
+					ReinstantiateEditorNodeInstanceData(EditorNode);
+				}
+			};
+
+		for (FStateTreeEditorNode& EditorNode : EditorNodesToFix)
+		{
+			FixEditorNode(EditorNode);
+		}
+
+		for (FStateTreeTransition& Transition : TransitionsToFix)
+		{
+			if (bShouldRegenerateGUID)
+			{
+				FGuid OldTransitionID = Transition.ID;
+				Transition.ID = FGuid::NewGuid();
+				IDsMap.Add(OldTransitionID, Transition.ID);
+			}
+
+			for (FStateTreeEditorNode& CondNode : Transition.Conditions)
+			{
+				FixEditorNode(CondNode);
+			}
+		}
+
+		for (FStateTreePropertyPathBinding& Binding : BindingsToFix)
+		{
+			FStructView PropertyFunctionView = Binding.GetMutablePropertyFunctionNode();
+			if (FStateTreeEditorNode* PropertyFunctionNode = PropertyFunctionView.GetPtr<FStateTreeEditorNode>())
+			{
+				FixEditorNode(*PropertyFunctionNode);
+			}
+		}
+
+		UpdateBindings();
+	}
+
+private:
+	TArrayView<FStateTreeEditorNode> EditorNodesToFix;
+	TArrayView<FStateTreeTransition> TransitionsToFix;
+	TArrayView<FStateTreePropertyPathBinding> BindingsToFix;
+
+	TNotNull<UStateTreeEditorData*> EditorData;
+	TNotNull<UObject*> Owner;
+
+	uint8 bShouldRegenerateGUID : 1;
+	uint8 bShouldReinstantiateInstanceData : 1;
+};
+
+void FClipboardEditorData::Append(TNotNull<const UStateTreeEditorData*> InStateTree, TConstArrayView<FStateTreeEditorNode> InEditorNodes)
+{
+	bBufferProcessed = false;
+	EditorNodesBuffer.Append(InEditorNodes);
+	
+	CollectBindingsForEditorNodes(InStateTree, InEditorNodes);
+}
+
+void FClipboardEditorData::Append(TNotNull<const UStateTreeEditorData*> InStateTree, TConstArrayView<FStateTreeTransition> InTransitions)
+{
+	bBufferProcessed = false;
+	TransitionsBuffer.Append(InTransitions);
+
+	for (const FStateTreeTransition& Transition : InTransitions)
+	{
+		CollectBindingsForEditorNodes(InStateTree, Transition.Conditions);
+	}
+}
+
+void FClipboardEditorData::Reset()
+{
+	EditorNodesBuffer.Reset();
+	TransitionsBuffer.Reset();
+	BindingsBuffer.Reset();
+	bBufferProcessed = false;
+}
+
+void FClipboardEditorData::CollectBindingsForEditorNodes(TNotNull<const UStateTreeEditorData*> InStateTree, TConstArrayView<FStateTreeEditorNode> InEditorNodes)
+{
+	TArray<const FPropertyBindingBinding*> TempBindingPtrs;
+	
+	for (const FStateTreeEditorNode& EditorNode : InEditorNodes)
+	{
+		InStateTree->GetPropertyEditorBindings()->GetBindingsFor(EditorNode.ID, TempBindingPtrs);
+
+		// recursively collect property function node bindings
+		constexpr auto EmptyStatePath = TEXT("");
+		InStateTree->VisitStructBoundPropertyFunctions(EditorNode.ID, EmptyStatePath, 
+			[InStateTree, &TempBindingPtrs](const FStateTreeEditorNode& EditorNode, const FStateTreeBindableStructDesc& Desc, const FStateTreeDataView Value)
+			{
+				TArray<const FPropertyBindingBinding*> StructBindingsPtr;
+				InStateTree->GetPropertyEditorBindings()->GetBindingsFor(Desc.ID, StructBindingsPtr);
+				TempBindingPtrs.Append(MoveTemp(StructBindingsPtr));
+
+				return EStateTreeVisitor::Continue;
+			});
+	}
+
+	Algo::Transform(TempBindingPtrs, BindingsBuffer, [](const FPropertyBindingBinding* BindingPtr) { return *static_cast<const FStateTreePropertyPathBinding*>(BindingPtr); });
+}
+
+void FClipboardEditorData::ProcessBuffer(const UScriptStruct* InTargetType, TNotNull<UStateTreeEditorData*> InEditorData, TNotNull<UObject*> InTargetOwner)
+{
+	bBufferProcessed = false;
+	check(!InTargetType || InTargetType->IsChildOf(TBaseStructure<FStateTreeNodeBase>::Get()) || InTargetType->IsChildOf(TBaseStructure<FStateTreeTransition>::Get()));
+	check(InTargetOwner->IsA<UStateTreeState>() || InTargetOwner->IsA<UStateTreeEditorData>());
+
+	auto AddErrorNotification = [](const FText& InText)
+		{
+			FNotificationInfo NotificationInfo(FText::GetEmpty());
+			NotificationInfo.Text = InText;
+			NotificationInfo.ExpireDuration = 5.0f;
+			FSlateNotificationManager::Get().AddNotification(NotificationInfo);
+		};
+
+	auto ValidateEditorNodes = [&InEditorData, &AddErrorNotification](TConstArrayView<FStateTreeEditorNode> InEditorNodes, const UScriptStruct* InRequiredType)
+		{
+			for (const FStateTreeEditorNode& EditorNode : InEditorNodes)
+			{
+				FStructView NodeView = EditorNode.GetNode();
+				FStateTreeDataView InstanceView = EditorNode.GetInstance();
+				if (!NodeView.IsValid() || !InstanceView.IsValid())
+				{
+					AddErrorNotification(LOCTEXT("MalformedNode", "Clipboard text contains invalid data."));
+					return false;
+				}
+
+				if (InRequiredType && !NodeView.GetScriptStruct()->IsChildOf(InRequiredType))
+				{
+					AddErrorNotification(FText::Format(LOCTEXT("NotSupportedByType", "This property only accepts nodes of type {0}."), InRequiredType->GetDisplayNameText()));
+					return false;
+				}
+
+				if (const UStateTreeSchema* TargetSchema = InEditorData->Schema)
+				{
+					bool bIsNodeAllowed = true;
+					const UStruct* InstanceType = InstanceView.GetStruct();
+					if (const UScriptStruct* InstanceTypeStruct = Cast<UScriptStruct>(InstanceType))
+					{
+						if (!TargetSchema->IsStructAllowed(NodeView.GetScriptStruct()))
+						{
+							bIsNodeAllowed = false;
+						}
+					}
+					else if (const UClass* InstanceTypeClass = Cast<UClass>(InstanceType))
+					{
+						if (!TargetSchema->IsClassAllowed(InstanceTypeClass))
+						{
+							bIsNodeAllowed = false;
+						}
+					}
+
+					if (!bIsNodeAllowed)
+					{
+						AddErrorNotification(FText::Format(LOCTEXT("NotSupportedBySchema", "Node {0} is not supported by {1} schema."),
+							NodeView.GetScriptStruct()->GetDisplayNameText(), TargetSchema->GetClass()->GetDisplayNameText()));
+
+						return false;
+					}
+				}
+			}
+
+			return true;
+		};
+
+	auto ValidateTransitions = [Self = this, &AddErrorNotification, &ValidateEditorNodes](TConstArrayView<FStateTreeTransition> InTransitions, const UScriptStruct* InRequiredType)
+		{
+			if (Self->TransitionsBuffer.Num() && InRequiredType && InRequiredType != TBaseStructure<FStateTreeTransition>::Get())
+			{
+				AddErrorNotification(FText::Format(LOCTEXT("NotSupportedByType", "This property only accepts nodes of type {0}."), InRequiredType->GetDisplayNameText()));
+				return false;
+			}
+
+			for (const FStateTreeTransition& Transition : InTransitions)
+			{
+				static const UScriptStruct* ConditionRequiredType = TBaseStructure<FStateTreeConditionBase>::Get();
+				if (!ValidateEditorNodes(Transition.Conditions, ConditionRequiredType))
+				{
+					return false;
+				}
+			}
+
+			return true;
+		};
+
+	if (ValidateEditorNodes(EditorNodesBuffer, InTargetType) && ValidateTransitions(TransitionsBuffer, InTargetType))
+	{
+		FScopedEditorDataFixer Fixer(InEditorData, InTargetOwner, *this);
+		bBufferProcessed = true;
+	}
+	else
+	{
+		Reset();
+	}
+}
+}
+#undef LOCTEXT_NAMESPACE
diff --git a/Engine/Plugins/Runtime/StateTree/Source/StateTreeEditorModule/Private/StateTreeEditorDataClipboardHelpers.h b/Engine/Plugins/Runtime/StateTree/Source/StateTreeEditorModule/Private/StateTreeEditorDataClipboardHelpers.h
new file mode 100644
--- /dev/null
+++ b/Engine/Plugins/Runtime/StateTree/Source/StateTreeEditorModule/Private/StateTreeEditorDataClipboardHelpers.h
@@ -0,0 +1,74 @@
+// Copyright Epic Games, Inc. All Rights Reserved.
+
+#pragma once
+
+#include "StateTreeEditorDataClipboardHelpers.generated.h"
+
+namespace UE::StateTreeEditor
+{
+struct FClipboardEditorData;
+}
+
+struct FPropertyBindingBinding;
+struct FStateTreeTransition;
+struct FStateTreeEditorNode;
+struct FStateTreePropertyPathBinding;
+
+class UStateTreeEditorData;
+
+namespace UE::StateTreeEditor
+{
+void ExportTextAsClipboardEditorData(const FClipboardEditorData& InClipboardEditorData);
+
+bool ImportTextAsClipboardEditorData(const UScriptStruct* InTargetType, TNotNull<UStateTreeEditorData*> InTargetTree, TNotNull<UObject*> InOwner,
+                                     FClipboardEditorData& OutClipboardEditorData, bool bValidateBuffer = true);
+
+void RemoveInvalidBindings(TNotNull<UStateTreeEditorData*> InEditorData);
+
+USTRUCT()
+struct FClipboardEditorData
+{
+	GENERATED_BODY()
+
+	void ProcessBuffer(const UScriptStruct* InTargetType, TNotNull<UStateTreeEditorData*> InEditorData, TNotNull<UObject*> InOwner);
+
+	bool IsValid() const
+	{
+		return bBufferProcessed;
+	}
+
+	void Append(TNotNull<const UStateTreeEditorData*> InStateTree, TConstArrayView<FStateTreeEditorNode> InEditorNodes);
+	void Append(TNotNull<const UStateTreeEditorData*> InStateTree, TConstArrayView<FStateTreeTransition> InTransitions);
+
+	void Reset();
+
+	TArrayView<FStateTreeEditorNode> GetEditorNodesInBuffer()
+	{
+		return EditorNodesBuffer;
+	}
+
+	TArrayView<FStateTreeTransition> GetTransitionsInBuffer()
+	{
+		return TransitionsBuffer;
+	}
+
+	TArrayView<FStateTreePropertyPathBinding> GetBindingsInBuffer()
+	{
+		return BindingsBuffer;
+	}
+
+private:
+	void CollectBindingsForEditorNodes(TNotNull<const UStateTreeEditorData*> InStateTree, TConstArrayView<FStateTreeEditorNode> InEditorNodes);
+
+	UPROPERTY()
+	TArray<FStateTreeEditorNode> EditorNodesBuffer;
+
+	UPROPERTY()
+	TArray<FStateTreeTransition> TransitionsBuffer;
+
+	UPROPERTY()
+	TArray<FStateTreePropertyPathBinding> BindingsBuffer;
+
+	bool bBufferProcessed = false;
+};
+}
diff --git a/Engine/Plugins/Runtime/StateTree/Source/StateTreeEditorModule/Private/StateTreeScopedEditorDataFixer.cpp b/Engine/Plugins/Runtime/StateTree/Source/StateTreeEditorModule/Private/StateTreeScopedEditorDataFixer.cpp
deleted file mode 100644
diff --git a/Engine/Plugins/Runtime/StateTree/Source/StateTreeEditorModule/Private/StateTreeScopedEditorDataFixer.h b/Engine/Plugins/Runtime/StateTree/Source/StateTreeEditorModule/Private/StateTreeScopedEditorDataFixer.h
deleted file mode 100644
diff --git a/Engine/Plugins/Runtime/StateTree/Source/StateTreeEditorModule/Private/StateTreeViewModel.cpp b/Engine/Plugins/Runtime/StateTree/Source/StateTreeEditorModule/Private/StateTreeViewModel.cpp
--- a/Engine/Plugins/Runtime/StateTree/Source/StateTreeEditorModule/Private/StateTreeViewModel.cpp
+++ b/Engine/Plugins/Runtime/StateTree/Source/StateTreeEditorModule/Private/StateTreeViewModel.cpp
@@ -2,7 +2,6 @@
 
 #include "StateTreeViewModel.h"
 
-#include "Customizations/StateTreeEditorNodeUtils.h"
 #include "StateTree.h"
 #include "StateTreeEditorData.h"
 #include "StateTreeDelegates.h"
@@ -17,7 +16,7 @@
 #include "StateTreeConditionBase.h"
 #include "StateTreeConsiderationBase.h"
 #include "StateTreePropertyHelpers.h"
-#include "StateTreeScopedEditorDataFixer.h"
+#include "StateTreeEditorDataClipboardHelpers.h"
 #include "Misc/NotNull.h"
 #include "Misc/StringOutputDevice.h"
 
@@ -918,9 +917,6 @@ void FStateTreeViewModel::DeleteNode(TWeakObjectPtr<UStateTreeState> State, cons
 				{
 					StateNodeHandle.EditorData->Modify();
 
-					UE::StateTreeEditor::FScopedEditorDataFixer DataFixer(InEditorData);
-					DataFixer.bRemoveInvalidBindings = true;
-
 					if (StateNodeHandle.IsTransition())
 					{
 						StateNodeHandle.GetTransitionArray().RemoveAt(StateNodeHandle.GetNodeIndex());
@@ -929,6 +925,8 @@ void FStateTreeViewModel::DeleteNode(TWeakObjectPtr<UStateTreeState> State, cons
 					{
 						StateNodeHandle.GetEditorNodeArray().RemoveAt(StateNodeHandle.GetNodeIndex());
 					}
+
+					UE::StateTreeEditor::RemoveInvalidBindings(InEditorData);
 				};
 
 			UE::StateTree::PropertyHelpers::ModifyStateInPreAndPostEdit(
@@ -963,9 +961,6 @@ void FStateTreeViewModel::DeleteAllNodes(TWeakObjectPtr<UStateTreeState> State,
 				{
 					StateNodeHandle.EditorData->Modify();
 
-					UE::StateTreeEditor::FScopedEditorDataFixer DataFixer(InEditorData);
-					DataFixer.bRemoveInvalidBindings = true;
-
 					if (StateNodeHandle.IsTransition())
 					{
 						StateNodeHandle.GetTransitionArray().Empty();
@@ -974,6 +969,8 @@ void FStateTreeViewModel::DeleteAllNodes(TWeakObjectPtr<UStateTreeState> State,
 					{
 						StateNodeHandle.GetEditorNodeArray().Empty();
 					}
+
+					UE::StateTreeEditor::RemoveInvalidBindings(InEditorData);
 				};
 
 			UE::StateTree::PropertyHelpers::ModifyStateInPreAndPostEdit(
@@ -1004,22 +1001,44 @@ void FStateTreeViewModel::CopyNode(TWeakObjectPtr<UStateTreeState> State, const
 				return;
 			}
 
-			const UScriptStruct* BaseStruct = nullptr;
+			UE::StateTreeEditor::FClipboardEditorData ClipboardEditorData;
 			if (StateNodeHandle.IsTransition())
 			{
-				BaseStruct = FStateTreeTransition::StaticStruct();
+				ClipboardEditorData.Append(EditorData, TConstArrayView<FStateTreeTransition>(&StateNodeHandle.GetTransition(), 1));
 			}
 			else
 			{
-				BaseStruct = FStateTreeEditorNode::StaticStruct();
+				ClipboardEditorData.Append(EditorData, TConstArrayView<FStateTreeEditorNode>(&StateNodeHandle.GetEditorNode(), 1));
+			}
+			UE::StateTreeEditor::ExportTextAsClipboardEditorData(ClipboardEditorData);
+		}
+	}
+}
+
+void FStateTreeViewModel::CopyAllNodes(TWeakObjectPtr<UStateTreeState> State, const FGuid& ID)
+{
+	if (UStateTreeEditorData* EditorData = TreeDataWeak.Get())
+	{
+		if (UStateTreeState* OwnerState = State.Get())
+		{
+			UE::StateTree::Editor::Private::FStateTreeStateNodeEditorHandle StateNodeHandle(EditorData, OwnerState, ID);
+			// If the op can not be executed, don't even start the transaction
+			if (!StateNodeHandle.IsValid())
+			{
+				return;
 			}
 
-			// Use PPF_Copy so that all properties get copied.
-			FString Value;
-			constexpr EPropertyPortFlags PortFlags = PPF_Copy;
-			BaseStruct->ExportText(Value, StateNodeHandle.IsTransition() ? (void*)&StateNodeHandle.GetTransition() : (void*)&StateNodeHandle.GetEditorNode(), nullptr, nullptr, PortFlags, nullptr);
+			UE::StateTreeEditor::FClipboardEditorData ClipboardEditorData;
+			if (StateNodeHandle.IsTransition())
+			{
+				ClipboardEditorData.Append(EditorData, StateNodeHandle.GetTransitionArray());
+			}
+			else
+			{
+				ClipboardEditorData.Append(EditorData, StateNodeHandle.GetEditorNodeArray());
+			}
 
-			FPlatformApplicationMisc::ClipboardCopy(*Value);
+			UE::StateTreeEditor::ExportTextAsClipboardEditorData(ClipboardEditorData);
 		}
 	}
 }
@@ -1037,16 +1056,19 @@ void FStateTreeViewModel::PasteNode(TWeakObjectPtr<UStateTreeState> State, const
 				return;
 			}
 
-			// For Editor Node, we have a validation process which might early out if failed. Dry run with a Temp.
-			FStateTreeEditorNode TempEditorNode;
-			if (!StateNodeHandle.IsTransition())
+			const UScriptStruct* TargetType = nullptr;
+			if (StateNodeHandle.IsTransition())
+			{
+				TargetType = TBaseStructure<FStateTreeTransition>::Get();
+			}
+			else
 			{
 				static const UScriptStruct* TaskBaseStruct = FStateTreeTaskBase::StaticStruct();
 				static const UScriptStruct* ConditionBaseStruct = FStateTreeConditionBase::StaticStruct();
 				static const UScriptStruct* ConsiderationBaseStruct = FStateTreeConsiderationBase::StaticStruct();
 
 				const UScriptStruct* BaseNodeScriptStruct = nullptr;
-				const UScriptStruct* NodeScriptStruct = StateNodeHandle.GetEditorNode().Node.GetScriptStruct();
+				const UScriptStruct* NodeScriptStruct = StateNodeHandle.GetEditorNode().GetNode().GetScriptStruct();
 				if (NodeScriptStruct->IsChildOf(TaskBaseStruct))
 				{
 					BaseNodeScriptStruct = TaskBaseStruct;
@@ -1060,36 +1082,58 @@ void FStateTreeViewModel::PasteNode(TWeakObjectPtr<UStateTreeState> State, const
 					BaseNodeScriptStruct = ConsiderationBaseStruct;
 				}
 
-				const bool bSuccess = UE::StateTreeEditor::EditorNodeUtils::ImportTextAsNode(BaseNodeScriptStruct, StateNodeHandle.EditorData, TempEditorNode);
-				// If the op can not be executed, don't even start the transaction
-				if (!bSuccess)
-				{
-					return;
-				}
+				TargetType = BaseNodeScriptStruct;
 			}
 
-			auto PasteFunc = [&StateNodeHandle, &TempEditorNode](const TNotNull<UStateTreeState*> InOwnerState, const TNotNull<UStateTreeEditorData*> InEditorData, const FStateTreeEditPropertyPath& InPropertyPath)
+			UE::StateTreeEditor::FClipboardEditorData ClipboardEditorData;
+			const bool bSucceeded = UE::StateTreeEditor::ImportTextAsClipboardEditorData(TargetType, EditorData, OwnerState, ClipboardEditorData);
+
+			if (!bSucceeded)
+			{
+				return;
+			}
+
+			if (StateNodeHandle.IsTransition() && ClipboardEditorData.GetTransitionsInBuffer().IsEmpty())
+			{
+				return;
+			}
+
+			if (!StateNodeHandle.IsTransition() && ClipboardEditorData.GetEditorNodesInBuffer().IsEmpty())
+			{
+				return;
+			}
+
+			auto PasteFunc = [&StateNodeHandle, &ClipboardEditorData](const TNotNull<UStateTreeState*> InOwnerState, const TNotNull<UStateTreeEditorData*> InEditorData, const FStateTreeEditPropertyPath& InPropertyPath)
 				{
-					InEditorData->Modify();
+					auto PasteFuncInternal = [InEditorData, &ClipboardEditorData]<typename T>(TArray<T>& Dest, TArrayView<T> Source, int32 DestIndex) requires
+						std::is_same_v<T, FStateTreeEditorNode> || std::is_same_v<T, FStateTreeTransition>
+					{
+						check(Source.Num());
+
+						const int32 DestStartIndex = DestIndex;
+						Dest[DestIndex++] = MoveTemp(Source[0]);
 
-					UE::StateTreeEditor::FScopedEditorDataFixer DataFixer(StateNodeHandle.EditorData);
-					constexpr bool bShouldCopyBindings = true;
-					constexpr bool bShouldReinstantiateInstanceData = true;
-					constexpr bool bShouldRegenerateGUID = true;
+						Dest.InsertUninitialized(DestIndex, Source.Num() - 1);
+						for (int32 SourceIndex = 1; SourceIndex < Source.Num(); SourceIndex++, DestIndex++)
+						{
+							new (Dest.GetData() + DestIndex) T(MoveTemp(Source[SourceIndex]));
+						}
+
+						for (FStateTreePropertyPathBinding& Binding : ClipboardEditorData.GetBindingsInBuffer())
+						{
+							InEditorData->GetPropertyEditorBindings()->AddStateTreeBinding(MoveTemp(Binding));
+						}
+					};
+
+					InEditorData->Modify();
 
 					if (StateNodeHandle.IsTransition())
 					{
-						FString PastedText;
-						FPlatformApplicationMisc::ClipboardPaste(PastedText);
-
-						const UScriptStruct* TransitionStruct = FStateTreeTransition::StaticStruct();
-						TransitionStruct->ImportText(*PastedText, &StateNodeHandle.GetTransition(), nullptr, PPF_None, nullptr, TransitionStruct->GetName());
-						DataFixer.TransitionsToFix.Emplace(StateNodeHandle.OwnerState, StateNodeHandle.GetTransition(), bShouldCopyBindings, bShouldReinstantiateInstanceData, bShouldRegenerateGUID);
+						PasteFuncInternal(StateNodeHandle.GetTransitionArray(), ClipboardEditorData.GetTransitionsInBuffer(), StateNodeHandle.GetNodeIndex());
 					}
 					else
 					{
-						StateNodeHandle.GetEditorNode() = TempEditorNode;
-						DataFixer.EditorNodesToFix.Emplace(StateNodeHandle.OwnerState, StateNodeHandle.GetEditorNode(), bShouldCopyBindings, bShouldReinstantiateInstanceData, bShouldRegenerateGUID);
+						PasteFuncInternal(StateNodeHandle.GetEditorNodeArray(), ClipboardEditorData.GetEditorNodesInBuffer(), StateNodeHandle.GetNodeIndex());
 					}
 				};
 
@@ -1107,6 +1151,48 @@ void FStateTreeViewModel::PasteNode(TWeakObjectPtr<UStateTreeState> State, const
 	}
 }
 
+void FStateTreeViewModel::PasteNodesToSelectedStates()
+{
+	if (UStateTreeEditorData* EditorData = TreeDataWeak.Get())
+	{
+		TArray<UStateTreeState*> CurrentlySelectedStates;
+		GetSelectedStates(CurrentlySelectedStates);
+
+		if (CurrentlySelectedStates.Num())
+		{
+			using namespace UE::StateTreeEditor;
+			TArray<FClipboardEditorData, TInlineAllocator<4>> Clipboards;
+			Clipboards.AddDefaulted(CurrentlySelectedStates.Num());
+
+			for (int32 Idx = 0; Idx < CurrentlySelectedStates.Num(); ++Idx)
+			{
+				FClipboardEditorData& Clipboard = Clipboards[Idx];
+
+				// any editor nodes and transitions
+				constexpr const UScriptStruct* TargetType = nullptr;
+				ImportTextAsClipboardEditorData(TargetType, EditorData, CurrentlySelectedStates[Idx], Clipboard);
+
+				if (!Clipboard.IsValid())
+				{
+					return;
+				}
+			}
+
+			FScopedTransaction ScopedTransaction(LOCTEXT("PasteNodesToSelectedStatesTransaction", "Paste Nodes To Selected States"));
+			EditorData->Modify();
+
+			for (int32 Idx = 0; Idx < CurrentlySelectedStates.Num(); ++Idx)
+			{
+				UStateTreeState* SelectedState = CurrentlySelectedStates[Idx];
+
+				SelectedState->Modify();
+
+				PasteNodesToState(EditorData, SelectedState, Clipboards[Idx]);
+			}
+		}
+	}
+}
+
 void FStateTreeViewModel::DuplicateNode(TWeakObjectPtr<UStateTreeState> State, const FGuid& ID)
 {
 	if (UStateTreeEditorData* EditorData = TreeDataWeak.Get())
@@ -1122,28 +1208,56 @@ void FStateTreeViewModel::DuplicateNode(TWeakObjectPtr<UStateTreeState> State, c
 
 			auto DuplicateFunc = [&StateNodeHandle](const TNotNull<UStateTreeState*> InOwnerState, const TNotNull<UStateTreeEditorData*> InEditorData, const FStateTreeEditPropertyPath& InPropertyPath)
 				{
-					InEditorData->Modify();
+					int32 DestIndex = StateNodeHandle.GetNodeIndex();
+					UE::StateTreeEditor::FClipboardEditorData Clipboard;
 
-					UE::StateTreeEditor::FScopedEditorDataFixer DataFixer(StateNodeHandle.EditorData);
+					auto DuplicateFuncInternal = [InEditorData, &Clipboard]<typename T>(TArray<T>& Dest, int32 Index, TArrayView<T> Source)
+						requires std::is_same_v<T, FStateTreeEditorNode> || std::is_same_v<T, FStateTreeTransition>
+					{
+						check(Source.Num());
+						Dest.InsertUninitialized(Index, Source.Num());
 
-					constexpr bool bShouldCopyBindings = true;
-					constexpr bool bShouldReinstantiateInstanceData = true;
-					constexpr bool bShouldRegenerateGUID = true;
+						T* BaseAddress = Dest.GetData() + Index;
+						for (int32 SourceIdx = 0; SourceIdx < Source.Num(); ++SourceIdx)
+						{
+							::new (BaseAddress + SourceIdx) T(MoveTemp(Source[SourceIdx]));
+						}
+
+						for (FStateTreePropertyPathBinding& Binding : Clipboard.GetBindingsInBuffer())
+						{
+							InEditorData->GetPropertyEditorBindings()->AddStateTreeBinding(MoveTemp(Binding));
+						}
+					};
+
+					InEditorData->Modify();
 
 					if (StateNodeHandle.IsTransition())
 					{
-						// Make a temp copy
-						const FStateTreeTransition TempTransition = StateNodeHandle.GetTransition();
-						FStateTreeTransition& NewTransition = StateNodeHandle.GetTransitionArray().Insert_GetRef(TempTransition, StateNodeHandle.GetNodeIndex());
+						FStateTreeTransition& SourceTransition = StateNodeHandle.GetTransition();
+
+						Clipboard.Append(InEditorData, TConstArrayView<FStateTreeTransition>(&SourceTransition, 1));
+						Clipboard.ProcessBuffer(nullptr, InEditorData, InOwnerState);
 
-						DataFixer.TransitionsToFix.Emplace(StateNodeHandle.OwnerState, NewTransition, bShouldCopyBindings, bShouldReinstantiateInstanceData, bShouldRegenerateGUID);
+						if (!Clipboard.IsValid())
+						{
+							return;
+						}
+
+						DuplicateFuncInternal(StateNodeHandle.GetTransitionArray(), DestIndex, Clipboard.GetTransitionsInBuffer());
 					}
 					else
 					{
-						// Make a temp copy
-						const FStateTreeEditorNode TempEditorNode = StateNodeHandle.GetEditorNode();
-						FStateTreeEditorNode& NewEditorNode = StateNodeHandle.GetEditorNodeArray().Insert_GetRef(TempEditorNode, StateNodeHandle.GetNodeIndex());
-						DataFixer.EditorNodesToFix.Emplace(StateNodeHandle.OwnerState, NewEditorNode, bShouldCopyBindings, bShouldReinstantiateInstanceData, bShouldRegenerateGUID);
+						FStateTreeEditorNode& SourceEditorNode = StateNodeHandle.GetEditorNode();
+
+						Clipboard.Append(InEditorData, TConstArrayView<FStateTreeEditorNode>(&SourceEditorNode, 1));
+						Clipboard.ProcessBuffer(nullptr, InEditorData, InOwnerState);
+
+						if (!Clipboard.IsValid())
+						{
+							return;
+						}
+
+						DuplicateFuncInternal(StateNodeHandle.GetEditorNodeArray(), DestIndex, Clipboard.GetEditorNodesInBuffer());
 					}
 				};
 
@@ -1665,6 +1779,30 @@ bool FStateTreeViewModel::CanDisableStates() const
 	return false;
 }
 
+bool FStateTreeViewModel::CanPasteNodesToSelectedStates() const
+{
+	UStateTreeEditorData* TreeData = TreeDataWeak.Get();
+	if (!TreeData)
+	{
+		return false;
+	}
+
+	TArray<UStateTreeState*> States;
+	GetSelectedStates(States);
+	if (States.IsEmpty())
+	{
+		return false;
+	}
+
+	UE::StateTreeEditor::FClipboardEditorData ClipboardEditorData;
+
+	constexpr const UScriptStruct* TargetType = nullptr;
+	constexpr bool bProcessBuffer = false;
+	const bool bSucceeded = UE::StateTreeEditor::ImportTextAsClipboardEditorData(TargetType, TreeData, States[0], ClipboardEditorData, bProcessBuffer);
+
+	return bSucceeded && (ClipboardEditorData.GetEditorNodesInBuffer().Num() || ClipboardEditorData.GetTransitionsInBuffer().Num());
+}
+
 void FStateTreeViewModel::SetSelectedStatesEnabled(const bool bEnable)
 {
 	TArray<UStateTreeState*> States;
@@ -1844,6 +1982,109 @@ void FStateTreeViewModel::BindToDebuggerDelegates()
 #endif // WITH_STATETREE_TRACE_DEBUGGER
 }
 
+void FStateTreeViewModel::PasteNodesToState(TNotNull<UStateTreeEditorData*> InEditorData, TNotNull<UStateTreeState*> InState, UE::StateTreeEditor::FClipboardEditorData& InProcessedClipboard)
+{
+	auto AppendFunc = [&InProcessedClipboard]<typename T>(TNotNull<UStateTreeState*> InOwnerState, TNotNull<UStateTreeEditorData*> InEditorData, const FStateTreeEditPropertyPath & InPropertyPath)
+		requires TIsDerivedFrom<T, FStateTreeNodeBase>::IsDerived || std::is_same_v<T, FStateTreeTransition>
+	{
+		if constexpr (std::is_same_v<T, FStateTreeTransition>)
+		{
+			TArrayView<FStateTreeTransition> TransitionsInBuffer = InProcessedClipboard.GetTransitionsInBuffer();
+
+			for (FStateTreeTransition& Transition : TransitionsInBuffer)
+			{
+				InOwnerState->Transitions.Add(MoveTemp(Transition));
+			}
+		}
+		else
+		{
+			const UScriptStruct* NodeType = TBaseStructure<T>::Get();
+
+			TArray<FStateTreeEditorNode>* TargetArray = nullptr;
+			if (NodeType->IsChildOf<FStateTreeTaskBase>())
+			{
+				TargetArray = &InOwnerState->Tasks;
+			}
+			else if (NodeType->IsChildOf<FStateTreeConditionBase>())
+			{
+				TargetArray = &InOwnerState->EnterConditions;
+			}
+			else if (NodeType->IsChildOf<FStateTreeConsiderationBase>())
+			{
+				TargetArray = &InOwnerState->Considerations;
+			}
+
+			TArrayView<FStateTreeEditorNode> EditorNodesInBuffer = InProcessedClipboard.GetEditorNodesInBuffer();
+
+			for (FStateTreeEditorNode& EditorNode : EditorNodesInBuffer)
+			{
+				TStructView<FStateTreeNodeBase> NodeView = EditorNode.GetNode();
+				if (NodeView.IsValid() && NodeView.GetScriptStruct()->IsChildOf(NodeType))
+				{
+					TargetArray->Add(MoveTemp(EditorNode));
+				}
+			}
+		}
+	};
+
+	FScopedTransaction ScopedTransaction(LOCTEXT("PasteNodesToStateTransaction", "Paste Nodes To State"));
+
+	// Property Change requires one property at a time
+	using namespace UE::StateTree::PropertyHelpers;
+
+	// Enter Conditions
+	ModifyStateInPreAndPostEdit(
+		LOCTEXT("PasteNodesToStateTransaction", "Paste Nodes To State"),
+		InState,
+		InEditorData,
+		TEXT("EnterConditions"),
+		[&AppendFunc](TNotNull<UStateTreeState*> InOwnerState, TNotNull<UStateTreeEditorData*> InEditorData, const FStateTreeEditPropertyPath& InPropertyPath)
+		{ AppendFunc.operator() < FStateTreeConditionBase > (InOwnerState, InEditorData, InPropertyPath); },
+		INDEX_NONE,
+		EPropertyChangeType::ValueSet);
+
+	// Considerations
+	ModifyStateInPreAndPostEdit(
+		LOCTEXT("PasteNodesToStateTransaction", "Paste Nodes To State"),
+		InState,
+		InEditorData,
+		TEXT("Considerations"),
+		[&AppendFunc](TNotNull<UStateTreeState*> InOwnerState, TNotNull<UStateTreeEditorData*> InEditorData, const FStateTreeEditPropertyPath& InPropertyPath)
+		{ AppendFunc.operator() < FStateTreeConsiderationBase > (InOwnerState, InEditorData, InPropertyPath); },
+		INDEX_NONE,
+		EPropertyChangeType::ValueSet);
+
+	// Tasks
+	ModifyStateInPreAndPostEdit(
+		LOCTEXT("PasteNodesToStateTransaction", "Paste Nodes To State"),
+		InState,
+		InEditorData,
+		TEXT("Tasks"),
+		[&AppendFunc](TNotNull<UStateTreeState*> InOwnerState, TNotNull<UStateTreeEditorData*> InEditorData, const FStateTreeEditPropertyPath& InPropertyPath)
+		{ AppendFunc.operator() < FStateTreeTaskBase > (InOwnerState, InEditorData, InPropertyPath); },
+		INDEX_NONE,
+		EPropertyChangeType::ValueSet);
+
+	// Transitions
+	ModifyStateInPreAndPostEdit(
+		LOCTEXT("PasteNodesToStateTransaction", "Paste Nodes To State"),
+		InState,
+		InEditorData,
+		TEXT("Transitions"),
+		[&AppendFunc](TNotNull<UStateTreeState*> InOwnerState, TNotNull<UStateTreeEditorData*> InEditorData, const FStateTreeEditPropertyPath& InPropertyPath)
+		{ AppendFunc.operator() < FStateTreeTransition > (InOwnerState, InEditorData, InPropertyPath); },
+		INDEX_NONE,
+		EPropertyChangeType::ValueSet);
+
+	// Dump fixed property bindings in the end to avoid being cleaned up before their corresponding nodes are pushed in
+	for (FStateTreePropertyPathBinding& Binding : InProcessedClipboard.GetBindingsInBuffer())
+	{
+		InEditorData->GetPropertyEditorBindings()->AddStateTreeBinding(MoveTemp(Binding));
+	}
+
+	OnStateNodesChanged.Broadcast(InState);
+}
+
 bool FStateTreeViewModel::IsStateActiveInDebugger(const UStateTreeState& State) const
 {
 #if WITH_STATETREE_TRACE_DEBUGGER
diff --git a/Engine/Plugins/Runtime/StateTree/Source/StateTreeEditorModule/Public/StateTreeViewModel.h b/Engine/Plugins/Runtime/StateTree/Source/StateTreeEditorModule/Public/StateTreeViewModel.h
--- a/Engine/Plugins/Runtime/StateTree/Source/StateTreeEditorModule/Public/StateTreeViewModel.h
+++ b/Engine/Plugins/Runtime/StateTree/Source/StateTreeEditorModule/Public/StateTreeViewModel.h
@@ -4,12 +4,16 @@
 
 #include "UObject/WeakObjectPtr.h"
 #include "EditorUndoClient.h"
-#include "PropertyHandle.h"
-#include "StateTreeEditorNode.h"
+#include "IStateTreeEditorHost.h"
 #include "StateTreeViewModel.generated.h"
 
 #define UE_API STATETREEEDITORMODULE_API
 
+namespace UE::StateTreeEditor
+{
+	struct FClipboardEditorData;
+}
+
 struct FStateTreeTransition;
 class FMenuBuilder;
 class UStateTree;
@@ -111,14 +115,17 @@ public:
 	UE_API void MoveSelectedStatesInto(UStateTreeState* TargetState);
 	UE_API bool CanEnableStates() const;
 	UE_API bool CanDisableStates() const;
+	UE_API bool CanPasteNodesToSelectedStates() const;
 	UE_API void SetSelectedStatesEnabled(bool bEnable);
 
 	// EditorNode and Transition manipulation
-	// @todo: support Add, ReplaceWith and Rename
+	// @todo: support ReplaceWith and Rename
 	UE_API void DeleteNode(TWeakObjectPtr<UStateTreeState> State, const FGuid& ID);
 	UE_API void DeleteAllNodes(TWeakObjectPtr<UStateTreeState> State, const FGuid& ID);
 	UE_API void CopyNode(TWeakObjectPtr<UStateTreeState> State, const FGuid& ID);
+	UE_API void CopyAllNodes(TWeakObjectPtr<UStateTreeState> State, const FGuid& ID);
 	UE_API void PasteNode(TWeakObjectPtr<UStateTreeState> State, const FGuid& ID);
+	UE_API void PasteNodesToSelectedStates();
 	UE_API void DuplicateNode(TWeakObjectPtr<UStateTreeState> State, const FGuid& ID);
 
 	// Force to update the view externally.
@@ -207,7 +214,9 @@ protected:
 	UE_API void HandleIdentifierChanged(const UStateTree& StateTree) const;
 	
 	UE_API void BindToDebuggerDelegates();
-	
+
+	UE_API void PasteNodesToState(TNotNull<UStateTreeEditorData*> InEditorData, TNotNull<UStateTreeState*> InState, UE::StateTreeEditor::FClipboardEditorData& InProcessedClipboard);
+
 	TWeakObjectPtr<UStateTreeEditorData> TreeDataWeak;
 	TSet<TWeakObjectPtr<UStateTreeState>> SelectedStates;
 
diff --git a/Engine/Source/Runtime/CoreUObject/Public/StructUtils/StructView.h b/Engine/Source/Runtime/CoreUObject/Public/StructUtils/StructView.h
--- a/Engine/Source/Runtime/CoreUObject/Public/StructUtils/StructView.h
+++ b/Engine/Source/Runtime/CoreUObject/Public/StructUtils/StructView.h
@@ -139,7 +139,7 @@ public:
 	explicit TStructView(const UScriptStruct* InScriptStruct, uint8* InStructMemory = nullptr)
 		: FStructView(InScriptStruct, InStructMemory)
 	{
-		checkf(InScriptStruct->IsChildOf(TBaseStructure<T>::Get()), TEXT("ScriptStruct must be a child of T!"));
+		checkf(!InScriptStruct || InScriptStruct->IsChildOf(TBaseStructure<T>::Get()), TEXT("Struct type must be a child of T!"));
 	}
 
 	template<typename T = BaseStructT>
