[State Tree] fixed localization conflict

diff --git a/Engine/Plugins/Runtime/StateTree/Source/StateTreeEditorModule/Private/Customizations/StateTreeEditorNodeDetails.cpp b/Engine/Plugins/Runtime/StateTree/Source/StateTreeEditorModule/Private/Customizations/StateTreeEditorNodeDetails.cpp
--- a/Engine/Plugins/Runtime/StateTree/Source/StateTreeEditorModule/Private/Customizations/StateTreeEditorNodeDetails.cpp
+++ b/Engine/Plugins/Runtime/StateTree/Source/StateTreeEditorModule/Private/Customizations/StateTreeEditorNodeDetails.cpp
@@ -2116,7 +2116,7 @@ TSharedRef<SWidget> FStateTreeEditorNodeDetails::GenerateOptionsMenu()
 	// Copy all
 	MenuBuilder.AddMenuEntry(
 		LOCTEXT("CopyAllItems", "Copy all"),
-		LOCTEXT("CopyItemTooltip", "Copy all items"),
+		LOCTEXT("CopyAllItemsTooltip", "Copy all items"),
 		FSlateIcon(FAppStyle::GetAppStyleSetName(), "GenericCommands.Copy"),
 		FUIAction(
 			FExecuteAction::CreateSP(this, &FStateTreeEditorNodeDetails::OnCopyAllNodes),
