[State Tree] Added Multiply operator to Consideration score calculation
correctly display full operand text
renamed And, Or to Min, Max for Consideration

diff --git a/Engine/Plugins/Runtime/StateTree/Source/StateTreeEditorModule/Private/Customizations/StateTreeEditorNodeDetails.cpp b/Engine/Plugins/Runtime/StateTree/Source/StateTreeEditorModule/Private/Customizations/StateTreeEditorNodeDetails.cpp
--- a/Engine/Plugins/Runtime/StateTree/Source/StateTreeEditorModule/Private/Customizations/StateTreeEditorNodeDetails.cpp
+++ b/Engine/Plugins/Runtime/StateTree/Source/StateTreeEditorModule/Private/Customizations/StateTreeEditorNodeDetails.cpp
@@ -435,7 +435,7 @@ void FStateTreeEditorNodeDetails::CustomizeHeader(TSharedRef<class IPropertyHand
 				.VAlign(VAlign_Center)
 				[
 					SNew(SBox)
-					.WidthOverride(30.f)
+					.WidthOverride(30.0f)
 					.Visibility(this, &FStateTreeEditorNodeDetails::AreIndentButtonsVisible)
 					[
 						SNew(SButton)
@@ -482,7 +482,7 @@ void FStateTreeEditorNodeDetails::CustomizeHeader(TSharedRef<class IPropertyHand
 				.VAlign(VAlign_Center)
 				[
 					SNew(SBox)
-					.WidthOverride(30.0f)
+					.MinDesiredWidth(30.0f)
 					.Padding(FMargin(2.0f, 4.0f, 2.0f, 3.0f))
 					.VAlign(VAlign_Center)
 					.Visibility(this, &FStateTreeEditorNodeDetails::IsOperandVisible)
@@ -1441,38 +1441,71 @@ FText FStateTreeEditorNodeDetails::GetOperandText() const
 {
 	check(OperandProperty);
 
+	if (IsConditionVisible() == EVisibility::Visible)
+	{
+		return GetConditionOperandText();
+	}
+	else if (IsConsiderationVisible() == EVisibility::Visible)
+	{
+		return GetConsiderationOperandText();
+	}
+
+	return FText::GetEmpty();
+}
+
+FText FStateTreeEditorNodeDetails::GetConditionOperandText() const
+{
+	check(OperandProperty);
+
 	// First item does not relate to anything existing, it could be empty. 
 	// return IF to indicate that we're building condition and IS for consideration.
 	if (IsFirstItem())
 	{
-		if (IsConditionVisible() == EVisibility::Visible)
-		{
-			return LOCTEXT("IfOperand", "IF");
-		}
-		else //IsConsiderationVisible() == EVisibility::Visible
-		{
-			return LOCTEXT("IsOperand", "IS");
-		}
+		return LOCTEXT("IfOperand", "IF");
 	}
 
 	uint8 Value = 0;
 	OperandProperty->GetValue(Value);
-	const EStateTreeExpressionOperand Operand = (EStateTreeExpressionOperand)Value;
 
-	if (Operand == EStateTreeExpressionOperand::And)
+	switch (EStateTreeExpressionOperand Operand = static_cast<EStateTreeExpressionOperand>(Value))
 	{
+	case EStateTreeExpressionOperand::And:
 		return LOCTEXT("AndOperand", "AND");
-	}
-	else if (Operand == EStateTreeExpressionOperand::Or)
-	{
+	case EStateTreeExpressionOperand::Or:
 		return LOCTEXT("OrOperand", "OR");
-	}
-	else
-	{
+	case EStateTreeExpressionOperand::Multiply:
+	default:
 		ensureMsgf(false, TEXT("Unhandled operand %s"), *UEnum::GetValueAsString(Operand));
+		return FText::GetEmpty();
+	};
+}
+
+FText FStateTreeEditorNodeDetails::GetConsiderationOperandText() const
+{
+	check(OperandProperty);
+
+	// First item does not relate to anything existing, it could be empty. 
+	// return IF to indicate that we're building condition and IS for consideration.
+	if (IsFirstItem())
+	{
+		return LOCTEXT("IsOperand", "IS");
 	}
 
-	return FText::GetEmpty();
+	uint8 Value = 0;
+	OperandProperty->GetValue(Value);
+
+	switch (EStateTreeExpressionOperand Operand = static_cast<EStateTreeExpressionOperand>(Value))
+	{
+	case EStateTreeExpressionOperand::And:
+		return LOCTEXT("MinOperand", "MIN");
+	case EStateTreeExpressionOperand::Or:
+		return LOCTEXT("MaxOperand", "MAX");
+	case EStateTreeExpressionOperand::Multiply:
+		return LOCTEXT("MultiplyOperand", "MULTIPLY");
+	default:
+		ensureMsgf(false, TEXT("Unhandled operand %s"), *UEnum::GetValueAsString(Operand));
+		return FText::GetEmpty();
+	};
 }
 
 FSlateColor FStateTreeEditorNodeDetails::GetOperandColor() const
@@ -1486,25 +1519,34 @@ FSlateColor FStateTreeEditorNodeDetails::GetOperandColor() const
 
 	uint8 Value = 0; 
 	OperandProperty->GetValue(Value);
-	const EStateTreeExpressionOperand Operand = (EStateTreeExpressionOperand)Value;
 
-	if (Operand == EStateTreeExpressionOperand::And)
+	switch (EStateTreeExpressionOperand Operand = static_cast<EStateTreeExpressionOperand>(Value))
 	{
+	case EStateTreeExpressionOperand::And:
 		return FStyleColors::AccentPink;
-	}
-	else if (Operand == EStateTreeExpressionOperand::Or)
-	{
+	case EStateTreeExpressionOperand::Or:
 		return FStyleColors::AccentBlue;
+	case EStateTreeExpressionOperand::Multiply:
+		return FStyleColors::AccentGreen;
+	default:
+		ensureMsgf(false, TEXT("Unhandled operand %s"), *UEnum::GetValueAsString(Operand));
+		return FStyleColors::Transparent;
+	};
+}
+
+TSharedRef<SWidget> FStateTreeEditorNodeDetails::OnGetOperandContent() const
+{
+	if (IsConditionVisible() == EVisibility::Visible)
+	{
+		return GetConditionOperandContent();
 	}
-	else
+	else //(IsConsiderationVisible() == EVisibility::Visible)
 	{
-		ensureMsgf(false, TEXT("Unhandled operand %s"), *UEnum::GetValueAsString(Operand));
+		return GetConsiderationOperandContent();
 	}
-
-	return FStyleColors::Transparent;
 }
 
-TSharedRef<SWidget> FStateTreeEditorNodeDetails::OnGetOperandContent() const
+TSharedRef<SWidget> FStateTreeEditorNodeDetails::GetConditionOperandContent() const
 {
 	FMenuBuilder MenuBuilder(true, NULL);
 
@@ -1522,6 +1564,30 @@ TSharedRef<SWidget> FStateTreeEditorNodeDetails::OnGetOperandContent() const
 	return MenuBuilder.MakeWidget();
 }
 
+TSharedRef<SWidget> FStateTreeEditorNodeDetails::GetConsiderationOperandContent() const
+{
+	FMenuBuilder MenuBuilder(true, NULL);
+
+	FUIAction MinAction(
+		FExecuteAction::CreateSP(this, &FStateTreeEditorNodeDetails::SetOperand, EStateTreeExpressionOperand::And),
+		FCanExecuteAction(),
+		FIsActionChecked::CreateSP(this, &FStateTreeEditorNodeDetails::IsOperand, EStateTreeExpressionOperand::And));
+	MenuBuilder.AddMenuEntry(LOCTEXT("MinOperand", "MIN"), TAttribute<FText>(), FSlateIcon(), MinAction, FName(), EUserInterfaceActionType::Check);
+
+	FUIAction MaxAction(FExecuteAction::CreateSP(this, &FStateTreeEditorNodeDetails::SetOperand, EStateTreeExpressionOperand::Or),
+		FCanExecuteAction(),
+		FIsActionChecked::CreateSP(this, &FStateTreeEditorNodeDetails::IsOperand, EStateTreeExpressionOperand::Or));
+	MenuBuilder.AddMenuEntry(LOCTEXT("MaxOperand", "MAX"), TAttribute<FText>(), FSlateIcon(), MaxAction, FName(), EUserInterfaceActionType::Check);
+
+	FUIAction MultiplyAction(
+		FExecuteAction::CreateSP(this, &FStateTreeEditorNodeDetails::SetOperand, EStateTreeExpressionOperand::Multiply),
+		FCanExecuteAction(),
+		FIsActionChecked::CreateSP(this, &FStateTreeEditorNodeDetails::IsOperand, EStateTreeExpressionOperand::Multiply));
+	MenuBuilder.AddMenuEntry(LOCTEXT("MultiplyOperand", "MULTIPLY"), TAttribute<FText>(), FSlateIcon(), MultiplyAction, FName(), EUserInterfaceActionType::Check);
+
+	return MenuBuilder.MakeWidget();
+}
+
 bool FStateTreeEditorNodeDetails::IsOperandEnabled() const
 {
 	return !IsFirstItem();
@@ -1533,7 +1599,7 @@ bool FStateTreeEditorNodeDetails::IsOperand(const EStateTreeExpressionOperand Op
 
 	uint8 Value = 0; 
 	OperandProperty->GetValue(Value);
-	const EStateTreeExpressionOperand CurrOperand = (EStateTreeExpressionOperand)Value;
+	const EStateTreeExpressionOperand CurrOperand = static_cast<EStateTreeExpressionOperand>(Value);
 
 	return CurrOperand == Operand;
 }
@@ -1542,7 +1608,7 @@ void FStateTreeEditorNodeDetails::SetOperand(const EStateTreeExpressionOperand O
 {
 	check(OperandProperty);
 
-	OperandProperty->SetValue((uint8)Operand);
+	OperandProperty->SetValue(static_cast<uint8>(Operand));
 }
 
 EVisibility FStateTreeEditorNodeDetails::IsConditionVisible() const
diff --git a/Engine/Plugins/Runtime/StateTree/Source/StateTreeEditorModule/Private/Customizations/StateTreeEditorNodeDetails.h b/Engine/Plugins/Runtime/StateTree/Source/StateTreeEditorModule/Private/Customizations/StateTreeEditorNodeDetails.h
--- a/Engine/Plugins/Runtime/StateTree/Source/StateTreeEditorModule/Private/Customizations/StateTreeEditorNodeDetails.h
+++ b/Engine/Plugins/Runtime/StateTree/Source/StateTreeEditorModule/Private/Customizations/StateTreeEditorNodeDetails.h
@@ -67,8 +67,12 @@ private:
 	FSlateColor GetContentRowColor() const;
 
 	FText GetOperandText() const;
+	FText GetConditionOperandText() const;
+	FText GetConsiderationOperandText() const;
 	FSlateColor GetOperandColor() const;
 	TSharedRef<SWidget> OnGetOperandContent() const;
+	TSharedRef<SWidget> GetConditionOperandContent() const;
+	TSharedRef<SWidget> GetConsiderationOperandContent() const;
 	bool IsOperandEnabled() const;
 
 	void SetOperand(const EStateTreeExpressionOperand Operand) const;
diff --git a/Engine/Plugins/Runtime/StateTree/Source/StateTreeModule/Private/StateTreeExecutionContext.cpp b/Engine/Plugins/Runtime/StateTree/Source/StateTreeModule/Private/StateTreeExecutionContext.cpp
--- a/Engine/Plugins/Runtime/StateTree/Source/StateTreeModule/Private/StateTreeExecutionContext.cpp
+++ b/Engine/Plugins/Runtime/StateTree/Source/StateTreeModule/Private/StateTreeExecutionContext.cpp
@@ -4806,6 +4806,10 @@ bool FStateTreeExecutionContext::TestAllConditionsInternal(const FStateTreeExecu
 			case EStateTreeExpressionOperand::Or:
 				Values[Level] |= Values[Level + 1];
 				break;
+			case EStateTreeExpressionOperand::Multiply:
+			default:
+				checkf(false, TEXT("Unhandled operand %s"), *UEnum::GetValueAsString(Operand));
+				break;
 			}
 			Operands[Level] = EStateTreeExpressionOperand::Copy;
 		}
@@ -4917,6 +4921,12 @@ float FStateTreeExecutionContext::EvaluateUtilityWithValidation(const FStateTree
 			case EStateTreeExpressionOperand::Or:
 				Values[Level] = FMath::Max(Values[Level], Values[Level + 1]);
 				break;
+			case EStateTreeExpressionOperand::Multiply:
+				Values[Level] = Values[Level] * Values[Level + 1];
+				break;
+			default:
+				checkf(false, TEXT("Unhandled operand %s"), *UEnum::GetValueAsString(Operand));
+				break;
 			}
 			Operands[Level] = EStateTreeExpressionOperand::Copy;
 		}
diff --git a/Engine/Plugins/Runtime/StateTree/Source/StateTreeModule/Public/StateTreeTypes.h b/Engine/Plugins/Runtime/StateTree/Source/StateTreeModule/Public/StateTreeTypes.h
--- a/Engine/Plugins/Runtime/StateTree/Source/StateTreeModule/Public/StateTreeTypes.h
+++ b/Engine/Plugins/Runtime/StateTree/Source/StateTreeModule/Public/StateTreeTypes.h
@@ -119,11 +119,14 @@ enum class EStateTreeExpressionOperand : uint8
 	/** Copy result */
 	Copy UMETA(Hidden),
 	
-	/** Combine results with AND. */
+	/** Combine results with AND for Condition, Min(a, b) for Consideration */
 	And,
 	
-	/** Combine results with OR. */
+	/** Combine results with OR for Condition, Max(a, b) for Consideration */
 	Or,
+
+	/** Combine results with (a * b) for Consideration, not applicable for Condition */
+	Multiply,
 };
 
 UENUM()
