[State Tree] Re-renamed Min, max operator back to And, Or for consideration bc it broke the visuals in the editor: Min Max is Op(lhs, rhs) and And Or is (lhs Op rhs)

diff --git a/Engine/Plugins/Runtime/StateTree/Source/StateTreeEditorModule/Private/Customizations/StateTreeEditorNodeDetails.cpp b/Engine/Plugins/Runtime/StateTree/Source/StateTreeEditorModule/Private/Customizations/StateTreeEditorNodeDetails.cpp
--- a/Engine/Plugins/Runtime/StateTree/Source/StateTreeEditorModule/Private/Customizations/StateTreeEditorNodeDetails.cpp
+++ b/Engine/Plugins/Runtime/StateTree/Source/StateTreeEditorModule/Private/Customizations/StateTreeEditorNodeDetails.cpp
@@ -1497,9 +1497,9 @@ FText FStateTreeEditorNodeDetails::GetConsiderationOperandText() const
 	switch (EStateTreeExpressionOperand Operand = static_cast<EStateTreeExpressionOperand>(Value))
 	{
 	case EStateTreeExpressionOperand::And:
-		return LOCTEXT("MinOperand", "MIN");
+		return LOCTEXT("AndOperand", "AND");
 	case EStateTreeExpressionOperand::Or:
-		return LOCTEXT("MaxOperand", "MAX");
+		return LOCTEXT("OrOperand", "OR");
 	case EStateTreeExpressionOperand::Multiply:
 		return LOCTEXT("MultiplyOperand", "MULTIPLY");
 	default:
@@ -1568,16 +1568,16 @@ TSharedRef<SWidget> FStateTreeEditorNodeDetails::GetConsiderationOperandContent(
 {
 	FMenuBuilder MenuBuilder(true, NULL);
 
-	FUIAction MinAction(
+	FUIAction AndAction(
 		FExecuteAction::CreateSP(this, &FStateTreeEditorNodeDetails::SetOperand, EStateTreeExpressionOperand::And),
 		FCanExecuteAction(),
 		FIsActionChecked::CreateSP(this, &FStateTreeEditorNodeDetails::IsOperand, EStateTreeExpressionOperand::And));
-	MenuBuilder.AddMenuEntry(LOCTEXT("MinOperand", "MIN"), TAttribute<FText>(), FSlateIcon(), MinAction, FName(), EUserInterfaceActionType::Check);
+	MenuBuilder.AddMenuEntry(LOCTEXT("AndOperand", "AND"), TAttribute<FText>(), FSlateIcon(), AndAction, FName(), EUserInterfaceActionType::Check);
 
-	FUIAction MaxAction(FExecuteAction::CreateSP(this, &FStateTreeEditorNodeDetails::SetOperand, EStateTreeExpressionOperand::Or),
+	FUIAction OrAction(FExecuteAction::CreateSP(this, &FStateTreeEditorNodeDetails::SetOperand, EStateTreeExpressionOperand::Or),
 		FCanExecuteAction(),
 		FIsActionChecked::CreateSP(this, &FStateTreeEditorNodeDetails::IsOperand, EStateTreeExpressionOperand::Or));
-	MenuBuilder.AddMenuEntry(LOCTEXT("MaxOperand", "MAX"), TAttribute<FText>(), FSlateIcon(), MaxAction, FName(), EUserInterfaceActionType::Check);
+	MenuBuilder.AddMenuEntry(LOCTEXT("OrOperand", "OR"), TAttribute<FText>(), FSlateIcon(), OrAction, FName(), EUserInterfaceActionType::Check);
 
 	FUIAction MultiplyAction(
 		FExecuteAction::CreateSP(this, &FStateTreeEditorNodeDetails::SetOperand, EStateTreeExpressionOperand::Multiply),
