[StateTree] NodeCache uses the struct path instead of the short name. That prevents nameclashing between modules.

diff --git a/Engine/Plugins/Runtime/StateTree/Source/StateTreeEditorModule/Private/StateTreeNodeClassCache.cpp b/Engine/Plugins/Runtime/StateTree/Source/StateTreeEditorModule/Private/StateTreeNodeClassCache.cpp
--- a/Engine/Plugins/Runtime/StateTree/Source/StateTreeEditorModule/Private/StateTreeNodeClassCache.cpp
+++ b/Engine/Plugins/Runtime/StateTree/Source/StateTreeEditorModule/Private/StateTreeNodeClassCache.cpp
@@ -19,82 +19,80 @@
 
 #define LOCTEXT_NAMESPACE "StateTreeEditor"
 
-FStateTreeNodeClassData::FStateTreeNodeClassData(UStruct* InStruct) :
-	Struct(InStruct)
+FStateTreeNodeClassData::FStateTreeNodeClassData(TNotNull<UStruct*> InStruct)
+{
+	CachedStruct = InStruct;
+	StructAssetPath = InStruct->GetStructPathName();
+}
+
+// deprecated
+FStateTreeNodeClassData::FStateTreeNodeClassData(const FString& InClassAssetName, const FString& InClassPackageName, const FName InStructName, UStruct* InStruct)
+	: StructAssetPath(*InClassPackageName, InStructName)
+	, CachedStruct(InStruct)
 {
-	if (InStruct)
-	{
-		StructName = InStruct->GetFName();
-	}
 }
 
-FStateTreeNodeClassData::FStateTreeNodeClassData(const FString& InClassAssetName, const FString& InClassPackageName, const FName InStructName, UStruct* InStruct) :
-	Struct(InStruct),
-	StructName(InStructName),
-	ClassAssetName(InClassAssetName),
-	ClassPackageName(InClassPackageName)
+FStateTreeNodeClassData::FStateTreeNodeClassData(const FTopLevelAssetPath& AssetPath)
+	: StructAssetPath(AssetPath)
 {
 }
 
+
+
 UStruct* FStateTreeNodeClassData::GetStruct(bool bSilent)
 {
-	UStruct* Ret = Struct.Get();
+	UStruct* Result = CachedStruct.Get();
 	
-	if (Ret == nullptr && ClassPackageName.Len())
+	if (Result == nullptr && StructAssetPath.IsValid())
 	{
 		GWarn->BeginSlowTask(LOCTEXT("LoadPackage", "Loading Package..."), true);
 
-		UPackage* Package = LoadPackage(nullptr, *ClassPackageName, LOAD_NoRedirects);
-		if (Package)
+		UObject* Object = StaticLoadAsset(nullptr, StructAssetPath);
+		if (Object)
 		{
-			Package->FullyLoad();
-
-			UObject* Object = FindObject<UObject>(Package, *ClassAssetName);
-
-			GWarn->EndSlowTask();
-
 			if (UBlueprint* BlueprintOb = Cast<UBlueprint>(Object))
 			{
-				Ret = *BlueprintOb->GeneratedClass;
+				Result = *BlueprintOb->GeneratedClass;
 			}
 			else if (Object != nullptr)
 			{
-				Ret = Object->GetClass();
+				Result = Object->GetClass();
 			}
 
-			Struct = Ret;
+			CachedStruct = Result;
 		}
 		else
 		{
-			GWarn->EndSlowTask();
-
 			if (!bSilent)
 			{
 				FMessageLog EditorErrors("EditorErrors");
 				EditorErrors.Error(LOCTEXT("PackageLoadFail", "Package Load Failed"));
-				EditorErrors.Info(FText::FromString(ClassPackageName));
+				EditorErrors.Info(FText::FromName(StructAssetPath.GetPackageName()));
 				EditorErrors.Notify(LOCTEXT("PackageLoadFail", "Package Load Failed"));
 			}
 		}
+
+		GWarn->EndSlowTask();
 	}
 
-	return Ret;
+	return Result;
 }
 
 const UStruct* FStateTreeNodeClassData::GetInstanceDataStruct(bool bSilent /*= false*/)
 {
-	if (!InstanceDataStruct.IsValid())
+	const UStruct* Result = CachedInstanceDataStruct.Get();
+	if (Result == nullptr && StructAssetPath.IsValid())
 	{
 		if (const UScriptStruct* ScriptStruct = GetScriptStruct(bSilent))
 		{
 			TInstancedStruct<FStateTreeNodeBase> NodeInstance;
 			NodeInstance.InitializeAsScriptStruct(ScriptStruct);
 
-			InstanceDataStruct = NodeInstance.Get().GetInstanceDataType();
+			CachedInstanceDataStruct = NodeInstance.Get().GetInstanceDataType();
 		}
 	}
 
-	return InstanceDataStruct.Get();
+	return Result;
 }
 
 //----------------------------------------------------------------------//
@@ -146,9 +144,12 @@ FStateTreeNodeClassCache::~FStateTreeNodeClassCache()
 	}
 }
 
-void FStateTreeNodeClassCache::AddRootStruct(UStruct* RootStruct)
+void FStateTreeNodeClassCache::AddRootStruct(TNotNull<UStruct*> RootStruct)
 {
-	if (RootClasses.ContainsByPredicate([RootStruct](const FRootClassContainer& RootClass){ return RootClass.BaseStruct == RootStruct; }))
+	if (RootClasses.ContainsByPredicate([RootStruct](const FRootClassContainer& RootClass)
+		{
+			return RootClass.BaseStruct == RootStruct;
+		}))
 	{
 		return;
 	}
@@ -158,11 +159,14 @@ void FStateTreeNodeClassCache::AddRootStruct(UStruct* RootStruct)
 	InvalidateCache();
 }
 
-void FStateTreeNodeClassCache::GetStructs(const UStruct* BaseStruct, TArray<TSharedPtr<FStateTreeNodeClassData>>& AvailableClasses)
+void FStateTreeNodeClassCache::GetStructs(TNotNull<const UStruct*> BaseStruct, TArray<TSharedPtr<FStateTreeNodeClassData>>& AvailableClasses)
 {
 	AvailableClasses.Reset();
 	
-	const int32 RootClassIndex = RootClasses.IndexOfByPredicate([BaseStruct](const FRootClassContainer& RootClass){ return RootClass.BaseStruct == BaseStruct; });
+	const int32 RootClassIndex = RootClasses.IndexOfByPredicate([BaseStruct](const FRootClassContainer& RootClass)
+		{
+			return RootClass.BaseStruct == BaseStruct;
+		});
 	if (RootClassIndex != INDEX_NONE)
 	{
 		const FRootClassContainer& RootClass = RootClasses[RootClassIndex];
@@ -185,17 +189,16 @@ void FStateTreeNodeClassCache::OnAssetRemoved(const FAssetData& AssetData)
 {
 	FString AssetClassName;
 	FString AssetNativeParentClassName;
-
 	if (AssetData.GetTagValue(FBlueprintTags::GeneratedClassPath, AssetClassName) && AssetData.GetTagValue(FBlueprintTags::NativeParentClassPath, AssetNativeParentClassName))
 	{
-		ConstructorHelpers::StripObjectClass(AssetClassName);
-		AssetClassName = FPackageName::ObjectPathToObjectName(AssetClassName);
-
 		if (const int32* RootClassIndex = ClassNameToRootIndex.Find(AssetNativeParentClassName))
 		{
 			FRootClassContainer& RootClass = RootClasses[*RootClassIndex];
-			const FName ClassName(AssetClassName);
-			RootClass.ClassData.RemoveAll([&ClassName](const TSharedPtr<FStateTreeNodeClassData>& ClassData) { return ClassData->GetStructName() == ClassName; });
+			const FTopLevelAssetPath StructAssetPath = FTopLevelAssetPath(AssetClassName);
+			RootClass.ClassData.RemoveAll([&StructAssetPath](const TSharedPtr<FStateTreeNodeClassData>& ClassData)
+			{
+				return ClassData->GetStructPath() == StructAssetPath;
+			});
 		}
 	}
 }
@@ -219,30 +222,37 @@ void FStateTreeNodeClassCache::UpdateBlueprintClass(const FAssetData& AssetData)
 {
 	FString AssetClassName;
 	FString AssetNativeParentClassName;
-
 	if (AssetData.GetTagValue(FBlueprintTags::GeneratedClassPath, AssetClassName) && AssetData.GetTagValue(FBlueprintTags::NativeParentClassPath, AssetNativeParentClassName))
 	{
-		UObject* AssetClassPackage = nullptr;
-		ResolveName(AssetClassPackage, AssetClassName, false, false);
-
-		UObject* AssetNativeParentClassPackage = nullptr;
-		ResolveName(AssetNativeParentClassPackage, AssetNativeParentClassName, false, false);
-
+		ConstructorHelpers::StripObjectClass(AssetNativeParentClassName);
 		if (const int32* RootClassIndex = ClassNameToRootIndex.Find(AssetNativeParentClassName))
 		{
-			FRootClassContainer& RootClass = RootClasses[*RootClassIndex];
-
-			const FName ClassName(AssetClassName);
-			const int32 ClassDataIndex = RootClass.ClassData.IndexOfByPredicate([&ClassName](const TSharedPtr<FStateTreeNodeClassData>& ClassData) { return ClassData->GetStructName() == ClassName; });
-			
-			if (ClassDataIndex == INDEX_NONE)
+			const FTopLevelAssetPath StructAssetPath = FTopLevelAssetPath(AssetClassName);
+			if (StructAssetPath.IsValid())
 			{
-				UObject* AssetOb = AssetData.IsAssetLoaded() ? AssetData.GetAsset() : nullptr;
-				UBlueprint* AssetBP = Cast<UBlueprint>(AssetOb);
-				UClass* AssetClass = AssetBP ? *AssetBP->GeneratedClass : AssetOb ? AssetOb->GetClass() : nullptr;
+				FRootClassContainer& RootClass = RootClasses[*RootClassIndex];
+				const int32 ClassDataIndex = RootClass.ClassData.IndexOfByPredicate([&StructAssetPath](const TSharedPtr<FStateTreeNodeClassData>& ClassData)
+					{
+						return ClassData->GetStructPath() == StructAssetPath;
+					});
 
-				TSharedPtr<FStateTreeNodeClassData> NewData = MakeShareable(new FStateTreeNodeClassData(AssetData.AssetName.ToString(), AssetData.PackageName.ToString(), ClassName, AssetClass));
-				RootClass.ClassData.Add(NewData);
+				if (ClassDataIndex == INDEX_NONE)
+				{
+					UObject* AssetOb = AssetData.IsAssetLoaded() ? AssetData.GetAsset() : nullptr;
+					UBlueprint* AssetBP = Cast<UBlueprint>(AssetOb);
+					UClass* AssetClass = AssetBP ? *AssetBP->GeneratedClass : AssetOb ? AssetOb->GetClass() : nullptr;
+
+					if (AssetClass)
+					{
+						TSharedRef<FStateTreeNodeClassData> NewData = MakeShared<FStateTreeNodeClassData>(AssetClass);
+						RootClass.ClassData.Add(NewData);
+					}
+					else
+					{
+						TSharedRef<FStateTreeNodeClassData> NewData = MakeShared<FStateTreeNodeClassData>(StructAssetPath);
+						RootClass.ClassData.Add(NewData);
+					}
+				}
 			}
 		}
 	}
@@ -256,7 +266,7 @@ void FStateTreeNodeClassCache::CacheClasses()
 		RootClass->bUpdated = true;
 
 		// gather all native classes
-		if (const UClass* Class = Cast<UClass>(RootClass->BaseStruct))
+		if (const UClass* Class = Cast<UClass>(RootClass->BaseStruct.Get()))
 		{
 			for (TObjectIterator<UClass> It; It; ++It)
 			{
@@ -264,11 +274,11 @@ void FStateTreeNodeClassCache::CacheClasses()
 				if (TestClass->HasAnyClassFlags(CLASS_Native) && TestClass->IsChildOf(Class))
 				{
 					RootClass->ClassData.Add(MakeShareable(new FStateTreeNodeClassData(TestClass)));
-					ClassNameToRootIndex.Add(TestClass->GetName(), RootClass.GetIndex());
+					ClassNameToRootIndex.Add(TestClass->GetPathName(), RootClass.GetIndex());
 				}
 			}
 		}
-		else if (const UScriptStruct* ScriptStruct = Cast<UScriptStruct>(RootClass->BaseStruct))
+		else if (const UScriptStruct* ScriptStruct = Cast<UScriptStruct>(RootClass->BaseStruct.Get()))
 		{
 			for (TObjectIterator<UScriptStruct> It; It; ++It)
 			{
@@ -276,7 +286,7 @@ void FStateTreeNodeClassCache::CacheClasses()
 				if (TestStruct->IsChildOf(ScriptStruct))
 				{
 					RootClass->ClassData.Add(MakeShareable(new FStateTreeNodeClassData(TestStruct)));
-					ClassNameToRootIndex.Add(TestStruct->GetName(), RootClass.GetIndex());
+					ClassNameToRootIndex.Add(TestStruct->GetPathName(), RootClass.GetIndex());
 				}
 			}
 		}
diff --git a/Engine/Plugins/Runtime/StateTree/Source/StateTreeEditorModule/Public/StateTreeNodeClassCache.h b/Engine/Plugins/Runtime/StateTree/Source/StateTreeEditorModule/Public/StateTreeNodeClassCache.h
--- a/Engine/Plugins/Runtime/StateTree/Source/StateTreeEditorModule/Public/StateTreeNodeClassCache.h
+++ b/Engine/Plugins/Runtime/StateTree/Source/StateTreeEditorModule/Public/StateTreeNodeClassCache.h
@@ -2,6 +2,7 @@
 
 #pragma once
 
+#include "UObject/TopLevelAssetPath.h"
 #include "StateTreeNodeClassCache.generated.h"
 
 #define UE_API STATETREEEDITORMODULE_API
@@ -20,11 +21,23 @@ struct FStateTreeNodeClassData
 {
 	GENERATED_BODY()
 
-	FStateTreeNodeClassData() {}
-	UE_API FStateTreeNodeClassData(UStruct* InStruct);
+	FStateTreeNodeClassData() = default;
+	UE_API FStateTreeNodeClassData(TNotNull<UStruct*> InStruct);
+	UE_API FStateTreeNodeClassData(const FTopLevelAssetPath& AssetPath);
+
+	UE_DEPRECATED(5.7, "Use the constructor with FAssetData")
 	UE_API FStateTreeNodeClassData(const FString& InClassAssetName, const FString& InClassPackage, const FName InStructName, UStruct* InStruct);
 
-	FName GetStructName() const { return StructName; }
+	UE_DEPRECATED(5.7, "Use GetStructPath")
+	FName GetStructNameA() const
+	{
+		return StructAssetPath.GetAssetName();
+	}
+
+	const FTopLevelAssetPath& GetStructPath() const
+	{
+		return StructAssetPath;
+	}
 
 	UE_API UStruct* GetStruct(bool bSilent = false);
 
@@ -40,21 +53,14 @@ struct FStateTreeNodeClassData
 	UE_API const UStruct* GetInstanceDataStruct(bool bSilent = false);
 
 private:
+	/**  struct or class asset path. */
+	FTopLevelAssetPath StructAssetPath;
 
 	/** Pointer to described struct or class. */
-	TWeakObjectPtr<UStruct> Struct;
-
-	/** Resolved name of struct or class. */
-	FName StructName;
-
-	/** Path to class if it's not loaded yet. */
-	FString ClassAssetName;
-
-	/** Package of the asset if it's not loaded yet. */
-	FString ClassPackageName;
+	TWeakObjectPtr<UStruct> CachedStruct;
 
 	/** Pointer to described node's instance data struct or class. */
-	TWeakObjectPtr<const UStruct> InstanceDataStruct;
+	TWeakObjectPtr<const UStruct> CachedInstanceDataStruct;
 };
 
 /**
@@ -67,31 +73,31 @@ struct FStateTreeNodeClassCache
 	UE_API ~FStateTreeNodeClassCache();
 
 	/** Adds a Struct to keep track of */
-	UE_API void AddRootStruct(UStruct* RootStruct);
+	UE_API void AddRootStruct(TNotNull<UStruct*> RootStruct);
 	
 	/** Adds a Class to keep track of */
-	void AddRootClass(UClass* RootClass)
+	void AddRootClass(TNotNull<UClass*> RootClass)
 	{
 		AddRootStruct(RootClass);
 	}
 	
 	/** Adds a ScriptStruct to keep track of */
-	void AddRootScriptStruct(UScriptStruct* RootStruct)
+	void AddRootScriptStruct(TNotNull<UScriptStruct*> RootStruct)
 	{
 		AddRootStruct(RootStruct);
 	}
 
 	/** Returns know derived Structs based on provided base. If the base Struct is not added as root Struct, nothing is returned. */
-	UE_API void GetStructs(const UStruct* BaseStruct, TArray<TSharedPtr<FStateTreeNodeClassData>>& AvailableClasses);
+	UE_API void GetStructs(TNotNull<const UStruct*> BaseStruct, TArray<TSharedPtr<FStateTreeNodeClassData>>& AvailableClasses);
 	
 	/** Returns know derived Classes based on provided base. If the base Class is not added as root Class, nothing is returned. */
-	void GetClasses(const UStruct* BaseClass, TArray<TSharedPtr<FStateTreeNodeClassData>>& AvailableClasses)
+	void GetClasses(TNotNull<const UStruct*> BaseClass, TArray<TSharedPtr<FStateTreeNodeClassData>>& AvailableClasses)
 	{
 		GetStructs(BaseClass, AvailableClasses);
 	}
 	
 	/** Returns know derived ScriptStructs based on provided base. If the base struct is not added as root ScriptStruct, nothing is returned. */
-	void GetScripStructs(const UScriptStruct* BaseStruct, TArray<TSharedPtr<FStateTreeNodeClassData>>& AvailableClasses)
+	void GetScripStructs(TNotNull<const UScriptStruct*> BaseStruct, TArray<TSharedPtr<FStateTreeNodeClassData>>& AvailableClasses)
 	{
 		GetStructs(BaseStruct, AvailableClasses);
 	}
@@ -111,11 +117,11 @@ private:
 	struct FRootClassContainer
 	{
 		FRootClassContainer() = default;
-		FRootClassContainer(UStruct* InBaseStruct)
+		FRootClassContainer(TNotNull<UStruct*> InBaseStruct)
+		 : BaseStruct(InBaseStruct)
 		{
-			BaseStruct = InBaseStruct;
 		} 
-		TWeakObjectPtr<UStruct> BaseStruct;
+		TWeakObjectPtr<const UStruct> BaseStruct;
 		TArray<TSharedPtr<FStateTreeNodeClassData>> ClassData;
 		bool bUpdated = false;
 	};
