[StateTree] Fix issues and review feedback from CL 45839886

diff --git a/Engine/Plugins/Runtime/StateTree/Source/StateTreeEditorModule/Private/StateTreeNodeClassCache.cpp b/Engine/Plugins/Runtime/StateTree/Source/StateTreeEditorModule/Private/StateTreeNodeClassCache.cpp
--- a/Engine/Plugins/Runtime/StateTree/Source/StateTreeEditorModule/Private/StateTreeNodeClassCache.cpp
+++ b/Engine/Plugins/Runtime/StateTree/Source/StateTreeEditorModule/Private/StateTreeNodeClassCache.cpp
@@ -20,9 +20,9 @@
 #define LOCTEXT_NAMESPACE "StateTreeEditor"
 
 FStateTreeNodeClassData::FStateTreeNodeClassData(TNotNull<UStruct*> InStruct)
+	: StructAssetPath(InStruct->GetStructPathName())
+	, CachedStruct(InStruct)
 {
-	CachedStruct = InStruct;
-	StructAssetPath = InStruct->GetStructPathName();
 }
 
 // deprecated
@@ -54,9 +54,9 @@ UStruct* FStateTreeNodeClassData::GetStruct(bool bSilent)
 			{
 				Result = *BlueprintOb->GeneratedClass;
 			}
-			else if (Object != nullptr)
+			else if (UStruct* Struct = Cast<UStruct>(Object))
 			{
-				Result = Object->GetClass();
+				Result = Struct;
 			}
 
 			CachedStruct = Result;
@@ -88,7 +88,8 @@ const UStruct* FStateTreeNodeClassData::GetInstanceDataStruct(bool bSilent /*= f
 			TInstancedStruct<FStateTreeNodeBase> NodeInstance;
 			NodeInstance.InitializeAsScriptStruct(ScriptStruct);
 
-			CachedInstanceDataStruct = NodeInstance.Get().GetInstanceDataType();
+			Result = NodeInstance.Get().GetInstanceDataType();
+			CachedInstanceDataStruct = Result;
 		}
 	}
 
diff --git a/Engine/Plugins/Runtime/StateTree/Source/StateTreeEditorModule/Public/StateTreeNodeClassCache.h b/Engine/Plugins/Runtime/StateTree/Source/StateTreeEditorModule/Public/StateTreeNodeClassCache.h
--- a/Engine/Plugins/Runtime/StateTree/Source/StateTreeEditorModule/Public/StateTreeNodeClassCache.h
+++ b/Engine/Plugins/Runtime/StateTree/Source/StateTreeEditorModule/Public/StateTreeNodeClassCache.h
@@ -25,11 +25,11 @@ struct FStateTreeNodeClassData
 	UE_API FStateTreeNodeClassData(TNotNull<UStruct*> InStruct);
 	UE_API FStateTreeNodeClassData(const FTopLevelAssetPath& AssetPath);
 
-	UE_DEPRECATED(5.7, "Use the constructor with FAssetData")
+	UE_DEPRECATED(5.7, "Use the constructor with FTopLevelAssetPath")
 	UE_API FStateTreeNodeClassData(const FString& InClassAssetName, const FString& InClassPackage, const FName InStructName, UStruct* InStruct);
 
 	UE_DEPRECATED(5.7, "Use GetStructPath")
-	FName GetStructNameA() const
+	FName GetStructName() const
 	{
 		return StructAssetPath.GetAssetName();
 	}
@@ -53,7 +53,7 @@ struct FStateTreeNodeClassData
 	UE_API const UStruct* GetInstanceDataStruct(bool bSilent = false);
 
 private:
-	/**  struct or class asset path. */
+	/** struct or class asset path. */
 	FTopLevelAssetPath StructAssetPath;
 
 	/** Pointer to described struct or class. */
