[StateTree] Deprecate StateTreeCustomVersion. Use a custom version related to a stream instead. Data made with a custom version for feature do not merge between streams.

diff --git a/Engine/Plugins/Runtime/StateTree/Source/StateTreeEditorModule/Private/StateTreeState.cpp b/Engine/Plugins/Runtime/StateTree/Source/StateTreeEditorModule/Private/StateTreeState.cpp
--- a/Engine/Plugins/Runtime/StateTree/Source/StateTreeEditorModule/Private/StateTreeState.cpp
+++ b/Engine/Plugins/Runtime/StateTree/Source/StateTreeEditorModule/Private/StateTreeState.cpp
@@ -272,8 +272,14 @@ void UStateTreeState::PostLoad()
 	}
 
 #if WITH_EDITORONLY_DATA
-	const int32 CurrentVersion = GetLinkerCustomVersion(FStateTreeCustomVersion::GUID);
-	if (CurrentVersion < FStateTreeCustomVersion::AddedTransitionIds)
+PRAGMA_DISABLE_DEPRECATION_WARNINGS
+	const int32 CurrentStateTreeCustomVersion = GetLinkerCustomVersion(FStateTreeCustomVersion::GUID);
+	constexpr int32 AddedTransitionIdsVersion = FStateTreeCustomVersion::AddedTransitionIds;
+	constexpr int32 OverridableStateParametersVersion = FStateTreeCustomVersion::OverridableStateParameters;
+	constexpr int32 AddedCheckingParentsPrerequisitesVersion = FStateTreeCustomVersion::AddedCheckingParentsPrerequisites;
+PRAGMA_ENABLE_DEPRECATION_WARNINGS
+
+	if (CurrentStateTreeCustomVersion < AddedTransitionIdsVersion)
 	{
 		// Make guids for transitions. These need to be deterministic when upgrading because of cooking.
 		for (int32 Index = 0; Index < Transitions.Num(); Index++)
@@ -283,7 +289,7 @@ void UStateTreeState::PostLoad()
 		}
 	}
 
-	if (CurrentVersion < FStateTreeCustomVersion::OverridableStateParameters)
+	if (CurrentStateTreeCustomVersion < OverridableStateParametersVersion)
 	{
 		// In earlier versions, all parameters were overwritten.
 		if (const UPropertyBag* Bag = Parameters.Parameters.GetPropertyBagStruct())
@@ -295,7 +301,7 @@ void UStateTreeState::PostLoad()
 		}
 	}
 
-	if (CurrentVersion < FStateTreeCustomVersion::AddedCheckingParentsPrerequisites)
+	if (CurrentStateTreeCustomVersion < AddedCheckingParentsPrerequisitesVersion)
 	{
 		bCheckPrerequisitesWhenActivatingChildDirectly = false;
 	}
diff --git a/Engine/Plugins/Runtime/StateTree/Source/StateTreeModule/Private/StateTree.cpp b/Engine/Plugins/Runtime/StateTree/Source/StateTreeModule/Private/StateTree.cpp
--- a/Engine/Plugins/Runtime/StateTree/Source/StateTreeModule/Private/StateTree.cpp
+++ b/Engine/Plugins/Runtime/StateTree/Source/StateTreeModule/Private/StateTree.cpp
@@ -30,11 +30,13 @@
 
 #include UE_INLINE_GENERATED_CPP_BY_NAME(StateTree)
 
+PRAGMA_DISABLE_DEPRECATION_WARNINGS
 const FGuid FStateTreeCustomVersion::GUID(0x28E21331, 0x501F4723, 0x8110FA64, 0xEA10DA1E);
 namespace UE::StateTree::Private
 {
 FCustomVersionRegistration GRegisterStateTreeCustomVersion(FStateTreeCustomVersion::GUID, FStateTreeCustomVersion::LatestVersion, TEXT("StateTreeAsset"));
 }
+PRAGMA_ENABLE_DEPRECATION_WARNINGS
 
 bool UStateTree::IsReadyToRun() const
 {
@@ -421,12 +423,14 @@ void UStateTree::PostLoad()
 		Compile();
 	}
 #else
+PRAGMA_DISABLE_DEPRECATION_WARNINGS
 	const int32 CurrentVersion = GetLinkerCustomVersion(FStateTreeCustomVersion::GUID);
 	if (CurrentVersion < FStateTreeCustomVersion::LatestVersion)
 	{		
 		UE_LOG(LogStateTree, Error, TEXT("%s: compiled data is in older format. Please recompile the StateTree asset."), *GetPathName());
 		return;
 	}
+PRAGMA_ENABLE_DEPRECATION_WARNINGS
 #endif
 
 	if (Schema)
@@ -468,7 +472,11 @@ void UStateTree::Serialize(const FStructuredArchiveRecord Record)
 {
 	Super::Serialize(Record);
 
-	Record.GetUnderlyingArchive().UsingCustomVersion(FStateTreeCustomVersion::GUID);
+PRAGMA_DISABLE_DEPRECATION_WARNINGS
+	const FGuid StateTreeCustomVersion = FStateTreeCustomVersion::GUID;
+PRAGMA_ENABLE_DEPRECATION_WARNINGS
+
+	Record.GetUnderlyingArchive().UsingCustomVersion(StateTreeCustomVersion);
 	
 	// We need to link and rebind property bindings each time a BP is compiled,
 	// because property bindings may get invalid, and instance data potentially needs refreshed.
diff --git a/Engine/Plugins/Runtime/StateTree/Source/StateTreeModule/Private/StateTreeReference.cpp b/Engine/Plugins/Runtime/StateTree/Source/StateTreeModule/Private/StateTreeReference.cpp
--- a/Engine/Plugins/Runtime/StateTree/Source/StateTreeModule/Private/StateTreeReference.cpp
+++ b/Engine/Plugins/Runtime/StateTree/Source/StateTreeModule/Private/StateTreeReference.cpp
@@ -109,14 +109,22 @@ void FStateTreeReference::SetPropertyOverridden(const FGuid PropertyID, const bo
 
 bool FStateTreeReference::Serialize(FStructuredArchive::FSlot Slot)
 {
-	Slot.GetUnderlyingArchive().UsingCustomVersion(FStateTreeCustomVersion::GUID);
+PRAGMA_DISABLE_DEPRECATION_WARNINGS
+	const FGuid StateTreeCustomVersion = FStateTreeCustomVersion::GUID;
+PRAGMA_ENABLE_DEPRECATION_WARNINGS
+
+	Slot.GetUnderlyingArchive().UsingCustomVersion(StateTreeCustomVersion);
 	return false; // Let the default serializer handle serializing.
 }
 
 void FStateTreeReference::PostSerialize(const FArchive& Ar)
 {
-	const int32 CurrentVersion = Ar.CustomVer(FStateTreeCustomVersion::GUID);
-	if (CurrentVersion < FStateTreeCustomVersion::OverridableParameters)
+PRAGMA_DISABLE_DEPRECATION_WARNINGS
+	const int32 CurrentStateTreeCustomVersion = Ar.CustomVer(FStateTreeCustomVersion::GUID);
+	constexpr int32 OverridableParametersVersion = FStateTreeCustomVersion::OverridableParameters;
+PRAGMA_ENABLE_DEPRECATION_WARNINGS
+
+	if (CurrentStateTreeCustomVersion < OverridableParametersVersion)
 	{
 		// In earlier versions, all parameters were overwritten.
 		if (const UPropertyBag* Bag = Parameters.GetPropertyBagStruct())
diff --git a/Engine/Plugins/Runtime/StateTree/Source/StateTreeModule/Private/StateTreeTypes.cpp b/Engine/Plugins/Runtime/StateTree/Source/StateTreeModule/Private/StateTreeTypes.cpp
--- a/Engine/Plugins/Runtime/StateTree/Source/StateTreeModule/Private/StateTreeTypes.cpp
+++ b/Engine/Plugins/Runtime/StateTree/Source/StateTreeModule/Private/StateTreeTypes.cpp
@@ -120,24 +120,31 @@ FString FStateTreeStateHandle::Describe() const
 
 bool FStateTreeStateLink::Serialize(FStructuredArchive::FSlot Slot)
 {
-	Slot.GetUnderlyingArchive().UsingCustomVersion(FStateTreeCustomVersion::GUID);
+PRAGMA_DISABLE_DEPRECATION_WARNINGS
+	const FGuid StateTreeCustomVersionID = FStateTreeCustomVersion::GUID;
+PRAGMA_ENABLE_DEPRECATION_WARNINGS
+
+	Slot.GetUnderlyingArchive().UsingCustomVersion(StateTreeCustomVersionID);
 	return false; // Let the default serializer handle serializing.
 }
 
 void FStateTreeStateLink::PostSerialize(const FArchive& Ar)
 {
 #if WITH_EDITORONLY_DATA
-	const int32 CurrentVersion = Ar.CustomVer(FStateTreeCustomVersion::GUID);
-	if (CurrentVersion < FStateTreeCustomVersion::AddedExternalTransitions)
-	{
 PRAGMA_DISABLE_DEPRECATION_WARNINGS
+	const int32 CurrentStateTreeCustomVersion = Ar.CustomVer(FStateTreeCustomVersion::GUID);
+	constexpr int32 AddedExternalTransitionsVersion = FStateTreeCustomVersion::AddedExternalTransitions;
+
+	if (CurrentStateTreeCustomVersion < AddedExternalTransitionsVersion)
+	{
 		LinkType = Type_DEPRECATED;
 		if (LinkType == EStateTreeTransitionType::NotSet)
 		{
 			LinkType = EStateTreeTransitionType::None;
 		}
-PRAGMA_ENABLE_DEPRECATION_WARNINGS
+
 	}
+PRAGMA_ENABLE_DEPRECATION_WARNINGS
 #endif // WITH_EDITORONLY_DATA
 }
 
diff --git a/Engine/Plugins/Runtime/StateTree/Source/StateTreeModule/Public/StateTree.h b/Engine/Plugins/Runtime/StateTree/Source/StateTreeModule/Public/StateTree.h
--- a/Engine/Plugins/Runtime/StateTree/Source/StateTreeModule/Public/StateTree.h
+++ b/Engine/Plugins/Runtime/StateTree/Source/StateTreeModule/Public/StateTree.h
@@ -21,7 +21,9 @@ template<bool>
 struct TStateTreeStrongExecutionContext;
 
 /** Custom serialization version for StateTree Asset */
-struct FStateTreeCustomVersion
+struct
+UE_DEPRECATED(all, "Use a stream custom version. Data made with a custom version for feature do not merge between streams.")
+FStateTreeCustomVersion
 {
 	enum Type
 	{
