[StateTreeDebugger] reset asset debug ids to fix new tracing sessions when using legacy implementation

#FailedValidations Preflight

diff --git a/Engine/Plugins/Runtime/StateTree/Source/StateTreeModule/Private/Debugger/StateTreeTrace.cpp b/Engine/Plugins/Runtime/StateTree/Source/StateTreeModule/Private/Debugger/StateTreeTrace.cpp
--- a/Engine/Plugins/Runtime/StateTree/Source/StateTreeModule/Private/Debugger/StateTreeTrace.cpp
+++ b/Engine/Plugins/Runtime/StateTree/Source/StateTreeModule/Private/Debugger/StateTreeTrace.cpp
@@ -650,6 +650,7 @@ void RegisterGlobalDelegates()
 	{
 		if (TraceStatus == EStateTreeTraceStatus::TracesStarted)
 		{
+			GBufferedEvents.Reset();
 			GBufferedEvents.OnTracesStarted_GameThread();
 		}
 		// Traces are about to be stopped so allow the buffered events to react.
