[StateTree] Fix compilation for the renamed property after the asset got loaded. When we load the asset, the paths use redirectors to evaluate. It also uses redirectors at compile time if the source has a valid instance. No redirector was used if the instance is unknown at compile time (nullptr of unknown type). Introduced with CL 25869227

#FailedValidations Preflight

diff --git a/Engine/Plugins/Runtime/StateTree/Source/StateTreeEditorModule/Private/StateTreeCompiler.cpp b/Engine/Plugins/Runtime/StateTree/Source/StateTreeEditorModule/Private/StateTreeCompiler.cpp
--- a/Engine/Plugins/Runtime/StateTree/Source/StateTreeEditorModule/Private/StateTreeCompiler.cpp
+++ b/Engine/Plugins/Runtime/StateTree/Source/StateTreeEditorModule/Private/StateTreeCompiler.cpp
@@ -274,7 +274,7 @@ namespace UE::StateTree::Compiler
 	}
 
 	template<typename... T>
-	bool IsPropertyChildOfAnyStruct(const FStateTreeBindableStructDesc& Struct, FPropertyBindingPath Path)
+	bool IsPropertyChildOfAnyStruct(const FStateTreeBindableStructDesc& Struct, const FPropertyBindingPath& Path)
 	{
 		TArray<FPropertyBindingPathIndirection> Indirection;
 		const bool bResolved = Path.ResolveIndirections(Struct.Struct, Indirection);
@@ -494,7 +494,7 @@ bool FStateTreeCompiler::Compile(TNotNull<UStateTree*> InStateTree)
 		return FailCompilation();
 	}
 	StateTree->NumContextData = static_cast<uint16>(ContextDataIndex);
-	
+
 	if (!CreateStates())
 	{
 		return FailCompilation();
@@ -2605,7 +2605,7 @@ bool FStateTreeCompiler::GetAndValidateBindings(const FStateTreeBindableStructDe
 
 		// Update the source structs only if we have value for it. For some sources (e.g. context structs) we know only type, and in that case there are no instance structs.
 		const FStateTreeDataView SourceValue = IDToStructValue[SourceStructID];
-		if (SourceValue.IsValid())
+		if (SourceValue.GetStruct() != nullptr)
 		{
 			if (!Binding.GetMutableSourcePath().UpdateSegmentsFromValue(SourceValue))
 			{
