[State Tree] fixed naming in Editor Menu

diff --git a/Engine/Plugins/Runtime/StateTree/Source/StateTreeEditorModule/Private/StateTreeEditorModeToolkit.cpp b/Engine/Plugins/Runtime/StateTree/Source/StateTreeEditorModule/Private/StateTreeEditorModeToolkit.cpp
--- a/Engine/Plugins/Runtime/StateTree/Source/StateTreeEditorModule/Private/StateTreeEditorModeToolkit.cpp
+++ b/Engine/Plugins/Runtime/StateTree/Source/StateTreeEditorModule/Private/StateTreeEditorModeToolkit.cpp
@@ -256,7 +256,7 @@ void FStateTreeEditorModeToolkit::ExtendSecondaryModeToolbar(UToolMenu* ToolBar)
 
 	static const FToolMenuInsert InsertLast(NAME_None, EToolMenuInsertType::Last);
 
-	FToolMenuSection& CompileSection = ToolBar->AddSection("Compile2", FText(), InsertLast);
+	FToolMenuSection& CompileSection = ToolBar->AddSection("Compile", FText(), InsertLast);
 
 	auto GetToolkitFromAssetEditorContext = [](const UAssetEditorToolkitMenuContext* InContext) -> TSharedPtr<FStateTreeEditorModeToolkit>
 	{
