[State Tree] fixed naming typo in AsyncStateTreeDiff

diff --git a/Engine/Plugins/Runtime/StateTree/Source/StateTreeEditorModule/Private/AsyncStateTreeDiff.cpp b/Engine/Plugins/Runtime/StateTree/Source/StateTreeEditorModule/Private/AsyncStateTreeDiff.cpp
--- a/Engine/Plugins/Runtime/StateTree/Source/StateTreeEditorModule/Private/AsyncStateTreeDiff.cpp
+++ b/Engine/Plugins/Runtime/StateTree/Source/StateTreeEditorModule/Private/AsyncStateTreeDiff.cpp
@@ -176,7 +176,7 @@ static bool AreTransitionsEqual(const UStateTreeState* StateA, const UStateTreeS
 			return TransitionA == TransitionB;
 		}
 		// Not checking transitions on IDs
-		const bool bEqual = FStateTreeTransition::StaticStruct()->CompareScriptStruct(TransitionB, TransitionB, 0);
+		const bool bEqual = FStateTreeTransition::StaticStruct()->CompareScriptStruct(TransitionA, TransitionB, 0);
 		if (!bEqual)
 		{
 			return false;
