[State Tree] Added copy and paste binding option when right click a detail item row in details view
If property is bound: Copy Value and Paste Value are disabled. Copy Binding and Paste Binding are allowed.
If property is not bound: Copy Value, Paste Value and Paste Binding are allowed. Copy Binding is disabled.

- fixed ArrayIndex not set correctly when resolving indirections

diff --git a/Engine/Plugins/Runtime/PropertyBindingUtils/Source/PropertyBindingUtils/Private/PropertyBindingBindingCollection.cpp b/Engine/Plugins/Runtime/PropertyBindingUtils/Source/PropertyBindingUtils/Private/PropertyBindingBindingCollection.cpp
--- a/Engine/Plugins/Runtime/PropertyBindingUtils/Source/PropertyBindingUtils/Private/PropertyBindingBindingCollection.cpp
+++ b/Engine/Plugins/Runtime/PropertyBindingUtils/Source/PropertyBindingUtils/Private/PropertyBindingBindingCollection.cpp
@@ -131,9 +131,9 @@ void FPropertyBindingBindingCollection::CopyBindingsImplementation(const FGuid I
 	}
 }
 
-bool FPropertyBindingBindingCollection::HasBinding(const FPropertyBindingPath& InTargetPath, ESearchMode RemoveMode) const
+bool FPropertyBindingBindingCollection::HasBinding(const FPropertyBindingPath& InTargetPath, ESearchMode InSearchMode) const
 {
-	if (RemoveMode == ESearchMode::Exact)
+	if (InSearchMode == ESearchMode::Exact)
 	{
 		return HasBindingInternal([&InTargetPath](const FPropertyBindingBinding& Binding)
 			{
@@ -149,9 +149,9 @@ bool FPropertyBindingBindingCollection::HasBinding(const FPropertyBindingPath& I
 	}
 }
 
-const FPropertyBindingBinding* FPropertyBindingBindingCollection::FindBinding(const FPropertyBindingPath& InTargetPath, ESearchMode RemoveMode) const
+const FPropertyBindingBinding* FPropertyBindingBindingCollection::FindBinding(const FPropertyBindingPath& InTargetPath, ESearchMode InSearchMode) const
 {
-	if (RemoveMode == ESearchMode::Exact)
+	if (InSearchMode == ESearchMode::Exact)
 	{
 		return FindBindingInternal([&InTargetPath](const FPropertyBindingBinding& Binding)
 			{
diff --git a/Engine/Plugins/Runtime/PropertyBindingUtils/Source/PropertyBindingUtils/Private/PropertyBindingPath.cpp b/Engine/Plugins/Runtime/PropertyBindingUtils/Source/PropertyBindingUtils/Private/PropertyBindingPath.cpp
--- a/Engine/Plugins/Runtime/PropertyBindingUtils/Source/PropertyBindingUtils/Private/PropertyBindingPath.cpp
+++ b/Engine/Plugins/Runtime/PropertyBindingUtils/Source/PropertyBindingUtils/Private/PropertyBindingPath.cpp
@@ -365,7 +365,7 @@ bool FPropertyBindingPath::ResolveIndirectionsWithValue(const FPropertyBindingDa
 			Indirection.PropertyGuid = PropertyGuid;
 #endif
 			
-			ArrayIndex = 0;
+			ArrayIndex = INDEX_NONE;
 			Offset = 0;
 			Property = ArrayProperty->Inner;
 
@@ -406,8 +406,8 @@ bool FPropertyBindingPath::ResolveIndirectionsWithValue(const FPropertyBindingDa
 					OutIndirections.Reset();
 					return false;
 				}
-				ArrayIndex = FMath::Max(0, Segment->GetArrayIndex());
-				Offset = Property->GetOffset_ForInternal() + Property->GetElementSize() * ArrayIndex;
+				Offset = Property->GetOffset_ForInternal() + Property->GetElementSize() * FMath::Max(0, Segment->GetArrayIndex());
+				ArrayIndex = Segment->GetArrayIndex();
 			}
 		}
 
diff --git a/Engine/Plugins/Runtime/PropertyBindingUtils/Source/PropertyBindingUtils/Public/PropertyBindingBindableStructDescriptor.h b/Engine/Plugins/Runtime/PropertyBindingUtils/Source/PropertyBindingUtils/Public/PropertyBindingBindableStructDescriptor.h
--- a/Engine/Plugins/Runtime/PropertyBindingUtils/Source/PropertyBindingUtils/Public/PropertyBindingBindableStructDescriptor.h
+++ b/Engine/Plugins/Runtime/PropertyBindingUtils/Source/PropertyBindingUtils/Public/PropertyBindingBindableStructDescriptor.h
@@ -47,7 +47,7 @@ struct FPropertyBindingBindableStructDescriptor
 	UPROPERTY()
 	TObjectPtr<const UStruct> Struct = nullptr;
 
-	/** Name of the struct (used for debugging, logging, cosmetic). */
+	/** Display name of the struct (used for debugging, logging, cosmetic), different from Struct->GetName(). */
 	UPROPERTY()
 	FName Name;
 
diff --git a/Engine/Plugins/Runtime/PropertyBindingUtils/Source/PropertyBindingUtilsEditor/Private/PropertyBindingExtension.cpp b/Engine/Plugins/Runtime/PropertyBindingUtils/Source/PropertyBindingUtilsEditor/Private/PropertyBindingExtension.cpp
--- a/Engine/Plugins/Runtime/PropertyBindingUtils/Source/PropertyBindingUtilsEditor/Private/PropertyBindingExtension.cpp
+++ b/Engine/Plugins/Runtime/PropertyBindingUtils/Source/PropertyBindingUtilsEditor/Private/PropertyBindingExtension.cpp
@@ -441,6 +441,23 @@ bool FCachedBindingData::HasBinding(const FPropertyBindingBindingCollection::ESe
 	return EditorBindings->HasBinding(TargetPath, SearchMode);
 }
 
+const FPropertyBindingBinding* FCachedBindingData::FindBinding(FPropertyBindingBindingCollection::ESearchMode SearchMode) const
+{
+	IPropertyBindingBindingCollectionOwner* BindingOwner = WeakBindingsOwner.Get();
+	if (BindingOwner == nullptr)
+	{
+		return nullptr;
+	}
+
+	const FPropertyBindingBindingCollection* EditorBindings = BindingOwner->GetEditorPropertyBindings();
+	if (EditorBindings == nullptr)
+	{
+		return nullptr;
+	}
+
+	return EditorBindings->FindBinding(TargetPath, SearchMode);
+}
+
 void FCachedBindingData::RemoveBinding(const FPropertyBindingBindingCollection::ESearchMode RemoveMode)
 {
 	IPropertyBindingBindingCollectionOwner* BindingOwner = WeakBindingsOwner.Get();
@@ -1419,29 +1436,9 @@ void FPropertyBindingExtension::ExtendWidgetRow(FDetailWidgetRow& InWidgetRow, c
 		);
 	}
 
-	// ResetToDefault
-	{
-		InWidgetRow.CustomResetToDefault = FResetToDefaultOverride::Create(
-			MakeAttributeLambda([CachedBindingData, InPropertyHandle]()
-			{
-				return InPropertyHandle->CanResetToDefault() || CachedBindingData->HasBinding(FPropertyBindingBindingCollection::ESearchMode::Includes);
-			}),
-			FSimpleDelegate::CreateLambda([CachedBindingData, &InDetailBuilder, InPropertyHandle]()
-				{
-					if (CachedBindingData->HasBinding(FPropertyBindingBindingCollection::ESearchMode::Includes))
-					{
-						CachedBindingData->RemoveBinding(FPropertyBindingBindingCollection::ESearchMode::Includes);
-						InDetailBuilder.GetPropertyUtilities()->RequestForceRefresh();
-					}
-					if (InPropertyHandle->CanResetToDefault())
-					{
-						InPropertyHandle->ResetToDefault();
-					}
-				}),
-			false);
-	}
+	CustomizeDetailWidgetRow(InWidgetRow, InDetailBuilder, BindingsOwner, InPropertyHandle, TargetPath, CachedBindingData);
 
-	CustomizeWidgetArgs(Args, BindingsOwner, TargetPath, CachedBindingData);
+	CustomizePropertyBindingWidgetArgs(Args, BindingsOwner, TargetPath, CachedBindingData);
 
 	InWidgetRow.ExtensionContent()
 	[
@@ -1465,6 +1462,38 @@ void FPropertyBindingExtension::UpdateContextStruct(TConstStructView<FPropertyBi
 	// nothing to do
 }
 
+void FPropertyBindingExtension::CustomizeDetailWidgetRow(FDetailWidgetRow& InWidgetRow, const IDetailLayoutBuilder& InDetailBuilder, IPropertyBindingBindingCollectionOwner* InBindingsOwner, TSharedPtr<IPropertyHandle> InPropertyHandle, const FPropertyBindingPath& InTargetPath, TSharedPtr<UE::PropertyBinding::FCachedBindingData> InCachedBindingData) const
+{
+	// ResetToDefault
+	{
+		InWidgetRow.CustomResetToDefault = FResetToDefaultOverride::Create(
+			MakeAttributeLambda([InCachedBindingData, InPropertyHandle]()
+				{
+					return InPropertyHandle->CanResetToDefault() || InCachedBindingData->HasBinding(FPropertyBindingBindingCollection::ESearchMode::Includes);
+				}),
+			FSimpleDelegate::CreateLambda([InCachedBindingData, &InDetailBuilder, InPropertyHandle]()
+				{
+					if (InCachedBindingData->HasBinding(FPropertyBindingBindingCollection::ESearchMode::Includes))
+					{
+						InCachedBindingData->RemoveBinding(FPropertyBindingBindingCollection::ESearchMode::Includes);
+						InDetailBuilder.GetPropertyUtilities()->RequestForceRefresh();
+					}
+					if (InPropertyHandle->CanResetToDefault())
+					{
+						InPropertyHandle->ResetToDefault();
+					}
+				}),
+			false);
+	}
+}
+
+void FPropertyBindingExtension::CustomizePropertyBindingWidgetArgs(FPropertyBindingWidgetArgs& Args,
+	IPropertyBindingBindingCollectionOwner* InBindingsOwner, const FPropertyBindingPath& InTargetPath,
+	TSharedPtr<UE::PropertyBinding::FCachedBindingData> InCachedBindingData)
+{
+	
+}
+
 bool FPropertyBindingExtension::CanBindToProperty(const FPropertyBindingPath& InTargetPath, const IPropertyHandle& InPropertyHandle) const
 {
 	return true;
@@ -1537,6 +1566,11 @@ bool FPropertyBindingExtension::CanPromoteToParameter(const TSharedPtr<IProperty
 	return true;
 }
 
+bool FPropertyBindingExtension::CanBindToArrayElements() const
+{
+	return false;
+}
+
 TSharedRef<SWidget> FPropertyBindingExtension::MakeContextStructWidget(const FPropertyBindingBindableStructDescriptor& InContextStruct)
 {
 	FEdGraphPinType PinType;
diff --git a/Engine/Plugins/Runtime/PropertyBindingUtils/Source/PropertyBindingUtilsEditor/Public/PropertyBindingExtension.h b/Engine/Plugins/Runtime/PropertyBindingUtils/Source/PropertyBindingUtilsEditor/Public/PropertyBindingExtension.h
--- a/Engine/Plugins/Runtime/PropertyBindingUtils/Source/PropertyBindingUtilsEditor/Public/PropertyBindingExtension.h
+++ b/Engine/Plugins/Runtime/PropertyBindingUtils/Source/PropertyBindingUtilsEditor/Public/PropertyBindingExtension.h
@@ -82,6 +82,7 @@ struct FCachedBindingData
 	virtual ~FCachedBindingData() = default;
 
 	[[nodiscard]] PROPERTYBINDINGUTILSEDITOR_API bool HasBinding(FPropertyBindingBindingCollection::ESearchMode SearchMode) const;
+	[[nodiscard]] PROPERTYBINDINGUTILSEDITOR_API const FPropertyBindingBinding* FindBinding(FPropertyBindingBindingCollection::ESearchMode SearchMode) const;
 	PROPERTYBINDINGUTILSEDITOR_API void AddBinding(TConstArrayView<FBindingChainElement> InBindingChain);
 	PROPERTYBINDINGUTILSEDITOR_API void RemoveBinding(FPropertyBindingBindingCollection::ESearchMode RemoveMode);
 
@@ -100,7 +101,6 @@ struct FCachedBindingData
 	[[nodiscard]] PROPERTYBINDINGUTILSEDITOR_API FLinearColor GetColor();
 	[[nodiscard]] PROPERTYBINDINGUTILSEDITOR_API const FSlateBrush* GetImage();
 
-protected:
 	[[nodiscard]] TConstStructView<FPropertyBindingBindableStructDescriptor> GetBindableStructDescriptor(int32 InStructIndex) const;
 	[[nodiscard]] TStructView<FPropertyBindingBindableStructDescriptor> GetMutableBindableStructDescriptor(int32 InStructIndex);
 	[[nodiscard]] TWeakObjectPtr<> GetWeakOwner() const;
@@ -109,6 +109,7 @@ protected:
 	[[nodiscard]] const FPropertyBindingPath& GetSourcePath() const;
 	[[nodiscard]] const FPropertyBindingPath& GetTargetPath() const;
 	[[nodiscard]] TConstArrayView<TInstancedStruct<FPropertyBindingBindableStructDescriptor>> GetAccessibleStructs() const;
+
 	[[nodiscard]] FText GetFormatableText() const;
 	[[nodiscard]] FText GetFormatableTooltipText() const;
 
@@ -238,12 +239,18 @@ protected:
 	PROPERTYBINDINGUTILSEDITOR_API virtual uint64 GetDisallowedPropertyFlags() const;
 	PROPERTYBINDINGUTILSEDITOR_API virtual bool GetPromotionToParameterOverrideInternal(const FProperty& InProperty, bool& bOutOverride) const;
 	PROPERTYBINDINGUTILSEDITOR_API virtual void UpdateContextStruct(TConstStructView<FPropertyBindingBindableStructDescriptor> InStructDesc, FBindingContextStruct& InOutContextStruct, TMap<FString, FText>& InOutSectionNames) const;
-	virtual void CustomizeWidgetArgs(FPropertyBindingWidgetArgs& Args, IPropertyBindingBindingCollectionOwner* InBindingsOwner, const FPropertyBindingPath& InTargetPath, TSharedPtr<UE::PropertyBinding::FCachedBindingData> InCachedBindingData) { }
-	virtual bool CanBindToArrayElements() const
+	PROPERTYBINDINGUTILSEDITOR_API virtual void CustomizeDetailWidgetRow(FDetailWidgetRow& InWidgetRow, const IDetailLayoutBuilder& InDetailBuilder, IPropertyBindingBindingCollectionOwner* InBindingsOwner, TSharedPtr<IPropertyHandle> InPropertyHandle, const FPropertyBindingPath& InTargetPath, TSharedPtr<UE::PropertyBinding::FCachedBindingData> InCachedBindingData) const;
+
+	PROPERTYBINDINGUTILSEDITOR_API virtual void CustomizePropertyBindingWidgetArgs(FPropertyBindingWidgetArgs& Args, IPropertyBindingBindingCollectionOwner* InBindingsOwner, const FPropertyBindingPath& InTargetPath, TSharedPtr<UE::PropertyBinding::FCachedBindingData> InCachedBindingData);
+
+	UE_DEPRECATED(5.7, "CustomizeWidgetArgs is deprecated. Use CustomizePropertyBindingWidgetArgs instead.")
+	virtual void CustomizeWidgetArgs(FPropertyBindingWidgetArgs& Args, IPropertyBindingBindingCollectionOwner* InBindingsOwner, const FPropertyBindingPath& InTargetPath, TSharedPtr<UE::PropertyBinding::FCachedBindingData> InCachedBindingData) final
 	{
-		return false;
+		CustomizePropertyBindingWidgetArgs(Args, InBindingsOwner, InTargetPath, InCachedBindingData);
 	}
 
+	PROPERTYBINDINGUTILSEDITOR_API virtual bool CanBindToArrayElements() const;
+
 	PROPERTYBINDINGUTILSEDITOR_API static TSharedRef<SWidget> MakeContextStructWidget(const FPropertyBindingBindableStructDescriptor& InContextStruct);
 
 private:
diff --git a/Engine/Plugins/Runtime/StateTree/Source/StateTreeEditorModule/Private/Customizations/StateTreeBindingExtension.cpp b/Engine/Plugins/Runtime/StateTree/Source/StateTreeEditorModule/Private/Customizations/StateTreeBindingExtension.cpp
--- a/Engine/Plugins/Runtime/StateTree/Source/StateTreeEditorModule/Private/Customizations/StateTreeBindingExtension.cpp
+++ b/Engine/Plugins/Runtime/StateTree/Source/StateTreeEditorModule/Private/Customizations/StateTreeBindingExtension.cpp
@@ -1,9 +1,13 @@
 // Copyright Epic Games, Inc. All Rights Reserved.
 
 #include "StateTreeBindingExtension.h"
+
+#include "ScopedTransaction.h"
+#include "DetailLayoutBuilder.h"
 #include "EdGraphSchema_K2.h"
 #include "Features/IModularFeatures.h"
 #include "IPropertyAccessEditor.h"
+#include "IPropertyUtilities.h"
 #include "StateTreeAnyEnum.h"
 #include "StateTreeCompiler.h"
 #include "StateTreeEditorPropertyBindings.h"
@@ -13,18 +17,16 @@
 #include "Widgets/Layout/SBox.h"
 #include "StateTreePropertyRef.h"
 #include "StateTreePropertyRefHelpers.h"
-#include "DetailLayoutBuilder.h"
 #if UE_ENABLE_INCLUDE_ORDER_DEPRECATED_IN_5_6
 #include "IPropertyUtilities.h"
 #endif // UE_ENABLE_INCLUDE_ORDER_DEPRECATED_IN_5_6
 #include "IDetailChildrenBuilder.h"
 #include "IStructureDataProvider.h"
-#include "Kismet2/BlueprintEditorUtils.h"
 #include "PropertyBindingExtension.h"
 #include "StateTreeEditorNodeUtils.h"
-#include "StateTreeEditorModule.h"
 #include "StateTreeEditorData.h"
 #include "StateTreeDelegate.h"
+#include "StateTreeEditorDataClipboardHelpers.h"
 
 #define LOCTEXT_NAMESPACE "StateTreeEditor"
 
@@ -41,6 +43,27 @@ struct FRefTypeInfo : UE::PropertyBinding::FPropertyInfoOverride
 
 const FName AllowAnyBindingName(TEXT("AllowAnyBinding"));
 
+void MakeBindingChainFromBindingPathIndirections(int32 InAccessibleStructIndex, TConstArrayView<FPropertyBindingPathIndirection> InResolvedIndirections, TArray<FBindingChainElement>& OutBindingChain)
+{
+	OutBindingChain.Reset();
+
+	OutBindingChain.Reserve(InResolvedIndirections.Num() + 1);
+	OutBindingChain.Emplace(nullptr, InAccessibleStructIndex);
+	for (int32 Idx = 0; Idx < InResolvedIndirections.Num(); ++Idx)
+	{
+		const FPropertyBindingPathIndirection& Indirection = InResolvedIndirections[Idx];
+		const int32 ArrayIndex = Indirection.GetArrayIndex();
+
+		OutBindingChain.Emplace(const_cast<FProperty*>(Indirection.GetProperty()), ArrayIndex);
+
+		// Indirection format: [ArrayProperty, InnerIndexInArray].[InnerProperty, -1] while BindingChain format skips the InnerProperty item.
+		if (ArrayIndex != INDEX_NONE && Indirection.GetAccessType() == EPropertyBindingPropertyAccessType::IndexArray)
+		{
+			Idx += 1;
+		}
+	}
+}
+
 UObject* FindEditorBindingsOwner(UObject* InObject)
 {
 	UObject* Result = nullptr;
@@ -57,6 +80,21 @@ UObject* FindEditorBindingsOwner(UObject* InObject)
 	return Result;
 }
 
+UObject* FindStateOrEditorDataOwner(UObject* InObject)
+{
+	UObject* Result = nullptr;
+
+	for (UObject* Outer = InObject; Outer; Outer = Outer->GetOuter())
+	{
+		if (Outer->IsA<UStateTreeEditorData>() || Outer->IsA<UStateTreeState>())
+		{
+			Result = Outer;
+			break;
+		}
+	}
+
+	return Result;
+}
 
 static bool IsDelegateDispatcherProperty(const FProperty* InProperty)
 {
@@ -116,6 +154,8 @@ FOnStateTreePropertyBindingChanged STATETREEEDITORMODULE_API OnStateTreeProperty
 
 struct FStateTreeCachedBindingData : UE::PropertyBinding::FCachedBindingData, TSharedFromThis<FStateTreeCachedBindingData>
 {
+	friend FStateTreeBindingExtension;
+
 	using Super = FCachedBindingData;
 
 	using FCachedBindingData::FCachedBindingData;
@@ -707,6 +747,301 @@ void FStateTreeBindingsChildrenCustomization::CustomizeChildren(IDetailChildrenB
 	}
 }
 
+void FStateTreeBindingExtension::CustomizeDetailWidgetRow(FDetailWidgetRow& InWidgetRow, const IDetailLayoutBuilder& InDetailBuilder, IPropertyBindingBindingCollectionOwner* InBindingsOwner,
+	TSharedPtr<IPropertyHandle> InPropertyHandle, const FPropertyBindingPath& InTargetPath,
+	TSharedPtr<UE::PropertyBinding::FCachedBindingData> InCachedBindingData) const
+{
+	auto IsClipboardSingleBinding = [](const UE::StateTreeEditor::FClipboardEditorData& InClipboard)
+		{
+			if (!(InClipboard.GetEditorNodesInBuffer().IsEmpty() && InClipboard.GetTransitionsInBuffer().IsEmpty()))
+			{
+				return false;
+			}
+
+			TConstArrayView<FStateTreePropertyPathBinding> Bindings = InClipboard.GetBindingsInBuffer();
+			if (Bindings.IsEmpty())
+			{
+				return false;
+			}
+
+			// Every Binding target should either be a property function, or the target node.
+			// @todo: detect cyclic reference of property function nodes
+			const FGuid TargetNodeID = Bindings[0].GetTargetPath().GetStructID();
+			TArray<FGuid, TInlineAllocator<8>> PropertyFunctionIDs;
+			for (const FStateTreePropertyPathBinding& Binding : Bindings)
+			{
+				if (Binding.GetPropertyFunctionNode().IsValid())
+				{
+					PropertyFunctionIDs.Add(Binding.GetPropertyFunctionNode().Get<const FStateTreeEditorNode>().ID);
+				}
+			}
+
+			for (const FStateTreePropertyPathBinding& Binding : Bindings)
+			{
+				const FGuid BindingTargetID = Binding.GetTargetPath().GetStructID();
+				if (BindingTargetID != TargetNodeID && !PropertyFunctionIDs.Contains(BindingTargetID))
+				{
+					return false;
+				}
+			}
+
+			return true;
+		};
+
+	auto ImportTextAsSingleBindingClipboard = [IsClipboardSingleBinding, InBindingsOwner, InPropertyHandle](bool bInProcessBuffer, UE::StateTreeEditor::FClipboardEditorData& OutClipboard)
+		{
+			TArray<UObject*> OuterObjects;
+			InPropertyHandle->GetOuterObjects(OuterObjects);
+			checkf(OuterObjects.Num() == 1, TEXT("Paste binding is only allowed on single selected object!"));
+
+			constexpr UScriptStruct* TargetType = nullptr;
+			UE::StateTreeEditor::ImportTextAsClipboardEditorData(TargetType, Cast<UStateTreeEditorData>(InBindingsOwner), UE::StateTree::PropertyBinding::FindStateOrEditorDataOwner(OuterObjects[0]), OutClipboard, bInProcessBuffer);
+
+			return (OutClipboard.IsValid() || !bInProcessBuffer) && IsClipboardSingleBinding(OutClipboard);
+		};
+
+	FPropertyBindingExtension::CustomizeDetailWidgetRow(InWidgetRow, InDetailBuilder, InBindingsOwner, InPropertyHandle, InTargetPath, InCachedBindingData);
+
+	// We will refresh the details view when any binding changed, so we can cache the value here
+	constexpr FPropertyBindingBindingCollection::ESearchMode SearchMode = FPropertyBindingBindingCollection::ESearchMode::Exact;
+	const bool bHasBinding = InCachedBindingData->HasBinding(SearchMode);
+
+	if (bHasBinding)
+	{
+		// Disable Copy and Paste if it has a binding
+		// SDetailSingleItemRow will override the default copy and paste action if FDetailWidgetRow::IsCopyPasteBound == true
+		FUIAction EmptyUIAction = FUIAction(
+			// create an empty ExecuteAction so that FDetailWidgetRow::IsCopyPasteBound returns true
+			FExecuteAction::CreateLambda([](){}),
+			FCanExecuteAction::CreateLambda([bHasBinding](){ return false; })
+			);
+
+		InWidgetRow.CopyAction(EmptyUIAction);
+		InWidgetRow.PasteAction(EmptyUIAction);
+	}
+
+	// Copy binding action
+	InWidgetRow.AddCustomContextMenuAction(FUIAction(
+		FExecuteAction::CreateLambda([InCachedBindingData, SearchMode]()
+			{
+				using namespace UE::StateTree::PropertyBinding;
+				if (const FStateTreeCachedBindingData* StateTreeCachedBindingData = static_cast<FStateTreeCachedBindingData*>(InCachedBindingData.Get()))
+				{
+					if (const FPropertyBindingBinding* PropertyBinding = StateTreeCachedBindingData->FindBinding(SearchMode))
+					{
+						if (const UStateTreeEditorData* EditorData = Cast<UStateTreeEditorData>(StateTreeCachedBindingData->GetOwner()))
+						{
+							UE::StateTreeEditor::FClipboardEditorData Clipboard;
+							Clipboard.Append(EditorData, {&PropertyBinding, 1 });
+
+							UE::StateTreeEditor::ExportTextAsClipboardEditorData(Clipboard);
+						}
+					}
+				}
+			}),
+			FCanExecuteAction::CreateLambda([bHasBinding]()
+			{
+				return bHasBinding;
+			})),
+
+		LOCTEXT("CopyBinding", "Copy Binding"),
+		LOCTEXT("CopyBindingTooltip", "Copy the binding of this property"),
+		FSlateIcon());
+
+	// Paste binding action
+	InWidgetRow.AddCustomContextMenuAction(FUIAction(
+		FExecuteAction::CreateLambda([ImportTextAsSingleBindingClipboard, InCachedBindingData, &InDetailBuilder, SearchMode]()
+			{
+				using namespace UE::StateTree::PropertyBinding;
+
+				if (FStateTreeCachedBindingData* StateTreeCachedBindingData = static_cast<FStateTreeCachedBindingData*>(InCachedBindingData.Get()))
+				{
+					UStateTreeEditorData* EditorData = Cast<UStateTreeEditorData>(StateTreeCachedBindingData->GetOwner());
+					if (!EditorData)
+					{
+						return;
+					}
+
+					UE::StateTreeEditor::FClipboardEditorData Clipboard;
+
+					constexpr bool bProcessBuffer = true;
+					if (ImportTextAsSingleBindingClipboard(bProcessBuffer, Clipboard))
+					{
+						TArrayView<FStateTreePropertyPathBinding> Bindings = Clipboard.GetBindingsInBuffer();
+
+						// This is the direct binding we copy pasted over. All the other bindings are associated with the direct binding(property functions)
+						FStateTreePropertyPathBinding& PropertyBinding = Bindings[0];
+						FConstStructView BindingPropertyFunctionNodeView = PropertyBinding.GetPropertyFunctionNode();
+						TConstArrayView<TInstancedStruct<FPropertyBindingBindableStructDescriptor>> AccessibleStructs = StateTreeCachedBindingData->GetAccessibleStructs();
+
+						bool bBindToPropertyFunctionNode = false;
+
+						int32 SourceStructIndex = INDEX_NONE;
+
+						// if the binding source struct is accessible by the new property
+						if (const FStateTreeEditorNode* PropertyFunctionNode = BindingPropertyFunctionNodeView.GetPtr<const FStateTreeEditorNode>())
+						{
+							bBindToPropertyFunctionNode = true;
+
+							// AccessibleStructDesc stores IDs of "template" property function nodes(ID is made of struct name), which is different from the actual property function nodes(their IDs are instantiated)
+							const UScriptStruct* PropertyFunctionNodeStruct = PropertyFunctionNode->GetNode().GetScriptStruct();
+							FGuid TemplatePropertyFunctionID;
+							EditorData->EnumerateBindablePropertyFunctionNodes(
+								[PropertyFunctionNodeStruct, &TemplatePropertyFunctionID](const UScriptStruct* NodeStruct, const FStateTreeBindableStructDesc& Desc, const FStateTreeDataView Value)
+								{
+									if (NodeStruct == PropertyFunctionNodeStruct)
+									{
+										TemplatePropertyFunctionID = Desc.ID;
+										return EStateTreeVisitor::Break;
+									}
+
+									return EStateTreeVisitor::Continue;
+
+								});
+
+							SourceStructIndex = AccessibleStructs.IndexOfByPredicate(
+								[TemplatePropertyFunctionID](const TInstancedStruct<FPropertyBindingBindableStructDescriptor>& AccessibleStructDesc)
+								{
+									return AccessibleStructDesc.Get().ID == TemplatePropertyFunctionID;
+								});
+						}
+						else
+						{
+							SourceStructIndex = AccessibleStructs.IndexOfByPredicate(
+								[&PropertyBinding](const TInstancedStruct<FPropertyBindingBindableStructDescriptor>& AccessibleStructDesc)
+								{
+									return AccessibleStructDesc.Get().ID == PropertyBinding.GetSourcePath().GetStructID();
+								});
+						}
+
+						// Report an error if Source Struct is not accessible
+						if (SourceStructIndex == INDEX_NONE)
+						{
+							FString InaccessibleSourceStructPathName;
+							FString TargetStructPathName;
+							// log the inaccessible source struct of the copied binding from editor data, unless its from a different tree
+							EditorData->VisitAllNodes(
+								[StateTreeCachedBindingData, &PropertyBinding, &InaccessibleSourceStructPathName, &TargetStructPathName](const UStateTreeState* State, const FStateTreeBindableStructDesc& Desc, const FStateTreeDataView Value)
+								{
+									if (PropertyBinding.GetSourcePath().GetStructID() == Desc.ID)
+									{
+										InaccessibleSourceStructPathName = Desc.StatePath + TEXT("/") + Desc.Name.ToString();
+									}
+
+									if (StateTreeCachedBindingData->GetTargetPath().GetStructID() == Desc.ID)
+									{
+										TargetStructPathName = Desc.StatePath + TEXT("/") + Desc.Name.ToString();
+									}
+
+									if (!InaccessibleSourceStructPathName.IsEmpty() && !TargetStructPathName.IsEmpty())
+									{
+										return EStateTreeVisitor::Break;
+									}
+
+									return EStateTreeVisitor::Continue;
+								});
+
+							UE::StateTreeEditor::AddErrorNotification(
+								FText::Format(LOCTEXT("InaccessibleBindingSource", "Source Struct {0} is not accessible from Target Struct {1}."), 
+								FText::FromString(InaccessibleSourceStructPathName), 
+								FText::FromString(TargetStructPathName)));
+
+							return;
+						}
+
+						const FPropertyBindingBindableStructDescriptor& SourceStructDesc = AccessibleStructs[SourceStructIndex].Get();
+
+						TArray<FPropertyBindingPathIndirection> Indirections;
+						constexpr FString* OutError = nullptr;
+						constexpr bool bHandleRedirects = true;
+						if (!PropertyBinding.GetSourcePath().ResolveIndirections(SourceStructDesc.Struct, Indirections, OutError, bHandleRedirects))
+						{
+							return;
+						}
+
+						const FProperty* SourceProperty = Indirections.Last().GetProperty();
+
+						TArray<FBindingChainElement> BindingChain;
+						MakeBindingChainFromBindingPathIndirections(SourceStructIndex, Indirections, BindingChain);
+
+						// there are essentially two code paths. One to check if the struct itself is bindable(CanBindToContextStruct).
+						// The other to check if property is bindable(need to pass both CanBindToProperty and CanAcceptPropertyOrChildren)
+						// The Source Property and Property Chain returned by both functions are different, so we only check CanBindToProperty here
+						// All the extra checks we did in CanAcceptPropertyOrChild have already been passed by the original binding we copy from(i.e. they are all agnostic to the target property)
+						// @todo: see if we could refactor out the common part of CanBindToProperty and CanAcceptPropertyOrChildren in FCachedBindingData so we don't have to convert back to BindingChain here  
+						const bool bIsSourcePropertyCompatible = bBindToPropertyFunctionNode ?
+							StateTreeCachedBindingData->CanBindToContextStruct(SourceStructDesc.Struct, BindingChain[0].ArrayIndex):
+							StateTreeCachedBindingData->CanBindToProperty(SourceProperty, BindingChain);
+
+						// Report an error if properties are not compatible
+						if (!bIsSourcePropertyCompatible)
+						{
+							FText SourceText;
+							if (bBindToPropertyFunctionNode)
+							{
+								FString SourceStructName;
+								// We need the class name instead of display name for error reporting property function
+								EditorData->EnumerateBindablePropertyFunctionNodes(
+									[&SourceStructName, &SourceStructDesc](const UScriptStruct* NodeStruct, const FStateTreeBindableStructDesc& Desc, const FStateTreeDataView Value)
+									{
+										if (SourceStructDesc.ID == Desc.ID)
+										{
+											SourceStructName = NodeStruct->GetName();
+											return EStateTreeVisitor::Break;
+										}
+
+										return EStateTreeVisitor::Continue;
+									});
+
+								SourceText = FText::FromString(MoveTemp(SourceStructName));
+							}
+							else
+							{
+								SourceText = SourceProperty->GetDisplayNameText();
+							}
+
+							UE::StateTreeEditor::AddErrorNotification(FText::Format(
+								LOCTEXT("IncompatibleSourceProperty", "Source Property {0} is not compatible with Target Property {1}."),
+								SourceText,
+								StateTreeCachedBindingData->GetPropertyHandle()->GetPropertyDisplayName()));
+
+							return;
+						}
+
+						// Add/Replace the binding to property
+						FScopedTransaction Transaction(LOCTEXT("StateTreeBindingData_PasteBinding", "Paste Binding"));
+
+						EditorData->Modify();
+
+						StateTreeCachedBindingData->RemoveBinding(SearchMode);
+
+						// Update the target path to be the current property
+						PropertyBinding.GetMutableTargetPath() = StateTreeCachedBindingData->GetTargetPath();
+						for (int32 Idx = 0; Idx < Bindings.Num(); ++Idx)
+						{
+							FStateTreePropertyPathBinding& CurrentBinding = Bindings[Idx];
+							FPropertyBindingPath SourceBindingPath = CurrentBinding.GetSourcePath();
+							FPropertyBindingPath TargetBindingPath = CurrentBinding.GetTargetPath();
+
+							EditorData->GetPropertyEditorBindings()->AddStateTreeBinding(MoveTemp(Bindings[Idx]));
+							EditorData->OnPropertyBindingChanged(SourceBindingPath, TargetBindingPath);
+						}
+
+						InDetailBuilder.GetPropertyUtilities()->RequestForceRefresh();
+					}
+				}
+			}),
+			FCanExecuteAction::CreateLambda([ImportTextAsSingleBindingClipboard]()
+			{
+				UE::StateTreeEditor::FClipboardEditorData Clipboard;
+				constexpr bool bProcessBuffer = false;
+				return ImportTextAsSingleBindingClipboard(bProcessBuffer, Clipboard);
+			})),
+		LOCTEXT("PasteBinding", "Paste Binding"),
+		LOCTEXT("PasteBindingTooltip", "Paste the binding from clipboard into this property"),
+		FSlateIcon());
+}
+
 bool FStateTreeBindingsChildrenCustomization::ShouldCustomizeChildren(TSharedRef<IPropertyHandle> InPropertyHandle)
 {
 	TArray<UObject*> OuterObjects;
diff --git a/Engine/Plugins/Runtime/StateTree/Source/StateTreeEditorModule/Private/Customizations/StateTreeBindingExtension.h b/Engine/Plugins/Runtime/StateTree/Source/StateTreeEditorModule/Private/Customizations/StateTreeBindingExtension.h
--- a/Engine/Plugins/Runtime/StateTree/Source/StateTreeEditorModule/Private/Customizations/StateTreeBindingExtension.h
+++ b/Engine/Plugins/Runtime/StateTree/Source/StateTreeEditorModule/Private/Customizations/StateTreeBindingExtension.h
@@ -31,6 +31,8 @@ private:
 	{
 		return true;
 	}
+
+	virtual void CustomizeDetailWidgetRow(FDetailWidgetRow& InWidgetRow, const IDetailLayoutBuilder& InDetailBuilder, IPropertyBindingBindingCollectionOwner* InBindingsOwner, TSharedPtr<IPropertyHandle> InPropertyHandle, const FPropertyBindingPath& InTargetPath, TSharedPtr<UE::PropertyBinding::FCachedBindingData> InCachedBindingData) const override;
 };
 
 /* Overrides bound property's children composition. */
diff --git a/Engine/Plugins/Runtime/StateTree/Source/StateTreeEditorModule/Private/StateTreeEditorDataClipboardHelpers.cpp b/Engine/Plugins/Runtime/StateTree/Source/StateTreeEditorModule/Private/StateTreeEditorDataClipboardHelpers.cpp
--- a/Engine/Plugins/Runtime/StateTree/Source/StateTreeEditorModule/Private/StateTreeEditorDataClipboardHelpers.cpp
+++ b/Engine/Plugins/Runtime/StateTree/Source/StateTreeEditorModule/Private/StateTreeEditorDataClipboardHelpers.cpp
@@ -48,6 +48,7 @@ bool ImportTextAsClipboardEditorData(const UScriptStruct* InTargetType, TNotNull
 	OutClipboardEditorData.Reset();
 
 	FString PastedText;
+
 	FPlatformApplicationMisc::ClipboardPaste(PastedText);
 
 	if (PastedText.IsEmpty())
@@ -55,9 +56,9 @@ bool ImportTextAsClipboardEditorData(const UScriptStruct* InTargetType, TNotNull
 		return false;
 	}
 
-	const UScriptStruct* ClipboardEditorDataScriptStruct = TBaseStructure<FClipboardEditorData>::Get();
+	const UScriptStruct* ScriptStruct = TBaseStructure<FClipboardEditorData>::Get();
 	FDefaultValueImportErrorContext ErrorPipe;
-	ClipboardEditorDataScriptStruct->ImportText(*PastedText, &OutClipboardEditorData, nullptr, PPF_None, &ErrorPipe, ClipboardEditorDataScriptStruct->GetName());
+	ScriptStruct->ImportText(*PastedText, &OutClipboardEditorData, nullptr, PPF_None, &ErrorPipe, ScriptStruct->GetName());
 
 	if (bProcessBuffer)
 	{
@@ -77,6 +78,14 @@ void RemoveInvalidBindings(TNotNull<UStateTreeEditorData*> InEditorData)
 	}
 }
 
+void AddErrorNotification(const FText& InText, float InExpiredDuration)
+{
+	FNotificationInfo NotificationInfo(FText::GetEmpty());
+	NotificationInfo.Text = InText;
+	NotificationInfo.ExpireDuration = InExpiredDuration;
+	FSlateNotificationManager::Get().AddNotification(NotificationInfo);
+}
+
 struct FScopedEditorDataFixer
 {
 	explicit FScopedEditorDataFixer(TNotNull<UStateTreeEditorData*> InEditorData, TNotNull<UObject*> InOwner, bool InbShouldRegenerateGUID = true,
@@ -212,6 +221,28 @@ void FClipboardEditorData::Append(TNotNull<const UStateTreeEditorData*> InStateT
 	}
 }
 
+void FClipboardEditorData::Append(TNotNull<const UStateTreeEditorData*> InStateTree, TConstArrayView<const FPropertyBindingBinding*> InBindingPtrs)
+{
+	bBufferProcessed = false;
+	for (const FPropertyBindingBinding* BindingPtr : InBindingPtrs)
+	{
+		if (BindingsBuffer.ContainsByPredicate([BindingPtr](const FStateTreePropertyPathBinding& InBinding)
+			{
+				return BindingPtr && BindingPtr->GetSourcePath() == InBinding.GetSourcePath() && BindingPtr->GetTargetPath() == InBinding.GetTargetPath();
+			}))
+		{
+			continue;
+		}
+
+		BindingsBuffer.Add(*static_cast<const FStateTreePropertyPathBinding*>(BindingPtr));
+		const FConstStructView FunctionNodeView = BindingPtr->GetPropertyFunctionNode();
+		if (const FStateTreeEditorNode* FunctionNode = FunctionNodeView.GetPtr<const FStateTreeEditorNode>())
+		{
+			CollectBindingsForEditorNodes(InStateTree, { FunctionNode, 1 });
+		}
+	}
+}
+
 void FClipboardEditorData::Reset()
 {
 	EditorNodesBuffer.Reset();
@@ -250,15 +281,7 @@ void FClipboardEditorData::ProcessBuffer(const UScriptStruct* InTargetType, TNot
 	check(!InTargetType || InTargetType->IsChildOf(TBaseStructure<FStateTreeNodeBase>::Get()) || InTargetType->IsChildOf(TBaseStructure<FStateTreeTransition>::Get()));
 	check(InTargetOwner->IsA<UStateTreeState>() || InTargetOwner->IsA<UStateTreeEditorData>());
 
-	auto AddErrorNotification = [](const FText& InText)
-		{
-			FNotificationInfo NotificationInfo(FText::GetEmpty());
-			NotificationInfo.Text = InText;
-			NotificationInfo.ExpireDuration = 5.0f;
-			FSlateNotificationManager::Get().AddNotification(NotificationInfo);
-		};
-
-	auto ValidateEditorNodes = [&InEditorData, &AddErrorNotification](TConstArrayView<FStateTreeEditorNode> InEditorNodes, const UScriptStruct* InRequiredType)
+	auto ValidateEditorNodes = [&InEditorData](TConstArrayView<FStateTreeEditorNode> InEditorNodes, const UScriptStruct* InRequiredType)
 		{
 			for (const FStateTreeEditorNode& EditorNode : InEditorNodes)
 			{
@@ -308,7 +331,7 @@ void FClipboardEditorData::ProcessBuffer(const UScriptStruct* InTargetType, TNot
 			return true;
 		};
 
-	auto ValidateTransitions = [Self = this, &AddErrorNotification, &ValidateEditorNodes](TConstArrayView<FStateTreeTransition> InTransitions, const UScriptStruct* InRequiredType)
+	auto ValidateTransitions = [Self = this, &ValidateEditorNodes](TConstArrayView<FStateTreeTransition> InTransitions, const UScriptStruct* InRequiredType)
 		{
 			if (Self->TransitionsBuffer.Num() && InRequiredType && InRequiredType != TBaseStructure<FStateTreeTransition>::Get())
 			{
diff --git a/Engine/Plugins/Runtime/StateTree/Source/StateTreeEditorModule/Private/StateTreeEditorDataClipboardHelpers.h b/Engine/Plugins/Runtime/StateTree/Source/StateTreeEditorModule/Private/StateTreeEditorDataClipboardHelpers.h
--- a/Engine/Plugins/Runtime/StateTree/Source/StateTreeEditorModule/Private/StateTreeEditorDataClipboardHelpers.h
+++ b/Engine/Plugins/Runtime/StateTree/Source/StateTreeEditorModule/Private/StateTreeEditorDataClipboardHelpers.h
@@ -21,10 +21,12 @@ namespace UE::StateTreeEditor
 void ExportTextAsClipboardEditorData(const FClipboardEditorData& InClipboardEditorData);
 
 bool ImportTextAsClipboardEditorData(const UScriptStruct* InTargetType, TNotNull<UStateTreeEditorData*> InTargetTree, TNotNull<UObject*> InOwner,
-                                     FClipboardEditorData& OutClipboardEditorData, bool bValidateBuffer = true);
+                                     FClipboardEditorData& OutClipboardEditorData, bool bProcessBuffer = true);
 
 void RemoveInvalidBindings(TNotNull<UStateTreeEditorData*> InEditorData);
 
+void AddErrorNotification(const FText& InText, float InExpiredDuration = 5.f);
+
 USTRUCT()
 struct FClipboardEditorData
 {
@@ -39,6 +41,7 @@ struct FClipboardEditorData
 
 	void Append(TNotNull<const UStateTreeEditorData*> InStateTree, TConstArrayView<FStateTreeEditorNode> InEditorNodes);
 	void Append(TNotNull<const UStateTreeEditorData*> InStateTree, TConstArrayView<FStateTreeTransition> InTransitions);
+	void Append(TNotNull<const UStateTreeEditorData*> InStateTree, TConstArrayView<const FPropertyBindingBinding*> InBindingPtrs);
 
 	void Reset();
 
@@ -47,16 +50,31 @@ struct FClipboardEditorData
 		return EditorNodesBuffer;
 	}
 
+	TConstArrayView<FStateTreeEditorNode> GetEditorNodesInBuffer() const
+	{
+		return EditorNodesBuffer;
+	}
+
 	TArrayView<FStateTreeTransition> GetTransitionsInBuffer()
 	{
 		return TransitionsBuffer;
 	}
 
+	TConstArrayView<FStateTreeTransition> GetTransitionsInBuffer() const
+	{
+		return TransitionsBuffer;
+	}
+
 	TArrayView<FStateTreePropertyPathBinding> GetBindingsInBuffer()
 	{
 		return BindingsBuffer;
 	}
 
+	TConstArrayView<FStateTreePropertyPathBinding> GetBindingsInBuffer() const
+	{
+		return BindingsBuffer;
+	}
+
 private:
 	void CollectBindingsForEditorNodes(TNotNull<const UStateTreeEditorData*> InStateTree, TConstArrayView<FStateTreeEditorNode> InEditorNodes);
 
diff --git a/Engine/Plugins/Runtime/StateTree/Source/StateTreeEditorModule/Public/StateTreeEditorData.h b/Engine/Plugins/Runtime/StateTree/Source/StateTreeEditorModule/Public/StateTreeEditorData.h
--- a/Engine/Plugins/Runtime/StateTree/Source/StateTreeEditorModule/Public/StateTreeEditorData.h
+++ b/Engine/Plugins/Runtime/StateTree/Source/StateTreeEditorModule/Public/StateTreeEditorData.h
@@ -100,7 +100,10 @@ public:
 
 	UE_API virtual bool CanCreateParameter(const FGuid InStructID) const override;
 	UE_API virtual void CreateParametersForStruct(const FGuid InStructID, TArrayView<UE::PropertyBinding::FPropertyCreationDescriptor> InOutCreationDescs) override;
+
+	// todo: (jira UE-337309) we have many sites that manipulate state/node which will change bindings, but not calling this function. Currently it is only called when you change binding from the details view
 	UE_API virtual void OnPropertyBindingChanged(const FPropertyBindingPath& InSourcePath, const FPropertyBindingPath& InTargetPath) override;
+
 	UE_API virtual void AppendBindablePropertyFunctionStructs(TArray<TInstancedStruct<FPropertyBindingBindableStructDescriptor>>& InOutStructs) const override;
 	//~ End IStateTreeEditorPropertyBindingsOwner interface
 
@@ -191,6 +194,13 @@ public:
 	 */
 	UE_API void GetAccessibleStructsInExecutionPath(const TConstArrayView<const UStateTreeState*> Path, const FGuid TargetStructID, TArray<TInstancedStruct<FPropertyBindingBindableStructDescriptor>>& OutStructDescs) const;
 
+	/**
+	 * Iterates over all template property function nodes.
+	 * Note it is different from calling VisitStructBoundPropertyFunctions which returns the instanced property functions for a node.
+	 * @param InFunc function called at each node, should return Continue if visiting is continued or Break to stop.
+	 */
+	UE_API virtual EStateTreeVisitor EnumerateBindablePropertyFunctionNodes(TFunctionRef<EStateTreeVisitor(const UScriptStruct* NodeStruct, const FStateTreeBindableStructDesc& Desc, const FStateTreeDataView Value)> InFunc) const override;
+
 	UE_DEPRECATED(5.6, "Use GetAccessibleStructsInExecutionPath instead")
 	void GetAccessibleStruct(const TConstArrayView<const UStateTreeState*> Path, const FGuid TargetStructID, TArray<TInstancedStruct<FPropertyBindingBindableStructDescriptor>>& OutStructDescs) const
 	{
@@ -347,8 +357,6 @@ PRAGMA_ENABLE_DEPRECATION_WARNINGS
 	virtual void CreateRootProperties(TArrayView<UE::PropertyBinding::FPropertyCreationDescriptor> InOutCreationDescs) { UE::PropertyBinding::CreateUniquelyNamedPropertiesInPropertyBag(InOutCreationDescs, RootParameterPropertyBag); }
 
 private:
-	UE_API virtual EStateTreeVisitor EnumerateBindablePropertyFunctionNodes(TFunctionRef<EStateTreeVisitor(const UScriptStruct* NodeStruct, const FStateTreeBindableStructDesc& Desc, const FStateTreeDataView Value)> InFunc) const override;
-
 	UE_API void FixObjectInstances(TSet<UObject*>& SeenObjects, UObject& Outer, FStateTreeEditorNode& Node);
 	UE_API void FixObjectNodes();
 	UE_API void FixDuplicateIDs();
