[State Tree] fixed details view not refreshing when pasting nodes into an exisiting one from the SStateTreeView

diff --git a/Engine/Plugins/Runtime/StateTree/Source/StateTreeEditorModule/Public/StateTreeEditorMode.h b/Engine/Plugins/Runtime/StateTree/Source/StateTreeEditorModule/Public/StateTreeEditorMode.h
--- a/Engine/Plugins/Runtime/StateTree/Source/StateTreeEditorModule/Public/StateTreeEditorMode.h
+++ b/Engine/Plugins/Runtime/StateTree/Source/StateTreeEditorModule/Public/StateTreeEditorMode.h
@@ -73,6 +73,9 @@ protected:
 	void HandleOnStateNodesChanged(const UStateTreeState*)
 	{
 		UpdateAsset();
+
+		// when node is being changed from the view model, details view might have already been opened and in that case, we need to rebuild
+		ForceRefreshDetailsView();
 	}
 
 	/** Resolve the internal editor data and fixup the StateTree nodes. */
