[State Tree] Added ExecutionExtension callback when a transition is about to be applied

diff --git a/Engine/Plugins/Runtime/StateTree/Source/StateTreeModule/Private/StateTreeExecutionContext.cpp b/Engine/Plugins/Runtime/StateTree/Source/StateTreeModule/Private/StateTreeExecutionContext.cpp
--- a/Engine/Plugins/Runtime/StateTree/Source/StateTreeModule/Private/StateTreeExecutionContext.cpp
+++ b/Engine/Plugins/Runtime/StateTree/Source/StateTreeModule/Private/StateTreeExecutionContext.cpp
@@ -899,12 +899,19 @@ void FStateTreeMinimalExecutionContext::SendEvent(const FGameplayTag Tag, const
 	}
 }
 
+FStateTreeExecutionExtension* FStateTreeMinimalExecutionContext::GetMutableExecutionExtension() const
+{
+	return Storage.GetMutableExecutionState().ExecutionExtension.GetMutablePtr();
+}
+
 void FStateTreeMinimalExecutionContext::ScheduleNextTick(UE::StateTree::ETickReason Reason)
 {
-	TInstancedStruct<FStateTreeExecutionExtension>& ExecutionExtension = Storage.GetMutableExecutionState().ExecutionExtension;
-	if (bAllowedToScheduleNextTick && RootStateTree.IsScheduledTickAllowed() && ExecutionExtension.IsValid())
+	if (bAllowedToScheduleNextTick && RootStateTree.IsScheduledTickAllowed())
 	{
-		ExecutionExtension.GetMutable().ScheduleNextTick(FStateTreeExecutionExtension::FContextParameters(Owner, RootStateTree, Storage), FStateTreeExecutionExtension::FNextTickArguments(Reason));
+		if (FStateTreeExecutionExtension* ExecExtension = GetMutableExecutionExtension())
+		{
+			ExecExtension->ScheduleNextTick(FStateTreeExecutionExtension::FContextParameters(Owner, RootStateTree, Storage), FStateTreeExecutionExtension::FNextTickArguments(Reason));
+		}
 	}
 }
 
@@ -1863,6 +1870,8 @@ void FStateTreeExecutionContext::TickTriggerTransitionsInternal()
 			UE_STATETREE_DEBUG_SCOPED_PHASE(this, EStateTreeUpdatePhase::ApplyTransitions);
 			UE_STATETREE_DEBUG_TRANSITION_EVENT(this, RequestedTransition->Source, EStateTreeTraceEventType::OnTransition);
 
+			BeginApplyTransition(RequestedTransition->Transition);
+
 			ExitState(RequestedTransition->Selection, RequestedTransition->Transition);
 
 			// Tree succeeded or failed.
@@ -1906,6 +1915,14 @@ void FStateTreeExecutionContext::TickTriggerTransitionsInternal()
 	}
 }
 
+void FStateTreeExecutionContext::BeginApplyTransition(const FStateTreeTransitionResult& InTransitionResult)
+{
+	if (FStateTreeExecutionExtension* ExecExtension = GetMutableExecutionExtension())
+	{
+		ExecExtension->OnBeginApplyTransition({ Owner, RootStateTree, Storage }, InTransitionResult);
+	}
+}
+
 void FStateTreeExecutionContext::BroadcastDelegate(const FStateTreeDelegateDispatcher& Dispatcher)
 {
 	if (!Dispatcher.IsValid())
@@ -2668,6 +2685,8 @@ EStateTreeRunStatus FStateTreeExecutionContext::ForceTransition(const FRecordedS
 	TransitionResult.ChangeType = EStateTreeStateChangeType::Changed;
 	TransitionResult.Priority = RecordedTransition.Priority;
 
+	BeginApplyTransition(TransitionResult);
+
 	ExitState(SelectStateResult, TransitionResult);
 	return EnterState(SelectStateResult, TransitionResult);
 }
diff --git a/Engine/Plugins/Runtime/StateTree/Source/StateTreeModule/Public/StateTreeExecutionContext.h b/Engine/Plugins/Runtime/StateTree/Source/StateTreeModule/Public/StateTreeExecutionContext.h
--- a/Engine/Plugins/Runtime/StateTree/Source/StateTreeModule/Public/StateTreeExecutionContext.h
+++ b/Engine/Plugins/Runtime/StateTree/Source/StateTreeModule/Public/StateTreeExecutionContext.h
@@ -252,6 +252,12 @@ public:
 	UE_API void SendEvent(const FGameplayTag Tag, const FConstStructView Payload = FConstStructView(), const FName Origin = FName());
 
 protected:
+	/**
+	 * Get ExecutionExtension from InstanceStorage for less indirections
+	 * The user is responsible for validity of the result, re-call the getter when needed
+	 * @return ExecutionExtension from InstanceStorage. could be null if Extension hasn't been set or Instance Storage has been reset.
+	 */
+	UE_API FStateTreeExecutionExtension* GetMutableExecutionExtension() const;
 
 	/** Informs the owner when the instance of the tree must woke up from a scheduled tick sleep. */
 	UE_API void ScheduleNextTick(UE::StateTree::ETickReason Reason = UE::StateTree::ETickReason::None);
@@ -1087,6 +1093,8 @@ protected:
 	/** Handles transition triggering part of the tick. */
 	UE_API void TickTriggerTransitionsInternal();
 
+	/** Gives Execution Extension a chance to react. */
+	UE_API void BeginApplyTransition(const FStateTreeTransitionResult& InTransitionResult);
 protected:
 	using FMemoryRequirement = UE::StateTree::InstanceData::FEvaluationScopeInstanceContainer::FMemoryRequirement;
 
diff --git a/Engine/Plugins/Runtime/StateTree/Source/StateTreeModule/Public/StateTreeExecutionExtension.h b/Engine/Plugins/Runtime/StateTree/Source/StateTreeModule/Public/StateTreeExecutionExtension.h
--- a/Engine/Plugins/Runtime/StateTree/Source/StateTreeModule/Public/StateTreeExecutionExtension.h
+++ b/Engine/Plugins/Runtime/StateTree/Source/StateTreeModule/Public/StateTreeExecutionExtension.h
@@ -11,6 +11,8 @@ class UStateTree;
 struct FStateTreeInstanceStorage;
 struct FStateTreeTransitionDelayedState;
 struct FStateTreeReferenceOverrides;
+struct FStateTreeTransitionResult;
+
 namespace UE::StateTree
 {
 	enum class ETickReason : uint8;
@@ -22,7 +24,6 @@ struct FStateTreeExecutionExtension
 {
 	GENERATED_BODY()
 
-public:
 	struct FContextParameters
 	{
 		FContextParameters(UObject& Owner, const UStateTree& StateTree, FStateTreeInstanceStorage& InstanceData)
@@ -60,16 +61,22 @@ public:
 	/** Callback when the execution context request the tree to wakeup from a schedule tick sleep. */
 	virtual void ScheduleNextTick(const FContextParameters& Context, const FNextTickArguments& Args)
 	{
+		
 	}
 
 	UE_DEPRECATED(5.7, "Use ScheduleNextTick with the FNextTickArguments parameter.")
-	virtual void ScheduleNextTick(const FContextParameters& Context) final
-	{
-	}
+	virtual void ScheduleNextTick(const FContextParameters& Context) final {}
 
 	/** Callback when the overrides are set to the execution context . */
 	virtual void OnLinkedStateTreeOverridesSet(const FContextParameters& Context, const FStateTreeReferenceOverrides& Overrides)
 	{
+		
+	}
+
+	/** Callback before the execution context applies a transition. */
+	virtual void OnBeginApplyTransition(const FContextParameters& Context, const FStateTreeTransitionResult& TransitionResult)
+	{
+		
 	}
 };
 
