[StateTree] allow state selection to continue when reaching the source state of the transition if it is completed (requires SelectionRules to use CompletedTransitionStatesCreateNewStates)

diff --git a/Engine/Plugins/Runtime/StateTree/Source/StateTreeModule/Private/StateTreeExecutionContext.cpp b/Engine/Plugins/Runtime/StateTree/Source/StateTreeModule/Private/StateTreeExecutionContext.cpp
--- a/Engine/Plugins/Runtime/StateTree/Source/StateTreeModule/Private/StateTreeExecutionContext.cpp
+++ b/Engine/Plugins/Runtime/StateTree/Source/StateTreeModule/Private/StateTreeExecutionContext.cpp
@@ -7048,7 +7048,7 @@ TOptional<bool> FStateTreeExecutionContext::SelectStateFromSourceInternal(
 
 		if (!bContinueWithStateSelection)
 		{
-			if (InternalArgs.MissingActiveStates.Num() == 1)
+			if (InternalArgs.MissingSourceStates.Num() == 1)
 			{
 				// MissingActiveState is the transition source.
 				if (EnumHasAllFlags(Args.SelectionRules, EStateTreeStateSelectionRules::CompletedTransitionStatesCreateNewStates))
diff --git a/Engine/Plugins/Runtime/StateTree/Source/StateTreeTestSuite/Private/StateTreeTransitionTest.cpp b/Engine/Plugins/Runtime/StateTree/Source/StateTreeTestSuite/Private/StateTreeTransitionTest.cpp
--- a/Engine/Plugins/Runtime/StateTree/Source/StateTreeTestSuite/Private/StateTreeTransitionTest.cpp
+++ b/Engine/Plugins/Runtime/StateTree/Source/StateTreeTestSuite/Private/StateTreeTransitionTest.cpp
@@ -1892,6 +1892,241 @@ struct FStateTreeTest_TransitionTwoEvents : FStateTreeTestBase
 };
 IMPLEMENT_STATE_TREE_INSTANT_TEST(FStateTreeTest_TransitionTwoEvents, "System.StateTree.Transition.TwoEvents");
 
+struct FStateTreeTest_TransitionLinkedAssetFromCompletedParent : FStateTreeTestBase
+{
+	virtual bool InstantTest() override
+	{
+		/*
+		- Tree1
+			- Root1
+				- State1: TryEnter, Task w/ 1 frame to complete
+					- State2 : LinkedAsset (On Complete -> State 3)
+				- State3
+		- Tree2
+			- Root1
+				- State1: Task w/ 2 frames to complete
+		*/
+
+		UStateTree& StateTree1 = NewStateTree();
+		UStateTree& StateTree2 = NewStateTree();
+
+		// Setup linked asset
+		{
+			UStateTreeEditorData& EditorData = *Cast<UStateTreeEditorData>(StateTree2.EditorData);
+			UStateTreeState& Root1 = EditorData.AddSubTree(FName("Tree2_Root1"));
+			UStateTreeState& State1 = Root1.AddChildState(FName("Tree2_State1"));
+
+			State1.AddTask<FTestTask_PrintValue>("Tree2_State1_TaskPrint")
+				.GetInstanceData().Value = 21;
+			State1.AddTask<FTestTask_Stand>(FName("Tree2_State1_TaskStand"))
+				.GetNode().TicksToCompletion = 2;
+			State1.AddTransition(EStateTreeTransitionTrigger::OnStateCompleted, EStateTreeTransitionType::GotoState, &Root1);
+
+			FStateTreeCompilerLog Log;
+			FStateTreeCompiler Compiler(Log);
+			const bool bResult = Compiler.Compile(StateTree2);
+			AITEST_TRUE("StateTree should get compiled", bResult);
+		}
+
+		// Setup main asset
+		{
+			UStateTreeEditorData& EditorData = *Cast<UStateTreeEditorData>(StateTree1.EditorData);
+			UStateTreeState& Root1 = EditorData.AddSubTree(FName("Tree1_Root1"));
+			UStateTreeState& State1 = Root1.AddChildState(FName("Tree1_State1"));
+
+			UStateTreeState& State2 = State1.AddChildState(FName("Tree1_State2"), EStateTreeStateType::LinkedAsset);
+			UStateTreeState& State3 = Root1.AddChildState(FName("Tree1_State3"));
+
+			{
+				TStateTreeEditorNode<FTestTask_PrintValue>& PrintTask = State1.AddTask<FTestTask_PrintValue>("Tree1_State1_TaskPrint");
+				PrintTask.GetInstanceData().Value = 11;
+				TStateTreeEditorNode<FTestTask_Stand>& StandTask = State1.AddTask<FTestTask_Stand>(FName("Tree1_State1_TaskStand"));
+				StandTask.GetNode().TicksToCompletion = 1;
+				StandTask.GetInstanceData().Value = 11;
+				State1.SelectionBehavior = EStateTreeStateSelectionBehavior::TryEnterState;
+				State1.AddTransition(EStateTreeTransitionTrigger::OnStateCompleted, EStateTreeTransitionType::GotoState, &State2);
+			}
+
+			{
+				State2.SetLinkedStateAsset(&StateTree2);
+				TStateTreeEditorNode<FTestTask_PrintValue>& PrintTask = State2.AddTask<FTestTask_PrintValue>("Tree1_State2_TaskPrint");
+				PrintTask.GetInstanceData().Value = 12;
+			}
+
+			FStateTreeCompilerLog Log;
+			FStateTreeCompiler Compiler(Log);
+			const bool bResult = Compiler.Compile(StateTree1);
+			AITEST_TRUE("StateTree should get compiled", bResult);
+		}
+
+		FStateTreeInstanceData InstanceData;
+		{
+			FTestStateTreeExecutionContext Exec(StateTree1, StateTree1, InstanceData);
+			const bool bInitSucceeded = Exec.IsValid();
+			AITEST_TRUE("StateTree should init", bInitSucceeded);
+		}
+
+		constexpr int32 MaxRules = 4;
+		auto MakeStateSelectionRule = [](int32 Index)
+			{
+				EStateTreeStateSelectionRules Rule = EStateTreeStateSelectionRules::None;
+				if ((Index % 2) == 1)
+				{
+					Rule |= EStateTreeStateSelectionRules::CompletedTransitionStatesCreateNewStates;
+					Rule |= EStateTreeStateSelectionRules::CompletedStateBeforeTransitionSourceFailsTransition;
+				}
+				if (Index >= 2)
+				{
+					Rule |= EStateTreeStateSelectionRules::ReselectedStateCreatesNewStates;
+				}
+				return Rule;
+			};
+
+		const TCHAR* EnterStateChangedStr = TEXT("EnterState=Changed");
+		const TCHAR* EnterStateSustainedStr = TEXT("EnterState=Sustained");
+		const TCHAR* ExitStateChangedStr = TEXT("ExitState=Changed");
+		const TCHAR* ExitStateSustainedStr = TEXT("ExitState=Sustained");
+
+		for (int32 Index = 0; Index < MaxRules; ++Index)
+		{
+			const EStateTreeStateSelectionRules StateSelectionRules = MakeStateSelectionRule(Index);
+			InstanceData = FStateTreeInstanceData();
+			if (StateTree1.GetStateSelectionRules() != StateSelectionRules)
+			{
+				StateTree1.ResetCompiled();
+				UStateTreeEditorData* EditorData = CastChecked<UStateTreeEditorData>(StateTree1.EditorData);
+				UStateTreeTestSchema* Schema = CastChecked<UStateTreeTestSchema>(EditorData->Schema);
+				Schema->SetStateSelectionRules(StateSelectionRules);
+
+				FStateTreeCompilerLog Log;
+				FStateTreeCompiler Compiler(Log);
+				const bool bResult = Compiler.Compile(StateTree1);
+				AITEST_TRUE("StateTree should get compiled", bResult);
+			}
+
+			{
+				FTestStateTreeExecutionContext Exec(StateTree1, StateTree1, InstanceData);
+				const bool bInitSucceeded = Exec.IsValid();
+				AITEST_TRUE("StateTree should init", bInitSucceeded);
+
+				// Start and enter state
+				EStateTreeRunStatus Status = Exec.Start();
+				AITEST_TRUE(TEXT("Valid states"), Exec.ExpectInActiveStates("Tree1_Root1", "Tree1_State1"));
+				AITEST_TRUE(TEXT("Expected amount of frames."), InstanceData.GetExecutionState()->ActiveFrames.Num() == 1);
+
+				FTestStateTreeExecutionContext::FLogOrder LogOrder = Exec.Expect("Tree1_State1_TaskPrint", TEXT("EnterState11"));
+				AITEST_TRUE(TEXT("Start enters in the correct order Tree1_State1_TaskPrint"), LogOrder);
+				LogOrder = LogOrder.Then("Tree1_State1_TaskPrint", EnterStateChangedStr);
+				AITEST_TRUE(TEXT("Start enters in the correct order Tree1_State1_TaskPrint"), LogOrder);
+
+				LogOrder = LogOrder.Then("Tree1_State1_TaskStand", TEXT("EnterState"));
+				AITEST_TRUE(TEXT("Start enters in the correct order Tree1_State1_TaskStand"), LogOrder);
+				LogOrder = LogOrder.Then("Tree1_State1_TaskStand", EnterStateChangedStr);
+				AITEST_TRUE(TEXT("Start enters in the correct order Tree1_State1_TaskStand"), LogOrder);
+				Exec.LogClear();
+			}
+
+			{
+				FTestStateTreeExecutionContext Exec(StateTree1, StateTree1, InstanceData);
+				EStateTreeRunStatus Status = Exec.Tick(0.f);
+				AITEST_TRUE(TEXT("Active states should now include linked asset states"), Exec.ExpectInActiveStates("Tree1_Root1", "Tree1_State1", "Tree1_State2", "Tree2_Root1", "Tree2_State1"));
+				AITEST_TRUE(TEXT("Expected amount of frames."), InstanceData.GetExecutionState()->ActiveFrames.Num() == 2);
+
+				FTestStateTreeExecutionContext::FLogOrder LogOrder = Exec.Expect("Tree1_State1_TaskPrint", TEXT("Tick11"));
+				AITEST_TRUE(TEXT("Tree1_State1_TaskPrint ticked"), LogOrder);
+				LogOrder = LogOrder.Then("Tree1_State1_TaskStand", TEXT("Tick"));
+				AITEST_TRUE(TEXT("Tree1_State1_TaskStand ticked"), LogOrder);
+				LogOrder = LogOrder.Then("Tree1_State1_TaskStand", TEXT("StateCompleted"));
+				AITEST_TRUE(TEXT("Completed State1 in Tree1"), LogOrder);
+				LogOrder = LogOrder.Then("Tree1_State2_TaskPrint", TEXT("EnterState12"));
+				AITEST_TRUE(TEXT("Entered State2 in Tree1"), LogOrder);
+				LogOrder = LogOrder.Then("Tree1_State2_TaskPrint", EnterStateChangedStr);
+				AITEST_TRUE(TEXT("Entered State2 in Tree1"), LogOrder);
+				LogOrder = LogOrder.Then("Tree2_State1_TaskPrint", TEXT("EnterState21"));
+				AITEST_TRUE(TEXT("Entered State1 in Tree2"), LogOrder);
+				LogOrder = LogOrder.Then("Tree2_State1_TaskStand", TEXT("EnterState"));
+				AITEST_TRUE(TEXT("Entered State1 in Tree2"), LogOrder);
+				LogOrder = LogOrder.Then("Tree2_State1_TaskStand", EnterStateChangedStr);
+				AITEST_TRUE(TEXT("Entered State1 in Tree2"), LogOrder);
+				Exec.LogClear();
+			}
+
+			{
+				FTestStateTreeExecutionContext Exec(StateTree1, StateTree1, InstanceData);
+				EStateTreeRunStatus Status = Exec.Tick(0.f);
+				AITEST_TRUE(TEXT("Active states should stay the same after Tree1_State1 completion"), Exec.ExpectInActiveStates("Tree1_Root1", "Tree1_State1", "Tree1_State2", "Tree2_Root1", "Tree2_State1"));
+				AITEST_TRUE(TEXT("Expected amount of frames."), InstanceData.GetExecutionState()->ActiveFrames.Num() == 2);
+
+				const bool bCompletedStatedCreateNewStates = EnumHasAnyFlags(StateSelectionRules, EStateTreeStateSelectionRules::CompletedTransitionStatesCreateNewStates);
+				const bool bReselectCreateNewStates = EnumHasAnyFlags(StateSelectionRules, EStateTreeStateSelectionRules::ReselectedStateCreatesNewStates);
+				const TCHAR* EnterStateStr = (bCompletedStatedCreateNewStates || bReselectCreateNewStates) ? EnterStateChangedStr: EnterStateSustainedStr;
+				const TCHAR* ExitStateStr = (bCompletedStatedCreateNewStates || bReselectCreateNewStates) ? ExitStateChangedStr: ExitStateSustainedStr;
+
+				FTestStateTreeExecutionContext::FLogOrder LogOrder = Exec.Expect("Tree1_State1_TaskPrint", TEXT("Tick11"));
+				AITEST_TRUE(TEXT("Tree1_State1_TaskPrint ticked"), LogOrder);
+
+				if (bCompletedStatedCreateNewStates)
+				{
+					LogOrder = LogOrder.Then("Tree1_State1_TaskStand", TEXT("Tick"));
+					AITEST_TRUE(TEXT("Tree1_State1_TaskStand ticked"), LogOrder);
+				}
+
+				LogOrder = LogOrder.Then("Tree1_State2_TaskPrint", TEXT("Tick12"));
+				AITEST_TRUE(TEXT("Tree1_State2_TaskPrint ticked"), LogOrder);
+				LogOrder = LogOrder.Then("Tree2_State1_TaskPrint", TEXT("Tick21"));
+				AITEST_TRUE(TEXT("Tree2_State1_TaskPrint ticked"), LogOrder);
+				LogOrder = LogOrder.Then("Tree2_State1_TaskStand", TEXT("Tick"));
+				AITEST_TRUE(TEXT("Tree2_State1_TaskStand ticked"), LogOrder);
+				LogOrder = LogOrder.Then("Tree2_State1_TaskStand", TEXT("StateCompleted"));
+				AITEST_TRUE(TEXT("Completed State1 in Tree2"), LogOrder);
+				LogOrder = LogOrder.Then("Tree1_State1_TaskStand", TEXT("StateCompleted"));
+				AITEST_TRUE(TEXT("Completed State1 in Tree1"), LogOrder);
+				LogOrder = LogOrder.Then("Tree2_State1_TaskStand", ExitStateStr);
+				AITEST_TRUE(TEXT("Exited State1 in Tree2"), LogOrder);
+				LogOrder = Exec.Expect("Tree2_State1_TaskPrint", TEXT("ExitState21"));
+				AITEST_TRUE(TEXT("Exited State1 in Tree2"), LogOrder);
+				LogOrder = Exec.Expect("Tree2_State1_TaskPrint", ExitStateStr);
+				AITEST_TRUE(TEXT("Exited State1 in Tree2"), LogOrder);
+				LogOrder = Exec.Expect("Tree1_State2_TaskPrint", TEXT("ExitState12"));
+				AITEST_TRUE(TEXT("Exited State2 in Tree1"), LogOrder);
+				LogOrder = Exec.Expect("Tree1_State2_TaskPrint", ExitStateStr);
+				AITEST_TRUE(TEXT("Exited State2 in Tree1"), LogOrder);
+
+				if (bCompletedStatedCreateNewStates)
+				{
+					LogOrder = LogOrder.Then("Tree1_State1_TaskStand", ExitStateStr);
+					AITEST_TRUE(TEXT("Exited State1 in Tree1"), LogOrder);
+					LogOrder = Exec.Expect("Tree1_State1_TaskPrint", TEXT("ExitState11"));
+					AITEST_TRUE(TEXT("Exited State1 in Tree1"), LogOrder);
+					LogOrder = Exec.Expect("Tree1_State1_TaskPrint", ExitStateStr);
+					AITEST_TRUE(TEXT("Exited State1 in Tree1"), LogOrder);
+					LogOrder = Exec.Expect("Tree1_State1_TaskPrint", TEXT("EnterState11"));
+					AITEST_TRUE(TEXT("Entered State1 in Tree1"), LogOrder);
+					LogOrder = LogOrder.Then("Tree1_State1_TaskStand", EnterStateStr);
+					AITEST_TRUE(TEXT("Entered State1 in Tree1"), LogOrder);
+				}
+
+				LogOrder = Exec.Expect("Tree1_State2_TaskPrint", TEXT("EnterState12"));
+				AITEST_TRUE(TEXT("Entered State2 in Tree1"), LogOrder);
+				LogOrder = LogOrder.Then("Tree1_State2_TaskPrint", EnterStateStr);
+				AITEST_TRUE(TEXT("Entered State2 in Tree1"), LogOrder);
+				LogOrder = Exec.Expect("Tree2_State1_TaskPrint", TEXT("EnterState21"));
+				AITEST_TRUE(TEXT("Entered State1 in Tree2"), LogOrder);
+				LogOrder = LogOrder.Then("Tree2_State1_TaskStand", EnterStateStr);
+				AITEST_TRUE(TEXT("Entered State1 in Tree2"), LogOrder);
+				Exec.LogClear();
+			}
+
+			{
+				FTestStateTreeExecutionContext Exec(StateTree1, StateTree1, InstanceData);
+				Exec.Stop();
+			}
+		}
+		return true;
+	}
+};
+IMPLEMENT_STATE_TREE_INSTANT_TEST(FStateTreeTest_TransitionLinkedAssetFromCompletedParent, "System.StateTree.Transition.LinkedAssetFromCompletedParent");
+
 struct FStateTreeTest_TransitionLinkedAssetWith2Root : FStateTreeTestBase
 {
 	virtual bool InstantTest() override
