[State Tree] fixed property not inisitalized properly

diff --git a/Engine/Plugins/Runtime/StateTree/Source/StateTreeTestSuite/Private/StateTreeTestTypes.h b/Engine/Plugins/Runtime/StateTree/Source/StateTreeTestSuite/Private/StateTreeTestTypes.h
--- a/Engine/Plugins/Runtime/StateTree/Source/StateTreeTestSuite/Private/StateTreeTestTypes.h
+++ b/Engine/Plugins/Runtime/StateTree/Source/StateTreeTestSuite/Private/StateTreeTestTypes.h
@@ -1416,7 +1416,7 @@ struct FTestTask_OutputBindingsTaskInstanceData
 	uint8 OutputBool : 1 = false;
 
 	UPROPERTY(EditAnywhere, Category = Output)
-	int32 OutputInt;
+	int32 OutputInt = 0;
 
 	UPROPERTY(EditAnywhere, Category = Output)
 	FIntWrapper OutputIntWrapper;
