[State Tree] fixed Delegate Listener cannot bind to nested Delegate Dispatcher because a false CanAcceptPropertyOrChildren will stop traversing any child properties of the current property. The same check has been performed in CanBindToProperty

diff --git a/Engine/Plugins/Runtime/StateTree/Source/StateTreeEditorModule/Private/Customizations/StateTreeBindingExtension.cpp b/Engine/Plugins/Runtime/StateTree/Source/StateTreeEditorModule/Private/Customizations/StateTreeBindingExtension.cpp
--- a/Engine/Plugins/Runtime/StateTree/Source/StateTreeEditorModule/Private/Customizations/StateTreeBindingExtension.cpp
+++ b/Engine/Plugins/Runtime/StateTree/Source/StateTreeEditorModule/Private/Customizations/StateTreeBindingExtension.cpp
@@ -475,12 +475,6 @@ struct FStateTreeCachedBindingData : UE::PropertyBinding::FCachedBindingData, TS
 			return false;
 		}
 
-		// Listener can only bind to dispatcher (prevents listener to listener)
-		if (IsDelegateListenerProperty(TargetProperty) && !IsDelegateDispatcherProperty(&SourceProperty))
-		{
-			return false;
-		}
-
 		if (IsOutputBinding() && !IsPropertyAccessibleForOutputBinding(InBindingChain))
 		{
 			return false;
