[StateTreeDebugger] Prevent StateTreeDebugger from auto-recording on PIE start when not using the legacy mode. Users should rely on the RewindDebugger option.

diff --git a/Engine/Plugins/Animation/GameplayInsights/Source/RewindDebugger/Private/RewindDebugger.cpp b/Engine/Plugins/Animation/GameplayInsights/Source/RewindDebugger/Private/RewindDebugger.cpp
--- a/Engine/Plugins/Animation/GameplayInsights/Source/RewindDebugger/Private/RewindDebugger.cpp
+++ b/Engine/Plugins/Animation/GameplayInsights/Source/RewindDebugger/Private/RewindDebugger.cpp
@@ -78,11 +78,6 @@ FRewindDebugger::FRewindDebugger()
 
 	UnrealInsightsModule = &FModuleManager::LoadModuleChecked<IUnrealInsightsModule>("TraceInsights");
 
-	if (GEditor->bIsSimulatingInEditor || GEditor->PlayWorld)
-	{
-		OnPIEStarted(true);
-	}
-
 	PreBeginPIEHandle = FEditorDelegates::PreBeginPIE.AddRaw(this, &FRewindDebugger::OnPIEStarted);
 	PausePIEHandle = FEditorDelegates::PausePIE.AddRaw(this, &FRewindDebugger::OnPIEPaused);
 	ResumePIEHandle = FEditorDelegates::ResumePIE.AddRaw(this, &FRewindDebugger::OnPIEResumed);
@@ -144,7 +139,15 @@ FRewindDebugger::~FRewindDebugger()
 
 void FRewindDebugger::Initialize()
 {
-	InternalInstance = new FRewindDebugger;
+	FRewindDebugger* Instance = new FRewindDebugger;
+	InternalInstance = Instance;
+
+	// Triggering callbacks after assigning the global instance
+	// so code called from the callback don't track to initialize the instance again.
+	if (GEditor->bIsSimulatingInEditor || GEditor->PlayWorld)
+	{
+		Instance->OnPIEStarted(true);
+	}
 }
 
 void FRewindDebugger::Shutdown()
diff --git a/Engine/Plugins/Animation/GameplayInsights/Source/RewindDebugger/Private/SRewindDebugger.cpp b/Engine/Plugins/Animation/GameplayInsights/Source/RewindDebugger/Private/SRewindDebugger.cpp
--- a/Engine/Plugins/Animation/GameplayInsights/Source/RewindDebugger/Private/SRewindDebugger.cpp
+++ b/Engine/Plugins/Animation/GameplayInsights/Source/RewindDebugger/Private/SRewindDebugger.cpp
@@ -211,7 +211,9 @@ TSharedRef<SWidget> SRewindDebugger::MakeObjectSelectionMenu()
 					}
 	
 					const FObjectInfo& OuterInfo = GameplayProvider->GetObjectInfo(OuterIdentifier);
-					if (OuterInfo.ClassId == LevelClassInfo->Id
+					if (LevelClassInfo == nullptr
+						|| WorldClassInfo == nullptr
+						|| OuterInfo.ClassId == LevelClassInfo->Id
 						|| OuterInfo.ClassId == WorldClassInfo->Id)
 					{
 						const FString Name = ObjectInfo.Name;
diff --git a/Engine/Plugins/Runtime/StateTree/Source/StateTreeEditorModule/Private/Debugger/SStateTreeDebuggerView.cpp b/Engine/Plugins/Runtime/StateTree/Source/StateTreeEditorModule/Private/Debugger/SStateTreeDebuggerView.cpp
--- a/Engine/Plugins/Runtime/StateTree/Source/StateTreeEditorModule/Private/Debugger/SStateTreeDebuggerView.cpp
+++ b/Engine/Plugins/Runtime/StateTree/Source/StateTreeEditorModule/Private/Debugger/SStateTreeDebuggerView.cpp
@@ -168,7 +168,8 @@ void SStateTreeDebuggerView::ResetTracks()
 
 void SStateTreeDebuggerView::OnPIEStarted(const bool bIsSimulating)
 {
-	if (UStateTreeEditorSettings::Get().bShouldDebuggerAutoRecordOnPIE)
+	if (UStateTreeEditorSettings::Get().bEnableLegacyDebuggerWindow
+		&& UStateTreeEditorSettings::Get().bShouldDebuggerAutoRecordOnPIE)
 	{
 		StartRecording();
 	}
@@ -569,7 +570,19 @@ void SStateTreeDebuggerView::ConstructView(const TSharedRef<FUICommandList>& InC
 						{
 							Debugger->SelectInstance(InstanceId);
 
+							// Make sure to open RewindDebugger to create the global instance
+							// or to bring focus on the selected instance
+							if (CanOpenRewindDebugger())
+							{
+								OpenRewindDebugger();
+							}
+
 							IRewindDebugger* RewindDebugger = IRewindDebugger::Instance();
+							if (RewindDebugger == nullptr)
+							{
+								UE_LOG(LogStateTree, Error, TEXT("Unable to access or create instance of RewindDebugger"));
+								return;
+							}
 
 							// By default, we set the debug object associated to the StateTree instance
 							RewindDebugger::FObjectId ObjectToDebug{InstanceId.ToUint64()};
@@ -974,7 +987,8 @@ void SStateTreeDebuggerView::OnNewSession()
 	//  - if previous analysis was not an Editor session
 	check(Debugger);
 	if (!Debugger->IsAnalyzingEditorSession()
-		|| UStateTreeEditorSettings::Get().bShouldDebuggerResetDataOnNewPIESession
+		|| (UStateTreeEditorSettings::Get().bEnableLegacyDebuggerWindow
+			&& UStateTreeEditorSettings::Get().bShouldDebuggerResetDataOnNewPIESession)
 		|| !Debugger->WasAnalyzingEditorSession())
 	{
 		ResetTracks();
diff --git a/Engine/Plugins/Runtime/StateTree/Source/StateTreeEditorModule/Public/StateTreeEditorSettings.h b/Engine/Plugins/Runtime/StateTree/Source/StateTreeEditorModule/Public/StateTreeEditorSettings.h
--- a/Engine/Plugins/Runtime/StateTree/Source/StateTreeEditorModule/Public/StateTreeEditorSettings.h
+++ b/Engine/Plugins/Runtime/StateTree/Source/StateTreeEditorModule/Public/StateTreeEditorSettings.h
@@ -29,22 +29,22 @@ class UStateTreeEditorSettings : public UDeveloperSettings
 	UPROPERTY(EditAnywhere, config, Category = "Compiler")
 	EStateTreeSaveOnCompile SaveOnCompile = EStateTreeSaveOnCompile::Never;
 
-	/** If enabled, debugger starts recording information at the start of each PIE session. */
-	UPROPERTY(EditAnywhere, config, Category = "Debugger")
-	bool bShouldDebuggerAutoRecordOnPIE = true;
-
-	/** If enabled, debugger will clear previous tracks at the start of each PIE session. */
-	UPROPERTY(EditAnywhere, config, Category = "Debugger")
-	bool bShouldDebuggerResetDataOnNewPIESession = false;
-
 	/**
 	 * If enabled, debugger window in the StateTree Asset Editor will display all widgets
 	 * related to the legacy debugger (recording controls, timelines, frame details, etc.).
 	 * Otherwise, it will display options to link to open RewindDebugger and select a given instance
-	 * */
+	 */
 	UPROPERTY(EditAnywhere, config, Category = "Debugger")
 	bool bEnableLegacyDebuggerWindow = false;
 
+	/** If enabled, debugger starts recording information at the start of each PIE session. */
+	UPROPERTY(EditAnywhere, config, Category = "Debugger", meta = (EditCondition = "bEnableLegacyDebuggerWindow", EditConditionHides))
+	bool bShouldDebuggerAutoRecordOnPIE = true;
+
+	/** If enabled, debugger will clear previous tracks at the start of each PIE session. */
+	UPROPERTY(EditAnywhere, config, Category = "Debugger", meta = (EditCondition = "bEnableLegacyDebuggerWindow", EditConditionHides))
+	bool bShouldDebuggerResetDataOnNewPIESession = false;
+
 	/**
 	 * If enabled, changing the class of a node will try to copy over values of properties with the same name and type.
 	 * i.e. if you change one condition for another, and both have a "Target" BB key selector, it'll be kept.
