[State Tree] fixed crash when paste a binding to context data

diff --git a/Engine/Plugins/Runtime/StateTree/Source/StateTreeEditorModule/Private/Customizations/StateTreeBindingExtension.cpp b/Engine/Plugins/Runtime/StateTree/Source/StateTreeEditorModule/Private/Customizations/StateTreeBindingExtension.cpp
--- a/Engine/Plugins/Runtime/StateTree/Source/StateTreeEditorModule/Private/Customizations/StateTreeBindingExtension.cpp
+++ b/Engine/Plugins/Runtime/StateTree/Source/StateTreeEditorModule/Private/Customizations/StateTreeBindingExtension.cpp
@@ -972,6 +972,7 @@ void FStateTreeBindingExtension::CustomizeDetailWidgetRow(FDetailWidgetRow& InWi
 					if (ImportTextAsSingleBindingClipboard(bProcessBuffer, Clipboard))
 					{
 						TArrayView<FStateTreePropertyPathBinding> Bindings = Clipboard.GetBindingsInBuffer();
+						check(Bindings.Num());
 
 						// This is the direct binding we copy pasted over. All the other bindings are associated with the direct binding(property functions)
 						FStateTreePropertyPathBinding& PropertyBinding = Bindings[0];
@@ -1027,6 +1028,7 @@ void FStateTreeBindingExtension::CustomizeDetailWidgetRow(FDetailWidgetRow& InWi
 							EditorData->VisitAllNodes(
 								[StateTreeCachedBindingData, &PropertyBinding, &InaccessibleSourceStructPathName, &TargetStructPathName](const UStateTreeState* State, const FStateTreeBindableStructDesc& Desc, const FStateTreeDataView Value)
 								{
+									// property function nodes are always accessible, so we don't worry about PropertyFunction SourcePath ID being renewed.
 									if (PropertyBinding.GetSourcePath().GetStructID() == Desc.ID)
 									{
 										InaccessibleSourceStructPathName = Desc.StatePath + TEXT("/") + Desc.Name.ToString();
@@ -1053,7 +1055,8 @@ void FStateTreeBindingExtension::CustomizeDetailWidgetRow(FDetailWidgetRow& InWi
 							return;
 						}
 
-						const FPropertyBindingBindableStructDescriptor& SourceStructDesc = AccessibleStructs[SourceStructIndex].Get();
+						const FStateTreeBindableStructDesc& SourceStructDesc = static_cast<const FStateTreeBindableStructDesc&>(AccessibleStructs[SourceStructIndex].Get());
+						const bool bBindToContextStruct = SourceStructDesc.DataSource == EStateTreeBindableStructSource::Context;
 
 						TArray<FPropertyBindingPathIndirection> Indirections;
 						constexpr FString* OutError = nullptr;
@@ -1063,7 +1066,7 @@ void FStateTreeBindingExtension::CustomizeDetailWidgetRow(FDetailWidgetRow& InWi
 							return;
 						}
 
-						const FProperty* SourceProperty = Indirections.Last().GetProperty();
+						const FProperty* SourceProperty = Indirections.Num() ? Indirections.Last().GetProperty() : nullptr;
 
 						TArray<FBindingChainElement> BindingChain;
 						MakeBindingChainFromBindingPathIndirections(SourceStructIndex, Indirections, BindingChain);
@@ -1073,9 +1076,16 @@ void FStateTreeBindingExtension::CustomizeDetailWidgetRow(FDetailWidgetRow& InWi
 						// The Source Property and Property Chain returned by both functions are different, so we only check CanBindToProperty here
 						// All the extra checks we did in CanAcceptPropertyOrChild have already been passed by the original binding we copy from(i.e. they are all agnostic to the target property)
 						// @todo: see if we could refactor out the common part of CanBindToProperty and CanAcceptPropertyOrChildren in FCachedBindingData so we don't have to convert back to BindingChain here  
-						const bool bIsSourcePropertyCompatible = bBindToPropertyFunctionNode ?
-							StateTreeCachedBindingData->CanBindToContextStruct(SourceStructDesc.Struct, BindingChain[0].ArrayIndex):
-							StateTreeCachedBindingData->CanBindToProperty(SourceProperty, BindingChain);
+						bool bIsSourcePropertyCompatible = false;
+						if (bBindToPropertyFunctionNode || bBindToContextStruct)
+						{
+							bIsSourcePropertyCompatible = StateTreeCachedBindingData->CanBindToContextStruct(SourceStructDesc.Struct, BindingChain[0].ArrayIndex);
+						}
+						else
+						{
+							check(SourceProperty);
+							bIsSourcePropertyCompatible = StateTreeCachedBindingData->CanBindToProperty(SourceProperty, BindingChain);
+						}
 
 						// Report an error if properties are not compatible
 						if (!bIsSourcePropertyCompatible)
@@ -1099,6 +1109,11 @@ void FStateTreeBindingExtension::CustomizeDetailWidgetRow(FDetailWidgetRow& InWi
 
 								SourceText = FText::FromString(MoveTemp(SourceStructName));
 							}
+							else if (bBindToContextStruct)
+							{
+								// For Context Struct, we use the context's name
+								SourceText = FText::FromName(SourceStructDesc.Name);
+							}
 							else
 							{
 								SourceText = SourceProperty->GetDisplayNameText();
