[StateTree] Set the alignment when there are only object wrapper nodes.

diff --git a/Engine/Plugins/Runtime/StateTree/Source/StateTreeModule/Private/StateTreeEvaluationScopeInstanceContainer.cpp b/Engine/Plugins/Runtime/StateTree/Source/StateTreeModule/Private/StateTreeEvaluationScopeInstanceContainer.cpp
--- a/Engine/Plugins/Runtime/StateTree/Source/StateTreeModule/Private/StateTreeEvaluationScopeInstanceContainer.cpp
+++ b/Engine/Plugins/Runtime/StateTree/Source/StateTreeModule/Private/StateTreeEvaluationScopeInstanceContainer.cpp
@@ -21,7 +21,7 @@ namespace Private
 		MemoryRequirement.Size += StructSize;
 		MemoryRequirement.Size += sizeof(DebugStructEndTag); // for struct end tag
 	}
-}
+} // namespace Private
 
 void FEvaluationScopeInstanceContainer::FMemoryRequirementBuilder::Add(TNotNull<const UScriptStruct*> Struct)
 {
@@ -38,6 +38,19 @@ void FEvaluationScopeInstanceContainer::FMemoryRequirementBuilder::Add(TNotNull<
 
 FEvaluationScopeInstanceContainer::FMemoryRequirement FEvaluationScopeInstanceContainer::FMemoryRequirementBuilder::Build()
 {
+	if (MemoryRequirement.NumberOfElements == 0)
+	{
+		return FMemoryRequirement();
+	}
+
+	// There is no container memory for object wrapper. The instance is directly saved in the StructView.
+	//In case all the structs are object wrapper.
+	if (FirstStructAlignment <= 0)
+	{
+		MemoryRequirement.Alignment = alignof(FItem);
+		FirstStructAlignment = alignof(FItem);
+	}
+
 	int32 ContainerSize = sizeof(FItem) * MemoryRequirement.NumberOfElements;
 	ContainerSize += sizeof(Private::DebugTableEndTag); // table for end tag
 	ContainerSize = Align(ContainerSize, FirstStructAlignment);
@@ -93,12 +106,13 @@ void FEvaluationScopeInstanceContainer::Add(FStateTreeDataHandle DataHandle, FCo
 	}
 	else
 	{
+		new (ItemContainer + NumberOfElements)FItem();
+
 		const FStateTreeInstanceObjectWrapper& Wrapper = DefaultInstance.Get<const FStateTreeInstanceObjectWrapper>();
 		if (Wrapper.InstanceObject)
 		{
 			UObject* DuplicatedObject = ::DuplicateObject(Wrapper.InstanceObject, GetTransientPackage());
 
-			new (ItemContainer + NumberOfElements)FItem();
 			ItemContainer[NumberOfElements].Instance = FStateTreeDataView(DuplicatedObject);
 			ItemContainer[NumberOfElements].DataHandle = DataHandle;
 		}
@@ -112,27 +126,28 @@ void FEvaluationScopeInstanceContainer::Reset()
 
 	if (bStructsHaveDestructor)
 	{
+		FItem* ItemContainer = static_cast<FItem*>(Memory);
 		for (; NumberOfElements > 0; --NumberOfElements)
 		{
-			FItem* ItemContainer = static_cast<FItem*>(Memory);
-			if (!ItemContainer[NumberOfElements - 1].DataHandle.IsObjectSource())
+			const UStruct* InstanceStruct = ItemContainer[NumberOfElements - 1].Instance.GetStruct();
+			if (InstanceStruct && Cast<const UClass>(InstanceStruct) == nullptr)
 			{
 				FStateTreeDataView& Instance = ItemContainer[NumberOfElements - 1].Instance;
 				Instance.GetStruct()->DestroyStruct(Instance.GetMutableMemory());
 			}
 			if constexpr (!std::is_trivially_destructible_v<FItem>)
 			{
-				ItemContainer->~FItem();
+				ItemContainer[NumberOfElements - 1].~FItem();
 			}
 		}
 		check(NumberOfElements == 0);
 	}
 	else if constexpr (!std::is_trivially_destructible_v<FItem>)
 	{
+		FItem* ItemContainer = static_cast<FItem*>(Memory);
 		for (; NumberOfElements > 0; --NumberOfElements)
 		{
-			FItem* ItemContainer = static_cast<FItem*>(Memory);
-			ItemContainer->~FItem();
+			ItemContainer[NumberOfElements - 1].~FItem();
 		}
 		check(NumberOfElements == 0);
 	}
@@ -183,10 +198,15 @@ void FEvaluationScopeInstanceContainer::TestDebugTags() const
 		FItem* ItemContainer = static_cast<FItem*>(Memory);
 		for (int32 Index = 0; Index < NumberOfElements; ++Index)
 		{
-			const void* MemoryForStructDebugTag = (const uint8*)ItemContainer[Index].Instance.GetMutableMemory()
-				+ ItemContainer[Index].Instance.GetStruct()->GetStructureSize();
-			const int32 StructEndTag = *reinterpret_cast<const uint32*>(MemoryForStructDebugTag);
-			ensure(StructEndTag == Private::DebugStructEndTag);
+			// Object wrapper do not have EndTag, they use the FItem directly.
+			const UStruct* InstanceStruct = ItemContainer[Index].Instance.GetStruct();
+			if (InstanceStruct && Cast<const UClass>(InstanceStruct) == nullptr)
+			{
+				const void* MemoryForStructDebugTag = (const uint8*)ItemContainer[Index].Instance.GetMutableMemory()
+					+ InstanceStruct->GetStructureSize();
+				const int32 StructEndTag = *reinterpret_cast<const uint32*>(MemoryForStructDebugTag);
+				ensure(StructEndTag == Private::DebugStructEndTag);
+			}
 		}
 	}
 #endif
diff --git a/Engine/Plugins/Runtime/StateTree/Source/StateTreeModule/Private/StateTreeInstanceData.cpp b/Engine/Plugins/Runtime/StateTree/Source/StateTreeModule/Private/StateTreeInstanceData.cpp
--- a/Engine/Plugins/Runtime/StateTree/Source/StateTreeModule/Private/StateTreeInstanceData.cpp
+++ b/Engine/Plugins/Runtime/StateTree/Source/StateTreeModule/Private/StateTreeInstanceData.cpp
@@ -1023,20 +1023,22 @@ void FStateTreeInstanceData::Append(UObject& InOwner, TConstArrayView<FConstStru
 		// The instance struct can be empty, in which case the temporary instance is ignored.
 		// If the source is specified, move it to the instance data.
 		// We assume that if the source is object wrapper, it is already the instance we want.
-		if (Struct.IsValid()
-			&& (Source && Source->IsValid()))
+		if (Struct.IsValid())
 		{
-			check(Struct.GetScriptStruct() == Source->GetScriptStruct());
+			if (Source && Source->IsValid())
+			{
+				check(Struct.GetScriptStruct() == Source->GetScriptStruct());
 				
-			FMemory::Memswap(Struct.GetMemory(), Source->GetMutableMemory(), Struct.GetScriptStruct()->GetStructureSize());
-			Source->Reset();
-		}
-		else if (FStateTreeInstanceObjectWrapper* Wrapper = Struct.GetPtr<FStateTreeInstanceObjectWrapper>())
-		{
-			if (Wrapper->InstanceObject)
+				FMemory::Memswap(Struct.GetMemory(), Source->GetMutableMemory(), Struct.GetScriptStruct()->GetStructureSize());
+				Source->Reset();
+			}
+			else if (FStateTreeInstanceObjectWrapper* Wrapper = Struct.GetPtr<FStateTreeInstanceObjectWrapper>())
 			{
-				const bool bDuplicate = Args.bDuplicateWrappedObject || &InOwner != Wrapper->InstanceObject->GetOuter();
-				Wrapper->InstanceObject = UE::StateTree::InstanceData::Private::CopyNodeInstance(Wrapper->InstanceObject, &InOwner, bDuplicate);
+				if (Wrapper->InstanceObject)
+				{
+					const bool bDuplicate = Args.bDuplicateWrappedObject || &InOwner != Wrapper->InstanceObject->GetOuter();
+					Wrapper->InstanceObject = UE::StateTree::InstanceData::Private::CopyNodeInstance(Wrapper->InstanceObject, &InOwner, bDuplicate);
+				}
 			}
 		}
 	}
