Fixed function description

#rnx

diff --git a/Engine/Plugins/Runtime/StateTree/Source/StateTreeModule/Private/PropertyFunctions/StateTreeTransformPropertyFunctions.h b/Engine/Plugins/Runtime/StateTree/Source/StateTreeModule/Private/PropertyFunctions/StateTreeTransformPropertyFunctions.h
--- a/Engine/Plugins/Runtime/StateTree/Source/StateTreeModule/Private/PropertyFunctions/StateTreeTransformPropertyFunctions.h
+++ b/Engine/Plugins/Runtime/StateTree/Source/StateTreeModule/Private/PropertyFunctions/StateTreeTransformPropertyFunctions.h
@@ -65,7 +65,7 @@ struct FStateTreeBreakTransformPropertyFunctionInstanceData
 };
 
 /**
- * Break a transform into position, rotation and scale.
+ * Break a transform into position and rotation.
  */
 USTRUCT(DisplayName = "Break Transform")
 struct FStateTreeBreakTransformPropertyFunction : public FStateTreePropertyFunctionCommonBase
