Fix uninitialized struct members

diff --git a/Engine/Plugins/Runtime/StateTree/Source/StateTreeModule/Private/PropertyFunctions/StateTreeTransformPropertyFunctions.h b/Engine/Plugins/Runtime/StateTree/Source/StateTreeModule/Private/PropertyFunctions/StateTreeTransformPropertyFunctions.h
--- a/Engine/Plugins/Runtime/StateTree/Source/StateTreeModule/Private/PropertyFunctions/StateTreeTransformPropertyFunctions.h
+++ b/Engine/Plugins/Runtime/StateTree/Source/StateTreeModule/Private/PropertyFunctions/StateTreeTransformPropertyFunctions.h
@@ -21,7 +21,7 @@ struct FStateTreeMakeTransformPropertyFunctionInstanceData
 	FQuat InRotation = FQuat::Identity;
 
 	UPROPERTY(EditAnywhere, Category = Output, meta = (DisplayName = "Transform"))
-	FTransform OutTransform;
+	FTransform OutTransform = FTransform::Identity;
 };
 
 /**
@@ -58,10 +58,10 @@ struct FStateTreeBreakTransformPropertyFunctionInstanceData
 	FTransform InTransform = FTransform::Identity;
 
 	UPROPERTY(EditAnywhere, Category = Output, meta = (DisplayName = "Translation"))
-	FVector OutTranslation;
+	FVector OutTranslation = FVector::ZeroVector;
 
 	UPROPERTY(EditAnywhere, Category = Output, meta = (DisplayName = "Rotation"))
-	FQuat OutRotation;
+	FQuat OutRotation = FQuat::Identity;
 };
 
 /**
