[StateTree] Add a condition to compare a single gameplaytag against another.

diff --git a/Engine/Plugins/Runtime/StateTree/Source/StateTreeModule/Private/Conditions/StateTreeGameplayTagConditions.cpp b/Engine/Plugins/Runtime/StateTree/Source/StateTreeModule/Private/Conditions/StateTreeGameplayTagConditions.cpp
--- a/Engine/Plugins/Runtime/StateTree/Source/StateTreeModule/Private/Conditions/StateTreeGameplayTagConditions.cpp
+++ b/Engine/Plugins/Runtime/StateTree/Source/StateTreeModule/Private/Conditions/StateTreeGameplayTagConditions.cpp
@@ -17,6 +17,56 @@ namespace UE::StateTree::Conditions
 
 #endif// WITH_EDITOR
 
+//----------------------------------------------------------------------//
+//  FGameplayTagSingleMatchCondition
+//----------------------------------------------------------------------//
+
+bool FGameplayTagSingleMatchCondition::TestCondition(FStateTreeExecutionContext& Context) const
+{
+	const FInstanceDataType& InstanceData = Context.GetInstanceData(*this);
+
+	SET_NODE_CUSTOM_TRACE_TEXT(Context, Override, TEXT("%s'%s' matches '%s%s'")
+		, *UE::StateTree::DescHelpers::GetInvertText(bInvert, EStateTreeNodeFormatting::Text).ToString()
+		, *InstanceData.Left.ToString()
+		, *UE::StateTree::DescHelpers::GetExactMatchText(bExactMatch, EStateTreeNodeFormatting::Text).ToString()
+		, *InstanceData.Right.ToString());
+
+	return (bExactMatch ? InstanceData.Left.MatchesTagExact(InstanceData.Right) : InstanceData.Left.MatchesTag(InstanceData.Right)) ^ bInvert;
+}
+
+#if WITH_EDITOR
+FText FGameplayTagSingleMatchCondition::GetDescription(const FGuid& ID, FStateTreeDataView InstanceDataView, const IStateTreeBindingLookup& BindingLookup, EStateTreeNodeFormatting Formatting) const
+{
+	const FInstanceDataType* InstanceData = InstanceDataView.GetPtr<FInstanceDataType>();
+	check(InstanceData);
+
+	FText LeftValue = BindingLookup.GetBindingSourceDisplayName(FPropertyBindingPath(ID, GET_MEMBER_NAME_CHECKED(FInstanceDataType, Left)), Formatting);
+	if (LeftValue.IsEmpty())
+	{
+		LeftValue = FText::FromString(InstanceData->Left.ToString());
+	}
+
+	FText RightValue = BindingLookup.GetBindingSourceDisplayName(FPropertyBindingPath(ID, GET_MEMBER_NAME_CHECKED(FInstanceDataType, Right)), Formatting);
+	if (RightValue.IsEmpty())
+	{
+		RightValue = FText::FromString(InstanceData->Right.ToString());
+	}
+
+	const FText InvertText = UE::StateTree::DescHelpers::GetInvertText(bInvert, Formatting);
+	const FText ExactMatchText = UE::StateTree::DescHelpers::GetExactMatchText(bExactMatch, Formatting);
+
+	const FText Format = (Formatting == EStateTreeNodeFormatting::RichText)
+		? LOCTEXT("GameplayTagSingleMatchRich", "{EmptyOrNot}{Left} <s>matches</> {EmptyOrExactly}{Right}")
+		: LOCTEXT("GameplayTagSingleMatch", "{EmptyOrNot}{Left} matches {EmptyOrExactly}{Right}");
+
+	return FText::FormatNamed(Format,
+		TEXT("EmptyOrNot"), InvertText,
+		TEXT("Left"), LeftValue,
+		TEXT("EmptyOrExactly"), ExactMatchText,
+		TEXT("Right"), RightValue);
+}
+#endif // WITH_EDITOR
+
 
 //----------------------------------------------------------------------//
 //  FGameplayTagMatchCondition
diff --git a/Engine/Plugins/Runtime/StateTree/Source/StateTreeModule/Public/Conditions/StateTreeGameplayTagConditions.h b/Engine/Plugins/Runtime/StateTree/Source/StateTreeModule/Public/Conditions/StateTreeGameplayTagConditions.h
--- a/Engine/Plugins/Runtime/StateTree/Source/StateTreeModule/Public/Conditions/StateTreeGameplayTagConditions.h
+++ b/Engine/Plugins/Runtime/StateTree/Source/StateTreeModule/Public/Conditions/StateTreeGameplayTagConditions.h
@@ -7,6 +7,62 @@
 
 #define UE_API STATETREEMODULE_API
 
+USTRUCT()
+struct FGameplayTagSingleMatchConditionInstanceData
+{
+	GENERATED_BODY()
+
+	/** First tag to compare. */
+	UPROPERTY(EditAnywhere, Category = Input)
+	FGameplayTag Left;
+
+	/** Second tag to compare. */
+	UPROPERTY(EditAnywhere, Category = Parameter)
+	FGameplayTag Right;
+};
+
+/**
+ * MatchesTag condition
+ * Succeeds if the tag matches the specified tag.
+ * 
+ * Condition can be used with multiple configurations:
+ *	Does Tag "A.1" match Tag "A" ?
+ *		exact match 'false' will SUCCEED
+ *		exact match 'true' will FAIL
+ *	Does Tag "A" match Tag "A.1" ?
+ *		always returns false
+ */
+USTRUCT(DisplayName="Matches Tag", Category="Gameplay Tags")
+struct FGameplayTagSingleMatchCondition : public FStateTreeConditionCommonBase
+{
+	GENERATED_BODY()
+
+	using FInstanceDataType = FGameplayTagSingleMatchConditionInstanceData;
+
+	FGameplayTagSingleMatchCondition() = default;
+	
+	virtual const UStruct* GetInstanceDataType() const override { return FInstanceDataType::StaticStruct(); }
+	UE_API virtual bool TestCondition(FStateTreeExecutionContext& Context) const override;
+#if WITH_EDITOR
+	UE_API virtual FText GetDescription(const FGuid& ID, FStateTreeDataView InstanceDataView, const IStateTreeBindingLookup& BindingLookup, EStateTreeNodeFormatting Formatting = EStateTreeNodeFormatting::Text) const override;
+	virtual FName GetIconName() const override
+	{
+		return FName("StateTreeEditorStyle|Node.Tag");
+	}
+	virtual FColor GetIconColor() const override
+	{
+		return UE::StateTree::Colors::DarkGrey;
+	}
+#endif // WITH_EDITOR
+	/** If true, the tag has to match exactly, if false then the tag may match its parent tags */
+	UPROPERTY(EditAnywhere, Category = Condition)
+	bool bExactMatch = false;
+
+	UPROPERTY(EditAnywhere, Category = Condition)
+	bool bInvert = false;
+};
+
+
 USTRUCT()
 struct FGameplayTagMatchConditionInstanceData
 {
@@ -52,7 +108,7 @@ struct FGameplayTagMatchCondition : public FStateTreeConditionCommonBase
 		return UE::StateTree::Colors::DarkGrey;
 	}
 #endif
-	/** If true, the tag has to be exactly present, if false then TagContainer will include it's parent tags while matching */
+	/** If true, the tag has to be exactly present, if false then TagContainer will include its parent tags while matching */
 	UPROPERTY(EditAnywhere, Category = Condition)
 	bool bExactMatch = false;
 
