[StateTree] Remove STATETREE_POD_INSTANCEDATA in favor of UE_STATETREE_CONSTRUCTED_TRIVIALLY_COPIED_NO_DESTRUCTOR_INSTANCEDATA or UE_STATETREE_ZEROED_TRIVIALLY_COPIED_NO_DESTRUCTOR_INSTANCEDATA. Properties that are initialized in-place need the constructor. ScriptStruct instances constructed from script (BP or StateTree instance data) are memset to 0 by default.
The POD has a different meaning in C and CPP (pre 2020). It is preferable to use the more descriptive term "trivial" to identify types that can be trivially constructed and copied.

diff --git a/Engine/Plugins/Runtime/StateTree/Source/StateTreeModule/Public/Conditions/StateTreeCommonConditions.h b/Engine/Plugins/Runtime/StateTree/Source/StateTreeModule/Public/Conditions/StateTreeCommonConditions.h
--- a/Engine/Plugins/Runtime/StateTree/Source/StateTreeModule/Public/Conditions/StateTreeCommonConditions.h
+++ b/Engine/Plugins/Runtime/StateTree/Source/StateTreeModule/Public/Conditions/StateTreeCommonConditions.h
@@ -22,7 +22,7 @@ struct FStateTreeCompareIntConditionInstanceData
 	UPROPERTY(EditAnywhere, Category = "Parameter")
 	int32 Right = 0;
 };
-STATETREE_POD_INSTANCEDATA(FStateTreeCompareIntConditionInstanceData);
+UE_STATETREE_ZEROED_TRIVIALLY_COPIED_NO_DESTRUCTOR_INSTANCEDATA(FStateTreeCompareIntConditionInstanceData);
 
 /**
  * Condition comparing two integers.
@@ -65,7 +65,7 @@ struct FStateTreeCompareFloatConditionInstanceData
 	UPROPERTY(EditAnywhere, Category = "Parameter")
 	double Right = 0.0;
 };
-STATETREE_POD_INSTANCEDATA(FStateTreeCompareFloatConditionInstanceData);
+UE_STATETREE_ZEROED_TRIVIALLY_COPIED_NO_DESTRUCTOR_INSTANCEDATA(FStateTreeCompareFloatConditionInstanceData);
 
 /**
  * Condition comparing two floats.
@@ -108,7 +108,7 @@ struct FStateTreeCompareBoolConditionInstanceData
 	UPROPERTY(EditAnywhere, Category = "Parameter")
 	bool bRight = false;
 };
-STATETREE_POD_INSTANCEDATA(FStateTreeCompareBoolConditionInstanceData);
+UE_STATETREE_ZEROED_TRIVIALLY_COPIED_NO_DESTRUCTOR_INSTANCEDATA(FStateTreeCompareBoolConditionInstanceData);
 
 /**
  * Condition comparing two booleans.
@@ -151,7 +151,7 @@ struct FStateTreeCompareEnumConditionInstanceData
 	UPROPERTY(EditAnywhere, Category = "Parameter")
 	FStateTreeAnyEnum Right;
 };
-STATETREE_POD_INSTANCEDATA(FStateTreeCompareEnumConditionInstanceData);
+UE_STATETREE_CONSTRUCTED_TRIVIALLY_COPIED_NO_DESTRUCTOR_INSTANCEDATA(FStateTreeCompareEnumConditionInstanceData);
 
 /**
  * Condition comparing two enums.
@@ -190,15 +190,15 @@ struct FStateTreeCompareDistanceConditionInstanceData
 	GENERATED_BODY()
 
 	UPROPERTY(EditAnywhere, Category = "Input")
-	FVector Source = FVector(EForceInit::ForceInitToZero);
+	FVector Source;
 
 	UPROPERTY(EditAnywhere, Category = "Parameter")
-	FVector Target = FVector(EForceInit::ForceInitToZero);
+	FVector Target;
 
 	UPROPERTY(EditAnywhere, Category = "Parameter")
 	double Distance = 0.0;
 };
-STATETREE_POD_INSTANCEDATA(FStateTreeCompareDistanceConditionInstanceData);
+UE_STATETREE_ZEROED_TRIVIALLY_COPIED_NO_DESTRUCTOR_INSTANCEDATA(FStateTreeCompareDistanceConditionInstanceData);
 
 /**
  * Condition comparing distance between two vectors.
@@ -240,7 +240,7 @@ struct FStateTreeCompareNameConditionInstanceData
 	UPROPERTY(EditAnywhere, Category = "Parameter")
 	FName Right = NAME_None;
 };
-STATETREE_POD_INSTANCEDATA(FStateTreeCompareNameConditionInstanceData);
+UE_STATETREE_CONSTRUCTED_TRIVIALLY_COPIED_NO_DESTRUCTOR_INSTANCEDATA(FStateTreeCompareNameConditionInstanceData);
 
 /**
  * Condition comparing two FNames.
@@ -279,7 +279,7 @@ struct FStateTreeRandomConditionInstanceData
 	UPROPERTY(EditAnywhere, Category = "Parameter", meta = (ClampMin = "0.0", ClampMax = "1.0", UIMin = "0.0", UIMax = "1.0"))
 	float Threshold = 0.5f;
 };
-STATETREE_POD_INSTANCEDATA(FStateTreeRandomConditionInstanceData);
+UE_STATETREE_CONSTRUCTED_TRIVIALLY_COPIED_NO_DESTRUCTOR_INSTANCEDATA(FStateTreeRandomConditionInstanceData);
 
 /**
  * Random condition
diff --git a/Engine/Plugins/Runtime/StateTree/Source/StateTreeModule/Public/StateTreeConditionBase.h b/Engine/Plugins/Runtime/StateTree/Source/StateTreeModule/Public/StateTreeConditionBase.h
--- a/Engine/Plugins/Runtime/StateTree/Source/StateTreeModule/Public/StateTreeConditionBase.h
+++ b/Engine/Plugins/Runtime/StateTree/Source/StateTreeModule/Public/StateTreeConditionBase.h
@@ -83,16 +83,3 @@ struct FStateTreeConditionCommonBase : public FStateTreeConditionBase
 {
 	GENERATED_BODY()
 };
-
-/** Helper macro to define instance data as simple constructible. */
-#define STATETREE_POD_INSTANCEDATA(Type) \
-template <> struct TIsPODType<Type> { enum { Value = true }; }; \
-template<> \
-struct TStructOpsTypeTraits<Type> : public TStructOpsTypeTraitsBase2<Type> \
-{ \
-	enum \
-	{ \
-		WithZeroConstructor = true, \
-		WithNoDestructor = true, \
-	}; \
-};
diff --git a/Engine/Plugins/Runtime/StateTree/Source/StateTreeModule/Public/StateTreeTypes.h b/Engine/Plugins/Runtime/StateTree/Source/StateTreeModule/Public/StateTreeTypes.h
--- a/Engine/Plugins/Runtime/StateTree/Source/StateTreeModule/Public/StateTreeTypes.h
+++ b/Engine/Plugins/Runtime/StateTree/Source/StateTreeModule/Public/StateTreeTypes.h
@@ -1226,3 +1226,44 @@ struct TStructOpsTypeTraits<FStateTreeStateLink> : public TStructOpsTypeTraitsBa
 		WithPostSerialize = true,
 	};
 };
+
+/**
+ * Helper macro to define instance data as a trivial type.
+ * The type can be safely copied using memcpy and don't require special copy construction or destruction. It requires a default constructor.
+ * The type defined default values.
+ * ex: struct { int32 A = 3; };
+ * This is not appropriate for element types that require non-trivial copy constructor/assignment  or non-trivial move constructor/assignment.
+ * ex: struct { TObjectPtr<UObject> A; }; // No. TObjectPtr has a non-trivial copy constructor when UE_OBJECT_PTR_GC_BARRIER is enabled.
+ */
+#define UE_STATETREE_CONSTRUCTED_TRIVIALLY_COPIED_NO_DESTRUCTOR_INSTANCEDATA(Type) \
+template <> struct TIsPODType<Type> { enum { Value = true }; }; \
+template<> \
+struct TStructOpsTypeTraits<Type> : public TStructOpsTypeTraitsBase2<Type> \
+{ \
+	enum \
+	{ \
+		WithNoDestructor = true, \
+	}; \
+};
+
+/**
+ * Helper macro to define instance data as a trivial type.
+ * The type can be safely copied using memcpy and don't require special construction or destruction.
+ * The memory will be memset to 0.
+ * ex: struct { int32 A = 0; };
+ * This is not appropriate for element types that require construction.
+ * ex: struct { FFrameRate A; }; // No. FFrameRate has a constructor.
+ */
+#define UE_STATETREE_ZEROED_TRIVIALLY_COPIED_NO_DESTRUCTOR_INSTANCEDATA(Type) \
+template <> struct TIsPODType<Type> { enum { Value = true }; }; \
+template<> \
+struct TStructOpsTypeTraits<Type> : public TStructOpsTypeTraitsBase2<Type> \
+{ \
+	enum \
+	{ \
+		WithZeroConstructor = true, \
+		WithNoDestructor = true, \
+	}; \
+};
+
+#define STATETREE_POD_INSTANCEDATA(Type) UE_DEPRECATED_MACRO(5.8, "Use UE_STATETREE_CONSTRUCTED_TRIVIALLY_COPIED_NO_DESTRUCTOR_INSTANCEDATA or UE_STATETREE_ZEROED_TRIVIALLY_COPIED_NO_DESTRUCTOR_INSTANCEDATA."); UE_STATETREE_ZEROED_TRIVIALLY_COPIED_NO_DESTRUCTOR_INSTANCEDATA(Type)
diff --git a/Engine/Plugins/Runtime/StateTree/Source/StateTreeTestSuite/Private/StateTreeTestTypes.h b/Engine/Plugins/Runtime/StateTree/Source/StateTreeTestSuite/Private/StateTreeTestTypes.h
--- a/Engine/Plugins/Runtime/StateTree/Source/StateTreeTestSuite/Private/StateTreeTestTypes.h
+++ b/Engine/Plugins/Runtime/StateTree/Source/StateTreeTestSuite/Private/StateTreeTestTypes.h
@@ -746,7 +746,6 @@ struct FStateTreeTestConditionInstanceData
 
 	static std::atomic<int32> GlobalCounter;
 };
-STATETREE_POD_INSTANCEDATA(FStateTreeTestConditionInstanceData);
 
 USTRUCT(meta = (Hidden))
 struct FStateTreeTestCondition : public FStateTreeConditionCommonBase
@@ -964,7 +963,7 @@ struct FTestPropertyFunction_InstanceData
 	UPROPERTY()
 	int32 Result = 0;
 };
-STATETREE_POD_INSTANCEDATA(FTestPropertyFunction_InstanceData);
+UE_STATETREE_ZEROED_TRIVIALLY_COPIED_NO_DESTRUCTOR_INSTANCEDATA(FTestPropertyFunction_InstanceData);
 
 USTRUCT()
 struct FTestPropertyFunction : public FStateTreePropertyFunctionBase
@@ -1028,7 +1027,7 @@ struct FStateTreeTestBooleanConditionInstanceData
 	UPROPERTY(EditAnywhere, Category = Parameters)
 	bool bSuccess = true;
 };
-STATETREE_POD_INSTANCEDATA(FStateTreeTestBooleanConditionInstanceData);
+UE_STATETREE_CONSTRUCTED_TRIVIALLY_COPIED_NO_DESTRUCTOR_INSTANCEDATA(FStateTreeTestBooleanConditionInstanceData);
 
 USTRUCT(meta = (Hidden))
 struct FStateTreeTestBooleanCondition : public FStateTreeConditionCommonBase
