[StateTree] Add a global callback and a callback on the editor extension to validate the state tree asset. The validation occurs when the asset loads, and before it compiles. The user can use the callback to modify the editor data before the asset compiles.

diff --git a/Engine/Plugins/Runtime/StateTree/Source/StateTreeEditorModule/Private/StateTreeEditingSubsystem.cpp b/Engine/Plugins/Runtime/StateTree/Source/StateTreeEditorModule/Private/StateTreeEditingSubsystem.cpp
--- a/Engine/Plugins/Runtime/StateTree/Source/StateTreeEditorModule/Private/StateTreeEditingSubsystem.cpp
+++ b/Engine/Plugins/Runtime/StateTree/Source/StateTreeEditorModule/Private/StateTreeEditingSubsystem.cpp
@@ -10,6 +10,7 @@
 #include "StateTreeEditorData.h"
 #include "StateTreeEditorModule.h"
 #include "StateTreeEditorSchema.h"
+#include "StateTreeEditorDataExtension.h"
 #include "StateTreeObjectHash.h"
 #include "StateTreeTaskBase.h"
 #include "UObject/UObjectGlobals.h"
@@ -270,6 +271,11 @@ void UStateTreeEditingSubsystem::ValidateStateTree(TNotNull<UStateTree*> InState
 			}
 		};
 
+	auto RemoveEmptyExtension = [](TNotNull<UStateTreeEditorData*> EditorData)
+		{
+			EditorData->Extensions.Remove(TObjectPtr<UStateTreeEditorDataExtension>());
+		};
+
 	FixEditorData(InStateTree);
 	if (InStateTree->EditorData)
 	{
@@ -280,15 +286,33 @@ void UStateTreeEditingSubsystem::ValidateStateTree(TNotNull<UStateTree*> InState
 		FixEditorSchema(EditorData);
 
 		EditorData->ReparentStates();
-		if (EditorData->EditorSchema)
-		{
-			EditorData->EditorSchema->Validate(InStateTree);
-		}
+		EditorData->FixDuplicateIDs();
+		EditorData->UpdateBindingsInstanceStructs();
 
 		RemoveUnusedBindings(EditorData);
 		ValidateLinkedStates(EditorData);
 		UpdateLinkedStateParameters(EditorData);
 		UpdateTransactionalFlags(EditorData);
+		RemoveEmptyExtension(EditorData);
+
+		// Notify the module, schema, and extension. Use that order to go from the less specific to the more specific.
+		FStateTreeEditorModule::GetModule().OnValidateStateTree().Broadcast(InStateTree);
+
+		// Schema validation
+		if (EditorData->EditorSchema)
+		{
+			EditorData->EditorSchema->Validate(InStateTree);
+		}
+
+		TArray<UStateTreeEditorDataExtension*> LocalExtensions = EditorData->Extensions;
+		for (UStateTreeEditorDataExtension* Extension : LocalExtensions)
+		{
+			if (Extension)
+			{
+				Extension->Validate();
+			}
+		}
+		RemoveEmptyExtension(EditorData);
 	}
 }
 
diff --git a/Engine/Plugins/Runtime/StateTree/Source/StateTreeEditorModule/Public/StateTreeEditorData.h b/Engine/Plugins/Runtime/StateTree/Source/StateTreeEditorModule/Public/StateTreeEditorData.h
--- a/Engine/Plugins/Runtime/StateTree/Source/StateTreeEditorModule/Public/StateTreeEditorData.h
+++ b/Engine/Plugins/Runtime/StateTree/Source/StateTreeEditorModule/Public/StateTreeEditorData.h
@@ -218,7 +218,12 @@ public:
 	UFUNCTION(BlueprintCallable, Category = "StateTree|Extension", Meta = (DisplayName="Get Extension", DeterminesOutputType = "ExtensionType"))
 	UE_API UStateTreeEditorDataExtension* K2_GetExtension(TSubclassOf<UStateTreeEditorDataExtension> ExtensionType);
 
+	/** Make sure the child and parent states are pointing to each other. */
 	UE_API void ReparentStates();
+	/** Fix any nodes or bindings with duplicated ID. */
+	UE_API void FixDuplicateIDs();
+	/** Update the binding segments. */
+	UE_API void UpdateBindingsInstanceStructs();
 	
 	// StateTree Builder API
 
@@ -359,9 +364,7 @@ PRAGMA_ENABLE_DEPRECATION_WARNINGS
 private:
 	UE_API void FixObjectInstances(TSet<UObject*>& SeenObjects, UObject& Outer, FStateTreeEditorNode& Node);
 	UE_API void FixObjectNodes();
-	UE_API void FixDuplicateIDs();
 	UE_API void DuplicateIDs();
-	UE_API void UpdateBindingsInstanceStructs();
 	UE_API void CallPostLoadOnNodes();
 
 #if WITH_EDITORONLY_DATA
diff --git a/Engine/Plugins/Runtime/StateTree/Source/StateTreeEditorModule/Public/StateTreeEditorDataExtension.h b/Engine/Plugins/Runtime/StateTree/Source/StateTreeEditorModule/Public/StateTreeEditorDataExtension.h
--- a/Engine/Plugins/Runtime/StateTree/Source/StateTreeEditorModule/Public/StateTreeEditorDataExtension.h
+++ b/Engine/Plugins/Runtime/StateTree/Source/StateTreeEditorModule/Public/StateTreeEditorDataExtension.h
@@ -23,11 +23,26 @@ class UStateTreeEditorDataExtension : public UObject
 	GENERATED_BODY()
 
 public:
+	/**
+	 * Validates and applies the editor data extension restrictions on the state tree owned by the extension.
+	 * Also serves as the "safety net" of fixing up editor data following an editor operation.
+	 * You are allowed to modify the editor data.
+	 */
+	virtual void Validate()
+	{
+	}
+
+	/**
+	 * Handled after the asset had properly compiled its public and internal data.
+	 * You are allowed to modify the state tree.
+	 * You should not modify the editor data. It might create an infinite loop.
+	 */
 	virtual bool HandlePostInternalCompile(const UE::StateTree::Compiler::FPostInternalContext& Context)
 	{
 		return true;
 	}
 
+	/** Customize the detail view for the selected state. */
 	virtual void CustomizeDetails(TNonNullPtr<UStateTreeState> State, IDetailLayoutBuilder& DetailBuilder)
 	{
 	}
diff --git a/Engine/Plugins/Runtime/StateTree/Source/StateTreeEditorModule/Public/StateTreeEditorModule.h b/Engine/Plugins/Runtime/StateTree/Source/StateTreeEditorModule/Public/StateTreeEditorModule.h
--- a/Engine/Plugins/Runtime/StateTree/Source/StateTreeEditorModule/Public/StateTreeEditorModule.h
+++ b/Engine/Plugins/Runtime/StateTree/Source/StateTreeEditorModule/Public/StateTreeEditorModule.h
@@ -72,6 +72,12 @@ public:
 		return RegisterLayoutExtensions;
 	}
 
+	DECLARE_EVENT_OneParam(FStateTreeEditorModule, FValidateStateTree, TNotNull<UStateTree*> StateTree);
+	FValidateStateTree& OnValidateStateTree()
+	{
+		return ValidateStateTree;
+	}
+
 	DECLARE_EVENT_OneParam(FStateTreeEditorModule, FPostInternalCompile, const UE::StateTree::Compiler::FPostInternalContext&);
 	/**
 	 * Handle post internal compilation for all state tree assets.
@@ -121,6 +127,8 @@ protected:
 	UE_DEPRECATED(5.7, "OnUserDefinedStructReinstancedHandle is not used.")
 	FDelegateHandle OnUserDefinedStructReinstancedHandle;
 	FOnRegisterLayoutExtensions RegisterLayoutExtensions;
+
+	FValidateStateTree ValidateStateTree;
 	FPostInternalCompile PostInternalCompile;
 };
 
