[StateTree] Implement EvaluationScopeInstanceData for condition's EnterState/ExitState/StateCompleted.

#FailedValidations Preflight

diff --git a/Engine/Plugins/Runtime/StateTree/Source/StateTreeModule/Private/StateTreeExecutionContext.cpp b/Engine/Plugins/Runtime/StateTree/Source/StateTreeModule/Private/StateTreeExecutionContext.cpp
--- a/Engine/Plugins/Runtime/StateTree/Source/StateTreeModule/Private/StateTreeExecutionContext.cpp
+++ b/Engine/Plugins/Runtime/StateTree/Source/StateTreeModule/Private/StateTreeExecutionContext.cpp
@@ -242,13 +242,18 @@ namespace UE::StateTree::ExecutionContext::Private
 			});
 	}
 
+	bool UseEvaluationScopeIntanceData(const FStateTreeDataHandle& DataHandle)
+	{
+		return DataHandle.GetSource() == EStateTreeDataSourceType::EvaluationScopeInstanceData
+			|| DataHandle.GetSource() == EStateTreeDataSourceType::EvaluationScopeInstanceDataObject;
+	}
+
 	void InitEvaluationScopeInstanceData(UE::StateTree::InstanceData::FEvaluationScopeInstanceContainer& Container, TNotNull<const UStateTree*> StateTree, const int32 NodesBegin, const int32 NodesEnd)
 	{
-		for (int32 NodeIndex = NodesBegin; NodeIndex != NodesEnd; NodeIndex++)
+		for (int32 NodeIndex = NodesBegin; NodeIndex != NodesEnd; ++NodeIndex)
 		{
 			const FStateTreeNodeBase& Node = StateTree->GetNodes()[NodeIndex].Get<const FStateTreeNodeBase>();
-			if (Node.InstanceDataHandle.GetSource() == EStateTreeDataSourceType::EvaluationScopeInstanceData
-				|| Node.InstanceDataHandle.GetSource() == EStateTreeDataSourceType::EvaluationScopeInstanceDataObject)
+			if (UseEvaluationScopeIntanceData(Node.InstanceDataHandle))
 			{
 				Container.Add(Node.InstanceDataHandle, StateTree->GetDefaultEvaluationScopeInstanceData().GetStruct(Node.InstanceTemplateIndex.Get()));
 			}
@@ -3518,41 +3523,10 @@ EStateTreeRunStatus FStateTreeExecutionContext::EnterState(const TSharedPtr<FSel
 				*UEnum::GetDisplayValueAsText(CurrentTransition.ChangeType).ToString()
 			);
 
-			// @todo: this needs to support EvaluationData Scope 
 			// Call state change events on conditions if needed.
-			if (CurrentState.bHasStateChangeConditions)
-			{
-				const int32 EnterConditionsEnd = CurrentState.EnterConditionsBegin + CurrentState.EnterConditionsNum;
-				for (int32 ConditionIndex = CurrentState.EnterConditionsBegin; ConditionIndex < EnterConditionsEnd; ++ConditionIndex)
-				{
-					const FStateTreeConditionBase& Cond = CurrentStateTree->Nodes[ConditionIndex].Get<const FStateTreeConditionBase>();
-					if (Cond.bHasShouldCallStateChangeEvents)
-					{
-						const bool bShouldCallEnterState = CurrentTransition.ChangeType == EStateTreeStateChangeType::Changed
-							|| (CurrentTransition.ChangeType == EStateTreeStateChangeType::Sustained && Cond.bShouldStateChangeOnReselect);
-
-						if (bShouldCallEnterState)
-						{
-							const FStateTreeDataView ConditionInstanceView = GetDataView(CurrentParentFrame, CurrentFrame, Cond.InstanceDataHandle);
-							FNodeInstanceDataScope DataScope(*this, &Cond, ConditionIndex, Cond.InstanceDataHandle, ConditionInstanceView);
-
-							if (Cond.BindingsBatch.IsValid())
-							{
-								CopyBatchOnActiveInstances(CurrentParentFrame, CurrentFrame, ConditionInstanceView, Cond.BindingsBatch);
-							}
-
-							UE_STATETREE_DEBUG_CONDITION_ENTER_STATE(this, CurrentFrame.StateTree, FStateTreeIndex16(ConditionIndex));
-							Cond.EnterState(*this, CurrentTransition);
-
-							// Reset copied properties that might contain object references.
-							if (Cond.BindingsBatch.IsValid())
-							{
-								CurrentFrame.StateTree->PropertyBindings.Super::ResetObjects(Cond.BindingsBatch, ConditionInstanceView);
-							}
-						}
-					}
-				}
-			}
+			constexpr bool bOnActiveInstances = true;
+			constexpr EConditionFunction ConditionFunction = EConditionFunction::EnterState;
+			RunAllConditionsInternal<bOnActiveInstances, ConditionFunction>(CurrentParentFrame, CurrentFrame, CurrentState, CurrentTransition);
 
 			// Activate tasks on current state.
 			UE::StateTree::FTasksCompletionStatus CurrentStateTasksStatus = CurrentFrame.ActiveTasksStatus.GetStatus(CurrentState);
@@ -3791,40 +3765,10 @@ void FStateTreeExecutionContext::ExitState(const TSharedPtr<const FSelectStateRe
 				CurrentFrame.ActiveNodeIndex = FStateTreeIndex16(TaskIndex - 1);
 			}
 
-			// @todo: this needs to support EvaluationScoped Data
 			// Call state change events on conditions if needed.
-			if (CurrentState.bHasStateChangeConditions)
-			{
-				for (int32 ConditionIndex = (CurrentState.EnterConditionsBegin + CurrentState.EnterConditionsNum) - 1; ConditionIndex >= CurrentState.EnterConditionsBegin; --ConditionIndex)
-				{
-					const FStateTreeConditionBase& Cond = CurrentFrame.StateTree->Nodes[ConditionIndex].Get<const FStateTreeConditionBase>();
-					if (Cond.bHasShouldCallStateChangeEvents)
-					{
-						const bool bShouldCallStateChange = CurrentTransition.ChangeType == EStateTreeStateChangeType::Changed
-							|| (CurrentTransition.ChangeType == EStateTreeStateChangeType::Sustained && Cond.bShouldStateChangeOnReselect);
-
-						if (bShouldCallStateChange)
-						{
-							const FStateTreeDataView ConditionInstanceView = GetDataView(CurrentParentFrame, CurrentFrame, Cond.InstanceDataHandle);
-							FNodeInstanceDataScope DataScope(*this, &Cond, ConditionIndex, Cond.InstanceDataHandle, ConditionInstanceView);
-
-							if (Cond.BindingsBatch.IsValid())
-							{
-								CopyBatchOnActiveInstances(CurrentParentFrame, CurrentFrame, ConditionInstanceView, Cond.BindingsBatch);
-							}
-
-							UE_STATETREE_DEBUG_CONDITION_EXIT_STATE(this, CurrentFrame.StateTree, FStateTreeIndex16(ConditionIndex));
-							Cond.ExitState(*this, CurrentTransition);
-
-							// Reset copied properties that might contain object references.
-							if (Cond.BindingsBatch.IsValid())
-							{
-								CurrentFrame.StateTree->PropertyBindings.Super::ResetObjects(Cond.BindingsBatch, ConditionInstanceView);
-							}
-						}
-					}
-				}
-			}
+			constexpr bool bOnActiveInstances = true;
+			constexpr EConditionFunction ConditionFunction = EConditionFunction::ExitState;
+			RunAllConditionsInternal<bOnActiveInstances, ConditionFunction>(CurrentParentFrame, CurrentFrame, CurrentState, CurrentTransition);
 
 			// Reset the completed state. This is to keep the wrong UE5.6 behavior.
 			if (!EnumHasAnyFlags(CurrentFrame.StateTree->GetStateSelectionRules(), EStateTreeStateSelectionRules::CompletedTransitionStatesCreateNewStates)
@@ -3913,33 +3857,10 @@ void FStateTreeExecutionContext::StateCompleted()
 				}
 			}
 
-			// @todo: this needs to support EvaluationScopedData
 			// Call state change events on conditions if needed.
-			if (State.bHasStateChangeConditions)
-			{
-				for (int32 ConditionIndex = (State.EnterConditionsBegin + State.EnterConditionsNum) - 1; ConditionIndex >= State.EnterConditionsBegin; ConditionIndex--)
-				{
-					const FStateTreeConditionBase& Cond = CurrentFrame.StateTree->Nodes[ConditionIndex].Get<const FStateTreeConditionBase>();
-					if (Cond.bHasShouldCallStateChangeEvents)
-					{
-						const FStateTreeDataView ConditionInstanceView = GetDataView(CurrentParentFrame, CurrentFrame, Cond.InstanceDataHandle);
-						FNodeInstanceDataScope DataScope(*this, &Cond, ConditionIndex, Cond.InstanceDataHandle, ConditionInstanceView);
-
-						if (Cond.BindingsBatch.IsValid())
-						{
-							CopyBatchOnActiveInstances(CurrentParentFrame, CurrentFrame, ConditionInstanceView, Cond.BindingsBatch);
-						}
-
-						Cond.StateCompleted(*this, Exec.LastTickStatus, CurrentFrame.ActiveStates);
-
-						// Reset copied properties that might contain object references.
-						if (Cond.BindingsBatch.IsValid())
-						{
-							CurrentFrame.StateTree->PropertyBindings.Super::ResetObjects(Cond.BindingsBatch, ConditionInstanceView);
-						}
-					}
-				}
-			}
+			constexpr bool bOnActiveInstances = true;
+			constexpr EConditionFunction ConditionFunction = EConditionFunction::StateCompleted;
+			RunAllConditionsInternal<bOnActiveInstances, ConditionFunction>(CurrentParentFrame, CurrentFrame, State, FStateTreeTransitionResult());
 		}
 	}
 }
@@ -4864,7 +4785,7 @@ bool FStateTreeExecutionContext::TestAllConditionsInternal(const FStateTreeExecu
 			UE_STATETREE_DEBUG_CONDITION_EVENT(this, ConditionIndex, ConditionInstanceView, bValue ? EStateTreeTraceEventType::Passed : EStateTreeTraceEventType::Failed);
 
 			// Reset copied properties that might contain object references.
-			if (Cond.BindingsBatch.IsValid())
+			if (Cond.BindingsBatch.IsValid() && !UseEvaluationScopeIntanceData(Cond.InstanceDataHandle))
 			{
 				CurrentFrame.StateTree->PropertyBindings.Super::ResetObjects(Cond.BindingsBatch, ConditionInstanceView);
 			}
@@ -4922,6 +4843,94 @@ bool FStateTreeExecutionContext::TestAllConditionsInternal(const FStateTreeExecu
 	return Values[0];
 }
 
+template<bool bOnActiveInstances, FStateTreeExecutionContext::EConditionFunction ConditionFunction>
+void FStateTreeExecutionContext::RunAllConditionsInternal(const FStateTreeExecutionFrame* CurrentParentFrame, const FStateTreeExecutionFrame& CurrentFrame, const FCompactStateTreeState& CurrentState, const FStateTreeTransitionResult& CurrentTransition)
+{
+	using namespace UE::StateTree;
+	using namespace UE::StateTree::ExecutionContext;
+	using namespace UE::StateTree::ExecutionContext::Private;
+	using namespace UE::StateTree::InstanceData;
+
+	if (CurrentState.bHasStateChangeConditions)
+	{
+		FStateTreeExecutionState& Exec = GetExecState();
+
+		FEvaluationScopeInstanceContainer EvaluationScopeContainer;
+		bool bEvaluationScopeContainerPushed = false;
+		if (CurrentState.EnterConditionEvaluationScopeMemoryRequirement.Size > 0)
+		{
+			void* ScopedConditionMemory = FMemory_Alloca_Aligned(CurrentState.EnterConditionEvaluationScopeMemoryRequirement.Size, CurrentState.EnterConditionEvaluationScopeMemoryRequirement.Alignment);
+			EvaluationScopeContainer = FEvaluationScopeInstanceContainer(ScopedConditionMemory, CurrentState.EnterConditionEvaluationScopeMemoryRequirement);
+			PushEvaluationScopeInstanceContainer(EvaluationScopeContainer, CurrentFrame);
+			InitEvaluationScopeInstanceData(EvaluationScopeContainer, CurrentFrame.StateTree, CurrentState.EnterConditionsBegin, CurrentState.EnterConditionsBegin + CurrentState.EnterConditionsNum);
+			bEvaluationScopeContainerPushed = true;
+		}
+
+		const int32 EnterConditionsEnd = CurrentState.EnterConditionsBegin + CurrentState.EnterConditionsNum;
+		constexpr bool bForwardIteration = ConditionFunction == EConditionFunction::EnterState;
+
+		int32 ConditionIndex = bForwardIteration ? CurrentState.EnterConditionsBegin : EnterConditionsEnd - 1;
+		bool bRunLoop = bForwardIteration ? ConditionIndex < EnterConditionsEnd : ConditionIndex >= CurrentState.EnterConditionsBegin;
+		while (bRunLoop)
+		{
+			const FStateTreeConditionBase& Cond = CurrentFrame.StateTree->Nodes[ConditionIndex].Get<const FStateTreeConditionBase>();
+			if (Cond.bHasShouldCallStateChangeEvents)
+			{
+				constexpr bool bAlwaysCallEnterState = ConditionFunction == EConditionFunction::StateCompleted;
+				const bool bShouldCallEnterState = CurrentTransition.ChangeType == EStateTreeStateChangeType::Changed
+					|| (CurrentTransition.ChangeType == EStateTreeStateChangeType::Sustained && Cond.bShouldStateChangeOnReselect);
+
+				if (bAlwaysCallEnterState || bShouldCallEnterState)
+				{
+					const FStateTreeDataView ConditionInstanceView = GetDataView(CurrentParentFrame, CurrentFrame, Cond.InstanceDataHandle);
+					FNodeInstanceDataScope DataScope(*this, &Cond, ConditionIndex, Cond.InstanceDataHandle, ConditionInstanceView);
+
+					if (Cond.BindingsBatch.IsValid())
+					{
+						if constexpr (bOnActiveInstances)
+						{
+							CopyBatchOnActiveInstances(CurrentParentFrame, CurrentFrame, ConditionInstanceView, Cond.BindingsBatch);
+						}
+						else
+						{
+							CopyBatchWithValidation(CurrentParentFrame, CurrentFrame, ConditionInstanceView, Cond.BindingsBatch);
+						}
+					}
+
+					if constexpr (ConditionFunction == EConditionFunction::EnterState)
+					{
+						UE_STATETREE_DEBUG_CONDITION_ENTER_STATE(this, CurrentFrame.StateTree, FStateTreeIndex16(ConditionIndex));
+						Cond.EnterState(*this, CurrentTransition);
+					}
+					else if constexpr (ConditionFunction == EConditionFunction::ExitState)
+					{
+						UE_STATETREE_DEBUG_CONDITION_EXIT_STATE(this, CurrentFrame.StateTree, FStateTreeIndex16(ConditionIndex));
+						Cond.ExitState(*this, CurrentTransition);
+					}
+					else if constexpr (ConditionFunction == EConditionFunction::StateCompleted)
+					{
+						Cond.StateCompleted(*this, Exec.LastTickStatus, CurrentFrame.ActiveStates);
+					}
+
+					// Reset copied properties that might contain object references.
+					if (Cond.BindingsBatch.IsValid() && !UseEvaluationScopeIntanceData(Cond.InstanceDataHandle))
+					{
+						CurrentFrame.StateTree->PropertyBindings.Super::ResetObjects(Cond.BindingsBatch, ConditionInstanceView);
+					}
+				}
+			}
+
+			ConditionIndex += bForwardIteration ? 1 : -1;
+			bRunLoop = bForwardIteration ? ConditionIndex < EnterConditionsEnd : ConditionIndex >= CurrentState.EnterConditionsBegin;
+		}
+
+		if (bEvaluationScopeContainerPushed)
+		{
+			PopEvaluationScopeInstanceContainer(EvaluationScopeContainer);
+		}
+	}
+}
+
 //Deprecated
 float FStateTreeExecutionContext::EvaluateUtility(const FStateTreeExecutionFrame* CurrentParentFrame, const FStateTreeExecutionFrame& CurrentFrame, const int32 ConsiderationsBegin, const int32 ConsiderationsNum, const float StateWeight)
 {
@@ -4985,7 +4994,7 @@ float FStateTreeExecutionContext::EvaluateUtilityWithValidation(const FStateTree
 		Value = Consideration.GetNormalizedScore(*this);
 
 		// Reset copied properties that might contain object references.
-		if (Consideration.BindingsBatch.IsValid())
+		if (Consideration.BindingsBatch.IsValid() && !UseEvaluationScopeIntanceData(Consideration.InstanceDataHandle))
 		{
 			CurrentFrame.StateTree->PropertyBindings.Super::ResetObjects(Consideration.BindingsBatch, ConsiderationInstanceView);
 		}
diff --git a/Engine/Plugins/Runtime/StateTree/Source/StateTreeModule/Public/StateTreeExecutionContext.h b/Engine/Plugins/Runtime/StateTree/Source/StateTreeModule/Public/StateTreeExecutionContext.h
--- a/Engine/Plugins/Runtime/StateTree/Source/StateTreeModule/Public/StateTreeExecutionContext.h
+++ b/Engine/Plugins/Runtime/StateTree/Source/StateTreeModule/Public/StateTreeExecutionContext.h
@@ -1102,6 +1102,15 @@ private:
 	template<bool bOnActiveInstances>
 	bool TestAllConditionsInternal(const FStateTreeExecutionFrame* CurrentParentFrame, const FStateTreeExecutionFrame& CurrentFrame, FStateTreeStateHandle CurrentStateHandle, const FMemoryRequirement& MemoryRequirement, const int32 ConditionsOffset, const int32 ConditionsNum, EStateTreeUpdatePhase Phase);
 
+	enum class EConditionFunction
+	{
+		EnterState,
+		ExitState,
+		StateCompleted,
+	};
+	template<bool bOnActiveInstances, EConditionFunction ConditionFunction>
+	void RunAllConditionsInternal(const FStateTreeExecutionFrame* CurrentParentFrame, const FStateTreeExecutionFrame& CurrentFrame, const FCompactStateTreeState& CurrentState, const FStateTreeTransitionResult& CurrentTransition);
+
 protected:
 	/**
 	 * Checks all conditions at given range.
