[StateTree] Add callback that handle pre-validation. Runs before the engine validation.

#rnx

diff --git a/Engine/Plugins/Runtime/StateTree/Source/StateTreeEditorModule/Private/StateTreeEditingSubsystem.cpp b/Engine/Plugins/Runtime/StateTree/Source/StateTreeEditorModule/Private/StateTreeEditingSubsystem.cpp
--- a/Engine/Plugins/Runtime/StateTree/Source/StateTreeEditorModule/Private/StateTreeEditingSubsystem.cpp
+++ b/Engine/Plugins/Runtime/StateTree/Source/StateTreeEditorModule/Private/StateTreeEditingSubsystem.cpp
@@ -285,31 +285,53 @@ void UStateTreeEditingSubsystem::ValidateStateTree(TNotNull<UStateTree*> InState
 		UStateTreeEditorData* EditorData = CastChecked<UStateTreeEditorData>(InStateTree->EditorData);
 		FixEditorSchema(EditorData);
 
+		// Notify the module, schema, and extension. Use that order to go from the less specific to the more specific.
+		{
+			FStateTreeEditorModule::GetModule().OnPreValidateStateTree().Broadcast(InStateTree);
+
+			// Schema validation
+			if (EditorData->EditorSchema)
+			{
+				EditorData->EditorSchema->PreValidate(InStateTree);
+			}
+
+			TArray<UStateTreeEditorDataExtension*> LocalExtensions = EditorData->Extensions;
+			for (UStateTreeEditorDataExtension* Extension : LocalExtensions)
+			{
+				if (Extension)
+				{
+					Extension->PreValidate();
+				}
+			}
+		}
+
 		EditorData->ReparentStates();
 		EditorData->FixDuplicateIDs();
-		EditorData->UpdateBindingsInstanceStructs();
 
 		RemoveUnusedBindings(EditorData);
+		EditorData->UpdateBindingsInstanceStructs();
 		ValidateLinkedStates(EditorData);
 		UpdateLinkedStateParameters(EditorData);
 		UpdateTransactionalFlags(EditorData);
 		RemoveEmptyExtension(EditorData);
 
-		// Notify the module, schema, and extension. Use that order to go from the less specific to the more specific.
-		FStateTreeEditorModule::GetModule().OnValidateStateTree().Broadcast(InStateTree);
-
-		// Schema validation
-		if (EditorData->EditorSchema)
+		// Notify the module, schema, and extension.
 		{
-			EditorData->EditorSchema->Validate(InStateTree);
-		}
+			FStateTreeEditorModule::GetModule().OnValidateStateTree().Broadcast(InStateTree);
 
-		TArray<UStateTreeEditorDataExtension*> LocalExtensions = EditorData->Extensions;
-		for (UStateTreeEditorDataExtension* Extension : LocalExtensions)
-		{
-			if (Extension)
+			// Schema validation
+			if (EditorData->EditorSchema)
 			{
-				Extension->Validate();
+				EditorData->EditorSchema->Validate(InStateTree);
+			}
+
+			TArray<UStateTreeEditorDataExtension*> LocalExtensions = EditorData->Extensions;
+			for (UStateTreeEditorDataExtension* Extension : LocalExtensions)
+			{
+				if (Extension)
+				{
+					Extension->Validate();
+				}
 			}
 		}
 		RemoveEmptyExtension(EditorData);
diff --git a/Engine/Plugins/Runtime/StateTree/Source/StateTreeEditorModule/Public/StateTreeEditorDataExtension.h b/Engine/Plugins/Runtime/StateTree/Source/StateTreeEditorModule/Public/StateTreeEditorDataExtension.h
--- a/Engine/Plugins/Runtime/StateTree/Source/StateTreeEditorModule/Public/StateTreeEditorDataExtension.h
+++ b/Engine/Plugins/Runtime/StateTree/Source/StateTreeEditorModule/Public/StateTreeEditorDataExtension.h
@@ -23,6 +23,15 @@ class UStateTreeEditorDataExtension : public UObject
 	GENERATED_BODY()
 
 public:
+	/**
+	 * Validates and applies the editor data extension restrictions on the state tree owned by the extension.
+	 * This is before the engine validation. The Schema and EditorData will be correctly set.
+	 * You are allowed to modify the editor data.
+	 */
+	virtual void PreValidate()
+	{
+	}
+
 	/**
 	 * Validates and applies the editor data extension restrictions on the state tree owned by the extension.
 	 * Also serves as the "safety net" of fixing up editor data following an editor operation.
diff --git a/Engine/Plugins/Runtime/StateTree/Source/StateTreeEditorModule/Public/StateTreeEditorModule.h b/Engine/Plugins/Runtime/StateTree/Source/StateTreeEditorModule/Public/StateTreeEditorModule.h
--- a/Engine/Plugins/Runtime/StateTree/Source/StateTreeEditorModule/Public/StateTreeEditorModule.h
+++ b/Engine/Plugins/Runtime/StateTree/Source/StateTreeEditorModule/Public/StateTreeEditorModule.h
@@ -73,6 +73,25 @@ public:
 	}
 
 	DECLARE_EVENT_OneParam(FStateTreeEditorModule, FValidateStateTree, TNotNull<UStateTree*> StateTree);
+	/**
+	 * Handle the validation a state tree asset.
+	 * This is before the engine validation. The Schema and EditorData will be correctly set.
+	 * You are allowed to modify the editor data.
+	 * @note Use the UStateTreeEditorDataExtension::PreValidate for controlling a single asset.
+	 * @note Use the UStateTreeEditorSchema::PreValidate for controlling a type of state tree asset.
+	 */
+	FValidateStateTree& OnPreValidateStateTree()
+	{
+		return PreValidateStateTree;
+	}
+
+	/**
+	 * Handle the validation a state tree asset.
+	 * Also serves as the "safety net" of fixing up editor data following an editor operation.
+	 * You are allowed to modify the editor data.
+	 * @note Use the UStateTreeEditorDataExtension::Validate for controlling a single asset.
+	 * @note Use the UStateTreeEditorSchema::Validate for controlling a type of state tree asset.
+	 */
 	FValidateStateTree& OnValidateStateTree()
 	{
 		return ValidateStateTree;
@@ -82,7 +101,7 @@ public:
 	/**
 	 * Handle post internal compilation for all state tree assets.
 	 * The state tree asset compiled successfully.
-	 * @note Use the UStateTreeEditorExtension::HandlePostInternalCompile for controlling a single asset.
+	 * @note Use the UStateTreeEditorDataExtension::HandlePostInternalCompile for controlling a single asset.
 	 * @note Use the UStateTreeEditorSchema::HandlePostInternalCompile for controlling a type of state tree asset.
 	 */
 	FPostInternalCompile& OnPostInternalCompile()
@@ -128,6 +147,7 @@ protected:
 	FDelegateHandle OnUserDefinedStructReinstancedHandle;
 	FOnRegisterLayoutExtensions RegisterLayoutExtensions;
 
+	FValidateStateTree PreValidateStateTree;
 	FValidateStateTree ValidateStateTree;
 	FPostInternalCompile PostInternalCompile;
 };
diff --git a/Engine/Plugins/Runtime/StateTree/Source/StateTreeEditorModule/Public/StateTreeEditorSchema.h b/Engine/Plugins/Runtime/StateTree/Source/StateTreeEditorModule/Public/StateTreeEditorSchema.h
--- a/Engine/Plugins/Runtime/StateTree/Source/StateTreeEditorModule/Public/StateTreeEditorSchema.h
+++ b/Engine/Plugins/Runtime/StateTree/Source/StateTreeEditorModule/Public/StateTreeEditorSchema.h
@@ -29,6 +29,14 @@ public:
 	}
 
 public:
+	/*
+	 * Validates and applies the schema restrictions on the StateTree.
+	 * This is before the engine validation. The EditorData will be correctly set.
+	 */
+	UE_API virtual void PreValidate(TNotNull<UStateTree*> StateTree)
+	{
+	}
+
 	/*
 	 * Validates and applies the schema restrictions on the StateTree.
 	 * Also serves as the "safety net" of fixing up editor data following an editor operation.
