Add callback to get the asset registry from the editor data. Use the editor schema to always have a valid schema in case the compilation fails.

diff --git a/Engine/Plugins/Runtime/StateTree/Source/StateTreeEditorModule/Private/StateTreeEditorModule.cpp b/Engine/Plugins/Runtime/StateTree/Source/StateTreeEditorModule/Private/StateTreeEditorModule.cpp
--- a/Engine/Plugins/Runtime/StateTree/Source/StateTreeEditorModule/Private/StateTreeEditorModule.cpp
+++ b/Engine/Plugins/Runtime/StateTree/Source/StateTreeEditorModule/Private/StateTreeEditorModule.cpp
@@ -29,14 +29,17 @@
 #include "StateTreeCompilerLog.h"
 #include "StateTreeCompilerManager.h"
 #include "StateTreeDelegates.h"
+#include "StateTreeDelegatesInternal.h"
 #include "StateTreeEditingSubsystem.h"
 #include "StateTreeEditor.h"
 #include "StateTreeEditorCommands.h"
 #include "StateTreeEditorData.h"
+#include "StateTreeEditorDataExtension.h"
 #include "StateTreeEditorSchema.h"
 #include "StateTreeEditorStyle.h"
 #include "StateTreeNodeClassCache.h"
 #include "StateTreePropertyFunctionBase.h"
+#include "UObject/AssetRegistryTagsContext.h"
 
 #define LOCTEXT_NAMESPACE "StateTreeEditor"
 
@@ -83,7 +86,11 @@ FStateTreeEditorModule* FStateTreeEditorModule::GetModulePtr()
 
 void FStateTreeEditorModule::StartupModule()
 {
-	UE::StateTree::Delegates::OnRequestEditorHash.BindLambda([](const UStateTree& InStateTree) -> uint32 { return UStateTreeEditingSubsystem::CalculateStateTreeHash(&InStateTree); });
+	UE::StateTree::Delegates::OnRequestEditorHash.BindLambda([](const UStateTree& InStateTree) -> uint32
+		{
+			return UStateTreeEditingSubsystem::CalculateStateTreeHash(&InStateTree);
+		});
+	UE::StateTree::Delegates::Private::OnRequestAssetRegistryTags.BindRaw(this, &FStateTreeEditorModule::HandleRequestAssetRegistryTags);
 	UE::StateTree::Compiler::FCompilerManager::Startup();
 
 	OnPostEngineInitHandle = FCoreDelegates::OnPostEngineInit.AddLambda([this]()
@@ -165,6 +172,7 @@ void FStateTreeEditorModule::ShutdownModule()
 	NodeClassCache.Reset();
 
 	UE::StateTree::Compiler::FCompilerManager::Shutdown();
+	UE::StateTree::Delegates::Private::OnRequestAssetRegistryTags.Unbind();
 	UE::StateTree::Delegates::OnRequestEditorHash.Unbind();
 }
 
@@ -187,6 +195,31 @@ TSharedPtr<FStateTreeNodeClassCache> FStateTreeEditorModule::GetNodeClassCache()
 	return NodeClassCache;
 }
 
+void FStateTreeEditorModule::HandleRequestAssetRegistryTags(TNotNull<const UStateTree*> InStateTree, FAssetRegistryTagsContext& Context)
+{
+	// Notify the module, schema, and extension. Use that order to go from the less specific to the more specific.
+	OnAssetRegistryTags().Broadcast(InStateTree, Context);
+
+	if (UStateTreeEditorData* EditorData = Cast<UStateTreeEditorData>(InStateTree->EditorData))
+	{
+		if (EditorData->Schema)
+		{
+			EditorData->Schema->GetAssetRegistryTags(Context);
+		}
+		for (const UStateTreeEditorDataExtension* Extension : EditorData->Extensions)
+		{
+			if (Extension)
+			{
+				Extension->GetAssetRegistryTags(Context);
+			}
+		}
+
+		// Add the schema tag (override previously set value).
+		const FString SchemaClassName = EditorData->Schema ? EditorData->Schema->GetClass()->GetPathName() : TEXT("");
+		Context.AddTag(UObject::FAssetRegistryTag(UE::StateTree::SchemaTag, SchemaClassName, UObject::FAssetRegistryTag::TT_Alphabetical));
+	}
+}
+
 bool FStateTreeEditorModule::FEditorTypes::HasData() const
 {
 	return EditorData.IsValid() || EditorSchema.IsValid();
diff --git a/Engine/Plugins/Runtime/StateTree/Source/StateTreeEditorModule/Public/StateTreeEditorModule.h b/Engine/Plugins/Runtime/StateTree/Source/StateTreeEditorModule/Public/StateTreeEditorModule.h
--- a/Engine/Plugins/Runtime/StateTree/Source/StateTreeEditorModule/Public/StateTreeEditorModule.h
+++ b/Engine/Plugins/Runtime/StateTree/Source/StateTreeEditorModule/Public/StateTreeEditorModule.h
@@ -10,6 +10,7 @@
 
 #define UE_API STATETREEEDITORMODULE_API
 
+class FAssetRegistryTagsContext;
 class UStateTree;
 class UStateTreeEditorData;
 class UStateTreeEditorSchema;
@@ -72,6 +73,13 @@ public:
 		return RegisterLayoutExtensions;
 	}
 
+	DECLARE_EVENT_TwoParams(FStateTreeEditorModule, FOnAssetRegistryTags, TNotNull<const UStateTree*> /*StateTree*/, FAssetRegistryTagsContext& /*Context*/);
+	/** Handle GetAssetRegistryTags for a state tree asset. */
+	FOnAssetRegistryTags& OnAssetRegistryTags()
+	{
+		return AssetRegistryTags;
+	}
+
 	DECLARE_EVENT_OneParam(FStateTreeEditorModule, FValidateStateTree, TNotNull<UStateTree*> StateTree);
 	/**
 	 * Handle the validation a state tree asset.
@@ -99,7 +107,7 @@ public:
 
 	DECLARE_EVENT_OneParam(FStateTreeEditorModule, FPostInternalCompile, const UE::StateTree::Compiler::FPostInternalContext&);
 	/**
-	 * Handle post internal compilation for all state tree assets.
+	 * Handle post internal compilation for a state tree asset.
 	 * The state tree asset compiled successfully.
 	 * @note Use the UStateTreeEditorDataExtension::HandlePostInternalCompile for controlling a single asset.
 	 * @note Use the UStateTreeEditorSchema::HandlePostInternalCompile for controlling a type of state tree asset.
@@ -120,9 +128,12 @@ public:
 	void RegisterEditorSchemaClass(TNonNullSubclassOf<const UStateTreeSchema> Schema, TNonNullSubclassOf<const UStateTreeEditorSchema> EditorSchema);
 	/** Unregister the editor schema type for a specific schema. */
 	void UnregisterEditorSchemaClass(TNonNullSubclassOf<const UStateTreeSchema> Schema);
-	/** Get the editor data type for a specific schema. */
+	/** Get the editor schema type for a specific schema. */
 	TNonNullSubclassOf<UStateTreeEditorSchema> GetEditorSchemaClass(TNonNullSubclassOf<const UStateTreeSchema> Schema) const;
 
+private:
+	void HandleRequestAssetRegistryTags(TNotNull<const UStateTree*> InStateTree, FAssetRegistryTagsContext& Context);
+
 protected:
 	TSharedPtr<FExtensibilityManager> MenuExtensibilityManager;
 	TSharedPtr<FExtensibilityManager> ToolBarExtensibilityManager;
@@ -147,6 +158,7 @@ protected:
 	FDelegateHandle OnUserDefinedStructReinstancedHandle;
 	FOnRegisterLayoutExtensions RegisterLayoutExtensions;
 
+	FOnAssetRegistryTags AssetRegistryTags;
 	FValidateStateTree PreValidateStateTree;
 	FValidateStateTree ValidateStateTree;
 	FPostInternalCompile PostInternalCompile;
diff --git a/Engine/Plugins/Runtime/StateTree/Source/StateTreeModule/Internal/StateTreeDelegatesInternal.h b/Engine/Plugins/Runtime/StateTree/Source/StateTreeModule/Internal/StateTreeDelegatesInternal.h
new file mode 100644
--- /dev/null
+++ b/Engine/Plugins/Runtime/StateTree/Source/StateTreeModule/Internal/StateTreeDelegatesInternal.h
@@ -0,0 +1,21 @@
+// Copyright Epic Games, Inc. All Rights Reserved.
+
+#pragma once
+
+#include "Delegates/Delegate.h"
+
+#define UE_API STATETREEMODULE_API
+
+class UStateTree;
+class FAssetRegistryTagsContext;
+
+namespace UE::StateTree::Delegates::Private
+{
+
+/** A callback for UStateTree::GetAssetRegistryTags. */
+DECLARE_DELEGATE_TwoParams(FOnStateTreeRequestAssetRegistryTags, TNotNull<const UStateTree*> /*StateTree*/, FAssetRegistryTagsContext& /*Context*/);
+extern UE_API FOnStateTreeRequestAssetRegistryTags OnRequestAssetRegistryTags;
+
+}; // UE::StateTree::Delegates
+
+#undef UE_API
diff --git a/Engine/Plugins/Runtime/StateTree/Source/StateTreeModule/Private/StateTree.cpp b/Engine/Plugins/Runtime/StateTree/Source/StateTreeModule/Private/StateTree.cpp
--- a/Engine/Plugins/Runtime/StateTree/Source/StateTreeModule/Private/StateTree.cpp
+++ b/Engine/Plugins/Runtime/StateTree/Source/StateTreeModule/Private/StateTree.cpp
@@ -12,6 +12,7 @@
 #include "AssetRegistry/AssetData.h"
 #include "Misc/ScopeRWLock.h"
 #include "StateTreeDelegates.h"
+#include "StateTreeDelegatesInternal.h"
 #include "Logging/LogScopedVerbosityOverride.h"
 #include "Misc/DataValidation.h"
 #include "Misc/EnumerateRange.h"
@@ -299,14 +300,7 @@ void UStateTree::PostDuplicate(EDuplicateMode::Type DuplicateMode)
 
 void UStateTree::GetAssetRegistryTags(FAssetRegistryTagsContext Context) const
 {
-	const FString SchemaClassName = Schema ? Schema->GetClass()->GetPathName() : TEXT("");
-	Context.AddTag(FAssetRegistryTag(UE::StateTree::SchemaTag, SchemaClassName, FAssetRegistryTag::TT_Alphabetical));
-
-	if (Schema)
-	{
-		Schema->GetAssetRegistryTags(Context);
-	}
-
+	UE::StateTree::Delegates::Private::OnRequestAssetRegistryTags.ExecuteIfBound(this, Context);
 	Super::GetAssetRegistryTags(Context);
 }
 
diff --git a/Engine/Plugins/Runtime/StateTree/Source/StateTreeModule/Private/StateTreeDelegates.cpp b/Engine/Plugins/Runtime/StateTree/Source/StateTreeModule/Private/StateTreeDelegates.cpp
--- a/Engine/Plugins/Runtime/StateTree/Source/StateTreeModule/Private/StateTreeDelegates.cpp
+++ b/Engine/Plugins/Runtime/StateTree/Source/StateTreeModule/Private/StateTreeDelegates.cpp
@@ -1,30 +1,38 @@
 // Copyright Epic Games, Inc. All Rights Reserved.
 
 #include "StateTreeDelegates.h"
+#include "StateTreeDelegatesInternal.h"
 
 namespace UE::StateTree::Delegates
 {
 
 #if WITH_EDITOR
-FOnIdentifierChanged OnIdentifierChanged;
-FOnSchemaChanged OnSchemaChanged;
-FOnParametersChanged OnParametersChanged;
-FOnGlobalDataChanged OnGlobalDataChanged;
-FOnVisualThemeChanged OnVisualThemeChanged;
-FOnStateParametersChanged OnStateParametersChanged;
-FOnBreakpointsChanged OnBreakpointsChanged;
-FOnPostCompile OnPostCompile;
-FOnRequestCompile OnRequestCompile;
-FOnRequestEditorHash OnRequestEditorHash;
+	FOnIdentifierChanged OnIdentifierChanged;
+	FOnSchemaChanged OnSchemaChanged;
+	FOnParametersChanged OnParametersChanged;
+	FOnGlobalDataChanged OnGlobalDataChanged;
+	FOnVisualThemeChanged OnVisualThemeChanged;
+	FOnStateParametersChanged OnStateParametersChanged;
+	FOnBreakpointsChanged OnBreakpointsChanged;
+	FOnPostCompile OnPostCompile;
+	FOnRequestCompile OnRequestCompile;
+	FOnRequestEditorHash OnRequestEditorHash;
 #endif // WITH_EDITOR
 
 #if WITH_STATETREE_TRACE
-FOnTracingStateChanged OnTracingStateChanged;
+	FOnTracingStateChanged OnTracingStateChanged;
 #endif // WITH_STATETREE_TRACE
 
 #if WITH_STATETREE_TRACE_DEBUGGER
-FOnTraceAnalysisStateChanged OnTraceAnalysisStateChanged;
-FOnTracingTimelineScrubbed OnTracingTimelineScrubbed;
+	FOnTraceAnalysisStateChanged OnTraceAnalysisStateChanged;
+	FOnTracingTimelineScrubbed OnTracingTimelineScrubbed;
 #endif // WITH_STATETREE_TRACE_DEBUGGER
 
 }; // UE::StateTree::Delegates
+
+namespace UE::StateTree::Delegates::Private
+{
+#if WITH_EDITOR
+	FOnStateTreeRequestAssetRegistryTags OnRequestAssetRegistryTags;
+#endif
+} //UE::StateTree::Delegates::Private
