[StateTree] Fix bad access when the user doesn't set the ExternalGlobalParameters. User needs to call SetExternalGlobalParameters when requesting external global parameters.

#rnx

diff --git a/Engine/Plugins/Runtime/StateTree/Source/StateTreeModule/Private/StateTreeExecutionContext.cpp b/Engine/Plugins/Runtime/StateTree/Source/StateTreeModule/Private/StateTreeExecutionContext.cpp
--- a/Engine/Plugins/Runtime/StateTree/Source/StateTreeModule/Private/StateTreeExecutionContext.cpp
+++ b/Engine/Plugins/Runtime/StateTree/Source/StateTreeModule/Private/StateTreeExecutionContext.cpp
@@ -2939,7 +2939,8 @@ FStateTreeDataView FStateTreeExecutionContext::GetDataViewOrTemporary(const FSta
 	const FStateTreeDataHandle Handle = CopyInfo.SourceDataHandle.Get<FStateTreeDataHandle>();
 	if (Handle.GetSource() == EStateTreeDataSourceType::ExternalGlobalParameterData)
 	{
-		uint8* MemoryPtr = ExternalGlobalParameters->Find(CopyInfo);
+		ensureMsgf(ExternalGlobalParameters != nullptr, TEXT("User requested external global parameters without setting it. See SetExternalGlobalParameters."));
+		uint8* MemoryPtr = ExternalGlobalParameters ? ExternalGlobalParameters->Find(CopyInfo) : nullptr;
 		return FStateTreeDataView(CopyInfo.SourceStructType, MemoryPtr);
 	}
 
