[StateTree] (UndoBackout 45432306) Add ability to mark a transition for reselection.

Changes:
- Disable option when target state is not in nested / linked hierarchy
- Require flag to enable in editor schema
- Use Enum for reselection for future extensibility
- Fix crash when updating editor schema

#tests PIE, editor, automation
#was-backout 45432306

diff --git a/Engine/Plugins/Experimental/UAF/UAFStateTree/Source/UAFStateTreeEditor/Private/AnimNextStateTreeEditorModule.cpp b/Engine/Plugins/Experimental/UAF/UAFStateTree/Source/UAFStateTreeEditor/Private/AnimNextStateTreeEditorModule.cpp
--- a/Engine/Plugins/Experimental/UAF/UAFStateTree/Source/UAFStateTreeEditor/Private/AnimNextStateTreeEditorModule.cpp
+++ b/Engine/Plugins/Experimental/UAF/UAFStateTree/Source/UAFStateTreeEditor/Private/AnimNextStateTreeEditorModule.cpp
@@ -27,6 +27,7 @@
 #include "AnimNextStateTreeSchema.h"
 #include "StateTreeAssetCompilationHandler.h"
 #include "StateTreeEditorModule.h"
+#include "UAFStateTreeEditorSchema.h"
 
 #define LOCTEXT_NAMESPACE "AnimNextStateTreeEditorModule"
 
@@ -341,7 +342,8 @@ void FAnimNextStateTreeEditorModule::StartupModule()
 	WorkspaceEditorModule.RegisterWorkspaceItemDetails(Workspace::FOutlinerItemDetailsId(FAnimNextStateTreeStateOutlinerData::StaticStruct()->GetFName()), StaticCastSharedPtr<UE::Workspace::IWorkspaceOutlinerItemDetails>(StateItemDetails));
 
 	FStateTreeEditorModule& StateTreeEditorModule = FStateTreeEditorModule::GetModule();
-	StateTreeEditorModule.RegisterEditorDataClass(UStateTreeAnimNextSchema::StaticClass(), UAnimNextStateTreeTreeEditorData::StaticClass());		
+	StateTreeEditorModule.RegisterEditorDataClass(UStateTreeAnimNextSchema::StaticClass(), UAnimNextStateTreeTreeEditorData::StaticClass());
+	StateTreeEditorModule.RegisterEditorSchemaClass(UStateTreeAnimNextSchema::StaticClass(), UUAFStateTreeEditorSchema::StaticClass());
 }
 
 void FAnimNextStateTreeEditorModule::ShutdownModule()
@@ -363,7 +365,8 @@ void FAnimNextStateTreeEditorModule::ShutdownModule()
 
 		if(FStateTreeEditorModule* StateTreeEditorModule = FStateTreeEditorModule::GetModulePtr())
 		{
-			StateTreeEditorModule->UnregisterEditorDataClass(UStateTreeAnimNextSchema::StaticClass());			
+			StateTreeEditorModule->UnregisterEditorDataClass(UStateTreeAnimNextSchema::StaticClass());
+			StateTreeEditorModule->UnregisterEditorSchemaClass(UStateTreeAnimNextSchema::StaticClass());
 		}
     }
 }
diff --git a/Engine/Plugins/Experimental/UAF/UAFStateTree/Source/UAFStateTreeEditor/Private/UAFStateTreeEditorSchema.cpp b/Engine/Plugins/Experimental/UAF/UAFStateTree/Source/UAFStateTreeEditor/Private/UAFStateTreeEditorSchema.cpp
new file mode 100644
--- /dev/null
+++ b/Engine/Plugins/Experimental/UAF/UAFStateTree/Source/UAFStateTreeEditor/Private/UAFStateTreeEditorSchema.cpp
@@ -0,0 +1,5 @@
+// Copyright Epic Games, Inc. All Rights Reserved.
+
+#include "UAFStateTreeEditorSchema.h"
+
+#include UE_INLINE_GENERATED_CPP_BY_NAME(UAFStateTreeEditorSchema)
\ No newline at end of file
diff --git a/Engine/Plugins/Experimental/UAF/UAFStateTree/Source/UAFStateTreeEditor/Public/UAFStateTreeEditorSchema.h b/Engine/Plugins/Experimental/UAF/UAFStateTree/Source/UAFStateTreeEditor/Public/UAFStateTreeEditorSchema.h
new file mode 100644
--- /dev/null
+++ b/Engine/Plugins/Experimental/UAF/UAFStateTree/Source/UAFStateTreeEditor/Public/UAFStateTreeEditorSchema.h
@@ -0,0 +1,18 @@
+// Copyright Epic Games, Inc. All Rights Reserved.
+
+#pragma once
+
+#include "StateTreeEditorSchema.h"
+
+#include "UAFStateTreeEditorSchema.generated.h"
+
+UCLASS(MinimalAPI)
+class UUAFStateTreeEditorSchema : public UStateTreeEditorSchema
+{
+	GENERATED_BODY()
+
+	virtual EStateTreeTransitionEditingRules GetTransitionEditingRules(TNotNull<UStateTreeEditorData*> StateTree) const override
+	{
+		return EStateTreeTransitionEditingRules::AllowReactivation;
+	}
+};
diff --git a/Engine/Plugins/Runtime/StateTree/Source/StateTreeEditorModule/Private/Customizations/StateTreeTransitionDetails.cpp b/Engine/Plugins/Runtime/StateTree/Source/StateTreeEditorModule/Private/Customizations/StateTreeTransitionDetails.cpp
--- a/Engine/Plugins/Runtime/StateTree/Source/StateTreeEditorModule/Private/Customizations/StateTreeTransitionDetails.cpp
+++ b/Engine/Plugins/Runtime/StateTree/Source/StateTreeEditorModule/Private/Customizations/StateTreeTransitionDetails.cpp
@@ -14,6 +14,7 @@
 #include "StateTreeEditor.h"
 #include "StateTreeEditorData.h"
 #include "StateTreeEditorNodeUtils.h"
+#include "StateTreeEditorSchema.h"
 #include "StateTreeEditorStyle.h"
 #include "StateTreePropertyHelpers.h"
 #include "StateTreeEditorDataClipboardHelpers.h"
@@ -44,20 +45,7 @@ void FStateTreeTransitionDetails::CustomizeHeader(TSharedRef<IPropertyHandle> St
 	PropUtils = StructCustomizationUtils.GetPropertyUtilities();
 
 	// Find StateTreeEditorData associated with this panel.
-	UStateTreeEditorData* EditorData = nullptr;
-	const TArray<TWeakObjectPtr<>>& Objects = PropUtils->GetSelectedObjects();
-	for (const TWeakObjectPtr<>& WeakObject : Objects)
-	{
-		if (const UObject* Object = WeakObject.Get())
-		{
-			if (UStateTreeEditorData* OuterEditorData = Object->GetTypedOuter<UStateTreeEditorData>())
-			{
-				EditorData = OuterEditorData;
-				break;
-			}
-		}
-	}
-
+	UStateTreeEditorData* EditorData = GetEditorData();
 	if (UStateTreeEditingSubsystem* StateTreeEditingSubsystem = EditorData != nullptr ? GEditor->GetEditorSubsystem<UStateTreeEditingSubsystem>() : nullptr)
 	{
 		StateTreeViewModel = StateTreeEditingSubsystem->FindOrAddViewModel(EditorData->GetOuterUStateTree());
@@ -69,6 +57,7 @@ void FStateTreeTransitionDetails::CustomizeHeader(TSharedRef<IPropertyHandle> St
 	DelegateListener = StructProperty->GetChildHandle(GET_MEMBER_NAME_CHECKED(FStateTreeTransition, DelegateListener));
 	StateProperty = StructProperty->GetChildHandle(GET_MEMBER_NAME_CHECKED(FStateTreeTransition, State));
 	DelayTransitionProperty = StructProperty->GetChildHandle(GET_MEMBER_NAME_CHECKED(FStateTreeTransition, bDelayTransition));
+	TransitionReactivateProperty = StructProperty->GetChildHandle(GET_MEMBER_NAME_CHECKED(FStateTreeTransition, ReactivateTargetState));
 	DelayDurationProperty = StructProperty->GetChildHandle(GET_MEMBER_NAME_CHECKED(FStateTreeTransition, DelayDuration));
 	DelayRandomVarianceProperty = StructProperty->GetChildHandle(GET_MEMBER_NAME_CHECKED(FStateTreeTransition, DelayRandomVariance));
 	ConditionsProperty = StructProperty->GetChildHandle(GET_MEMBER_NAME_CHECKED(FStateTreeTransition, Conditions));
@@ -157,18 +146,117 @@ void FStateTreeTransitionDetails::CustomizeHeader(TSharedRef<IPropertyHandle> St
 		.PasteAction(FUIAction(FExecuteAction::CreateSP(this, &FStateTreeTransitionDetails::OnPasteTransitions)));
 }
 
+EStateTreeVisitor VisitNestedHierachyIncludingLinks_Internal(const TMap<UStateTreeState*, TArray<UStateTreeState*>>& LinkedByStateMap
+	, TSet<UStateTreeState*>& VisitedStates
+	, UStateTreeState& NestedStateLeaf
+	, TFunctionRef<EStateTreeVisitor(UStateTreeState& State)> InFunc)
+{
+	bool bContinue = InFunc(NestedStateLeaf) == EStateTreeVisitor::Continue;
+	VisitedStates.Add(&NestedStateLeaf);
+	if (bContinue)
+	{
+		const TArray<UStateTreeState*>* LinkedByStates = LinkedByStateMap.Find(&NestedStateLeaf);
+		UStateTreeState* ParentState = NestedStateLeaf.Parent;
+
+		if (LinkedByStates)
+		{
+			for (UStateTreeState* LinkedByState : *LinkedByStates)
+			{
+				if (!VisitedStates.Contains(LinkedByState) && ensure(LinkedByState))
+				{
+					bContinue = VisitNestedHierachyIncludingLinks_Internal(LinkedByStateMap, VisitedStates, *LinkedByState, InFunc) == EStateTreeVisitor::Continue;
+					if (!bContinue)
+					{
+						return EStateTreeVisitor::Break;
+					}
+				}
+
+			}
+		}
+		if (!VisitedStates.Contains(ParentState) && ParentState)
+		{
+			bContinue = VisitNestedHierachyIncludingLinks_Internal(LinkedByStateMap, VisitedStates, *ParentState, InFunc) == EStateTreeVisitor::Continue;
+			if (!bContinue)
+			{
+				return EStateTreeVisitor::Break;
+			}
+		}
+
+		return EStateTreeVisitor::Continue;
+	}
+
+	return EStateTreeVisitor::Break;
+}
+
+void VisitNestedHierachyIncludingLinks(UStateTreeEditorData* EditorData, UStateTreeState& NestedStateLeaf, TFunctionRef<EStateTreeVisitor(UStateTreeState& State)> InFunc)
+{
+	// We need to visit our nested state, considering links. Suppose we have:
+	// 
+	// Root
+	// - a -> A
+	// - b -> B
+	// - C -> C
+	// - A -> X
+	// - B -> Y
+	// - C -> X
+	// X 
+	// - I
+	// -- P -> Root (Infinite loop, this won't compile, but we are editor time so still need consider)
+	// -- Q
+	// Y 
+	// - J
+	// 
+	// First we construct a LinkedBy list
+	// A: {a}
+	// B: {b}
+	// C: {c} 
+	// X: {A, C}
+	// Y: {B}
+	// Root: {P}
+	// 
+	// Then we visit Nested & Linked State Parents. 
+	// Algo:
+	// - Visit Current, add to Visited
+	// - For each node in union (LinkedBy[Current], Parent):
+	// -- If not in Visited, recurse Algo as new current.
+	
+	// Note: Ptr's not const since we don't require the callback to be const
+	TMap<UStateTreeState*, TArray<UStateTreeState*>> LinkedByStateMap;
+	EditorData->VisitHierarchy([&LinkedByStateMap, EditorData](UStateTreeState& State, UStateTreeState* ParentState)
+	{
+		if (State.Type == EStateTreeStateType::Linked)
+		{
+			UStateTreeState* LinkedState = EditorData->GetMutableStateByID(State.LinkedSubtree.ID);
+			if (ensure(LinkedState))
+			{
+				LinkedByStateMap.FindOrAdd(LinkedState).Add(&State);
+			}
+		}
+		return EStateTreeVisitor::Continue;
+	});
+
+	TSet<UStateTreeState*> VisitedStates;
+	VisitNestedHierachyIncludingLinks_Internal(LinkedByStateMap, VisitedStates, NestedStateLeaf, InFunc);
+}
+
 void FStateTreeTransitionDetails::CustomizeChildren(TSharedRef<class IPropertyHandle> StructPropertyHandle, class IDetailChildrenBuilder& StructBuilder, IPropertyTypeCustomizationUtils& StructCustomizationUtils)
 {
 	check(TriggerProperty);
 	check(RequiredEventProperty);
 	check(DelegateListener);
 	check(DelayTransitionProperty);
+	check(TransitionReactivateProperty);
 	check(DelayDurationProperty);
 	check(DelayRandomVarianceProperty);
 	check(StateProperty);
 	check(ConditionsProperty);
 	check(IDProperty);
 
+	UStateTreeState* OwningState = GetOwningStateTreeState();
+	TWeakObjectPtr<UStateTreeState> WeakOwningState = OwningState;
+	UStateTreeEditorData* EditorData = GetEditorData();
+	TWeakObjectPtr<UStateTreeEditorData> WeakEditorData = EditorData;
+
 	TWeakPtr<FStateTreeTransitionDetails> WeakSelf = SharedThis(this);
 	auto IsNotCompletionTransition = [WeakSelf]()
 	{
@@ -178,6 +266,50 @@ void FStateTreeTransitionDetails::CustomizeChildren(TSharedRef<class IPropertyHa
 		}
 		return EVisibility::Collapsed;
 	};
+	auto IsInNestedState = [WeakSelf, WeakOwningState, WeakEditorData]()
+	{
+		const TSharedPtr<FStateTreeTransitionDetails> Self = WeakSelf.Pin();
+		TStrongObjectPtr<UStateTreeState> OwningState = WeakOwningState.Pin();
+		TStrongObjectPtr<UStateTreeEditorData> EditorData = WeakEditorData.Pin();
+
+		if (!Self || !OwningState || !EditorData)
+		{
+			return false;
+		}
+
+		// Try to return cached value if it's correct
+		FStateTreeStateLink TargetStateLink = Self->GetTargetState();
+		if (Self->CachedLastTargetState.ID == TargetStateLink.ID)
+		{
+			return Self->bCachedIsTargetStateInNestedState;
+		}
+		else
+		{
+			Self->CachedLastTargetState = TargetStateLink;
+		}
+
+		UStateTree* StateTree = EditorData->GetOuterUStateTree();
+		UStateTreeState* CurrentState = OwningState.Get();
+		FStateTreeStateLink CurrentStateLink = CurrentState->GetLinkToState();
+
+		bool bIsTargetStateInNestedState = false;
+		auto IsTargetState = [TargetStateLink, &bIsTargetStateInNestedState](UStateTreeState& InCurrentState)
+		{
+			if (TargetStateLink.ID != InCurrentState.GetLinkToState().ID)
+			{
+				return EStateTreeVisitor::Continue;
+			}
+			else
+			{
+				bIsTargetStateInNestedState = true;
+				return EStateTreeVisitor::Break;
+			}
+		};
+		VisitNestedHierachyIncludingLinks(EditorData.Get(), *CurrentState, IsTargetState);
+
+		Self->bCachedIsTargetStateInNestedState = bIsTargetStateInNestedState;
+		return bIsTargetStateInNestedState;
+	};
 
 	if (UE::StateTree::Editor::GbDisplayItemIds)
 	{
@@ -189,24 +321,24 @@ void FStateTreeTransitionDetails::CustomizeChildren(TSharedRef<class IPropertyHa
 
 	// Show event only when the trigger is set to Event.
 	StructBuilder.AddProperty(RequiredEventProperty.ToSharedRef())
-		.Visibility(TAttribute<EVisibility>::Create(TAttribute<EVisibility>::FGetter::CreateLambda([WeakSelf]()
+		.Visibility(MakeAttributeLambda([WeakSelf]()
 		{
 			if (const TSharedPtr<FStateTreeTransitionDetails> Self = WeakSelf.Pin())
 			{
 				return (Self->GetTrigger() == EStateTreeTransitionTrigger::OnEvent) ? EVisibility::Visible : EVisibility::Collapsed;
 			}
 			return EVisibility::Collapsed;
-		})));
+		}));
 
 	IDetailPropertyRow& DelegateDispatcherRow = StructBuilder.AddProperty(DelegateListener.ToSharedRef())
-		.Visibility(TAttribute<EVisibility>::Create(TAttribute<EVisibility>::FGetter::CreateLambda([WeakSelf]()
+		.Visibility(MakeAttributeLambda([WeakSelf]()
 		{
 			if (const TSharedPtr<FStateTreeTransitionDetails> Self = WeakSelf.Pin())
 			{
 				return (Self->GetTrigger() == EStateTreeTransitionTrigger::OnDelegate) ? EVisibility::Visible : EVisibility::Collapsed;
 			}
 			return EVisibility::Collapsed;
-		})));
+		}));
 
 	FGuid ID;
 	UE::StateTree::PropertyHelpers::GetStructValue<FGuid>(IDProperty, ID);
@@ -217,15 +349,31 @@ void FStateTreeTransitionDetails::CustomizeChildren(TSharedRef<class IPropertyHa
 
 	// Priority
 	StructBuilder.AddProperty(PriorityProperty.ToSharedRef())
-		.Visibility(TAttribute<EVisibility>::Create(TAttribute<EVisibility>::FGetter::CreateLambda(IsNotCompletionTransition)));
+		.Visibility(MakeAttributeLambda(IsNotCompletionTransition));
 
 	// Delay
 	StructBuilder.AddProperty(DelayTransitionProperty.ToSharedRef())
-		.Visibility(TAttribute<EVisibility>::Create(TAttribute<EVisibility>::FGetter::CreateLambda(IsNotCompletionTransition)));
+		.Visibility(MakeAttributeLambda(IsNotCompletionTransition));
 	StructBuilder.AddProperty(DelayDurationProperty.ToSharedRef())
-		.Visibility(TAttribute<EVisibility>::Create(TAttribute<EVisibility>::FGetter::CreateLambda(IsNotCompletionTransition)));
+		.Visibility(MakeAttributeLambda(IsNotCompletionTransition));
 	StructBuilder.AddProperty(DelayRandomVarianceProperty.ToSharedRef())
-		.Visibility(TAttribute<EVisibility>::Create(TAttribute<EVisibility>::FGetter::CreateLambda(IsNotCompletionTransition)));
+		.Visibility(MakeAttributeLambda(IsNotCompletionTransition));
+
+	// Transition Reactivation
+	auto IsTransitionReactiviationSchemaEditable = [WeakEditorData]()
+	{
+		bool bAllowed = false;
+		if (TStrongObjectPtr<UStateTreeEditorData> EditorData = WeakEditorData.Pin())
+		{
+			UStateTree* StateTree = EditorData->GetOuterUStateTree();
+			bAllowed = EnumHasAnyFlags(EditorData->EditorSchema->GetTransitionEditingRules(EditorData.Get()), EStateTreeTransitionEditingRules::AllowReactivation);
+		}
+
+		return bAllowed ? EVisibility::Visible : EVisibility::Collapsed;
+	};
+	StructBuilder.AddProperty(TransitionReactivateProperty.ToSharedRef())
+		.EditCondition(MakeAttributeLambda(IsInNestedState), FOnBooleanValueChanged())
+		.Visibility(MakeAttributeLambda(IsTransitionReactiviationSchemaEditable));
 
 	// Show conditions always expanded, with simplified header (remove item count)
 	IDetailPropertyRow& ConditionsRow = StructBuilder.AddProperty(ConditionsProperty.ToSharedRef());
@@ -480,11 +628,25 @@ void FStateTreeTransitionDetails::OnDuplicateTransition() const
 	}
 }
 
+FStateTreeStateLink FStateTreeTransitionDetails::GetTargetState() const
+{
+	FStateTreeStateLink StateValue;
+	FStateTreeStateLink* StateData;
+	if (ensure(StateProperty.IsValid()))
+	{
+		// Note: This can fail. IE: If we are deleting the current transition
+		if (StateProperty->GetValueData((void*&)StateData) == FPropertyAccess::Success)
+		{
+			StateValue = *StateData;
+		}
+	}
+	return StateValue;
+}
+
 EStateTreeTransitionTrigger FStateTreeTransitionDetails::GetTrigger() const
 {
-	check(TriggerProperty);
 	EStateTreeTransitionTrigger TriggerValue = EStateTreeTransitionTrigger::None;
-	if (TriggerProperty.IsValid())
+	if (ensure(TriggerProperty.IsValid()))
 	{
 		TriggerProperty->GetValue((uint8&)TriggerValue);
 	}
@@ -493,9 +655,8 @@ EStateTreeTransitionTrigger FStateTreeTransitionDetails::GetTrigger() const
 
 bool FStateTreeTransitionDetails::GetDelayTransition() const
 {
-	check(DelayTransitionProperty);
 	bool bDelayTransition = false;
-	if (DelayTransitionProperty.IsValid())
+	if (ensure(DelayTransitionProperty.IsValid()))
 	{
 		DelayTransitionProperty->GetValue(bDelayTransition);
 	}
@@ -558,6 +719,21 @@ void FStateTreeTransitionDetails::OnCopyAllTransitions() const
 	}
 }
 
+UStateTreeState* FStateTreeTransitionDetails::GetOwningStateTreeState() const
+{
+	TArray<UObject*> OuterObjects;
+	StructProperty->GetOuterObjects(OuterObjects);
+	for (UObject* Outer : OuterObjects)
+	{
+		if (UStateTreeState* OuterStateTreeState = Cast<UStateTreeState>(Outer))
+		{
+			return OuterStateTreeState;
+		}
+	}
+
+	return nullptr;
+}
+
 UStateTreeEditorData* FStateTreeTransitionDetails::GetEditorData() const
 {
 	TArray<UObject*> OuterObjects;
diff --git a/Engine/Plugins/Runtime/StateTree/Source/StateTreeEditorModule/Private/Customizations/StateTreeTransitionDetails.h b/Engine/Plugins/Runtime/StateTree/Source/StateTreeEditorModule/Private/Customizations/StateTreeTransitionDetails.h
--- a/Engine/Plugins/Runtime/StateTree/Source/StateTreeEditorModule/Private/Customizations/StateTreeTransitionDetails.h
+++ b/Engine/Plugins/Runtime/StateTree/Source/StateTreeEditorModule/Private/Customizations/StateTreeTransitionDetails.h
@@ -4,6 +4,7 @@
 
 #include "IPropertyTypeCustomization.h"
 #include "Input/Reply.h"
+#include "StateTreeTypes.h"
 
 class FStateTreeViewModel;
 class IPropertyHandle;
@@ -15,6 +16,7 @@ class UStateTreeEditorData;
 enum class EStateTreeTransitionTrigger : uint8;
 class SBorder;
 class SWidget;
+class UStateTreeState;
 
 /**
  * Type customization for FStateTreeTransition.
@@ -32,10 +34,12 @@ public:
 
 private:
 
+	UStateTreeState* GetOwningStateTreeState() const;
 	UStateTreeEditorData* GetEditorData() const;
 
 	FText GetDescription() const;
 
+	FStateTreeStateLink GetTargetState() const;
 	EStateTreeTransitionTrigger GetTrigger() const;
 	bool GetDelayTransition() const;
 
@@ -60,6 +64,7 @@ private:
 	TSharedPtr<IPropertyHandle> DelegateListener;
 	TSharedPtr<IPropertyHandle> StateProperty;
 	TSharedPtr<IPropertyHandle> DelayTransitionProperty;
+	TSharedPtr<IPropertyHandle> TransitionReactivateProperty;
 	TSharedPtr<IPropertyHandle> DelayDurationProperty;
 	TSharedPtr<IPropertyHandle> DelayRandomVarianceProperty;
 	TSharedPtr<IPropertyHandle> ConditionsProperty;
@@ -72,4 +77,10 @@ private:
 
 	TSharedPtr<IPropertyHandle> ParentProperty;
 	TSharedPtr<IPropertyHandleArray> ParentArrayProperty;
+
+	/** Used to avoid having to re-compute if target state within this state's nested state hierachy. */
+	FStateTreeStateLink CachedLastTargetState = {};
+
+	/** True if our last target state was in the nested state hierachy */
+	bool bCachedIsTargetStateInNestedState = false;
 };
diff --git a/Engine/Plugins/Runtime/StateTree/Source/StateTreeEditorModule/Private/SStateTreeViewRow.cpp b/Engine/Plugins/Runtime/StateTree/Source/StateTreeEditorModule/Private/SStateTreeViewRow.cpp
--- a/Engine/Plugins/Runtime/StateTree/Source/StateTreeEditorModule/Private/SStateTreeViewRow.cpp
+++ b/Engine/Plugins/Runtime/StateTree/Source/StateTreeEditorModule/Private/SStateTreeViewRow.cpp
@@ -995,6 +995,7 @@ TSharedRef<SWidget> SStateTreeViewRow::MakeTransitionWidgetInternal(const EState
 		FText Tooltip;
 		FStateTreeStateLink Link;
 		FGuid NodeID;
+		bool bSkipComma = false;
 	};
 
 	TArray<FItem> DescItems;
@@ -1022,6 +1023,15 @@ TSharedRef<SWidget> SStateTreeViewRow::MakeTransitionWidgetInternal(const EState
 		if (bMatch)
 		{
 			DescItems.Emplace(Transition.State, Transition.ID);
+
+			if (Transition.ReactivateTargetState == EStateTreeTransitionChangeTypeRules::ForceChanged)
+			{
+				FItem ReselectionNote;
+				ReselectionNote.Desc = LOCTEXT("ReactivateDesc", "(Re-Activate)");
+				ReselectionNote.Tooltip = LOCTEXT("ReactivateDescToolTip", "The target state and all the following states will be reactivated. Their instance data will be reinstantiated and will receive EnterState.");
+				ReselectionNote.bSkipComma = true;
+				DescItems.Emplace(ReselectionNote);
+			}
 		}
 	}
 
@@ -1090,7 +1100,7 @@ TSharedRef<SWidget> SStateTreeViewRow::MakeTransitionWidgetInternal(const EState
 	{
 		const FItem& Item = DescItems[Index];
 
-		if (Index > 0)
+		if (Index > 0 && !Item.bSkipComma)
 		{
 			TransitionContainer->AddSlot()
 			.AutoWidth()
diff --git a/Engine/Plugins/Runtime/StateTree/Source/StateTreeEditorModule/Private/StateTreeCompiler.cpp b/Engine/Plugins/Runtime/StateTree/Source/StateTreeEditorModule/Private/StateTreeCompiler.cpp
--- a/Engine/Plugins/Runtime/StateTree/Source/StateTreeEditorModule/Private/StateTreeCompiler.cpp
+++ b/Engine/Plugins/Runtime/StateTree/Source/StateTreeEditorModule/Private/StateTreeCompiler.cpp
@@ -1441,6 +1441,7 @@ bool FStateTreeCompiler::CreateStateTransitions()
 				}
 			}
 
+			CompactTransition.ChangeTypeTargetStateRule = Transition.ReactivateTargetState;
 			CompactTransition.bTransitionEnabled = Transition.bTransitionEnabled;
 
 			if (Transition.bDelayTransition)
diff --git a/Engine/Plugins/Runtime/StateTree/Source/StateTreeEditorModule/Private/StateTreeDescriptionHelpers.cpp b/Engine/Plugins/Runtime/StateTree/Source/StateTreeEditorModule/Private/StateTreeDescriptionHelpers.cpp
--- a/Engine/Plugins/Runtime/StateTree/Source/StateTreeEditorModule/Private/StateTreeDescriptionHelpers.cpp
+++ b/Engine/Plugins/Runtime/StateTree/Source/StateTreeEditorModule/Private/StateTreeDescriptionHelpers.cpp
@@ -247,7 +247,16 @@ FText GetTransitionDesc(const UStateTreeEditorData* EditorData, const FStateTree
 			: LOCTEXT("ActionReturn", "return");
 	}
 	
-	return FText::Format(LOCTEXT("TransitionDesc", "{0} {1} {2}"), TriggerText, ActionText, GetStateLinkDesc(EditorData, Transition.State, Formatting, bShowStatePath));
+	FText ReactivateText = Formatting == EStateTreeNodeFormatting::RichText
+		? LOCTEXT("ReselectionTextRich", "<s>(Re-Activate)</>")
+		: LOCTEXT("ReselectionText", "(Re-Activate)");
+
+	if (Transition.ReactivateTargetState != EStateTreeTransitionChangeTypeRules::ForceChanged)
+	{
+		ReactivateText = FText::GetEmpty();
+	}
+
+	return FText::Format(LOCTEXT("TransitionDesc", "{0} {1} {2} {3}"), TriggerText, ActionText, GetStateLinkDesc(EditorData, Transition.State, Formatting, bShowStatePath), ReactivateText);
 }
 
 } // UE::StateTree::Editor
diff --git a/Engine/Plugins/Runtime/StateTree/Source/StateTreeEditorModule/Private/StateTreeEditingSubsystem.cpp b/Engine/Plugins/Runtime/StateTree/Source/StateTreeEditorModule/Private/StateTreeEditingSubsystem.cpp
--- a/Engine/Plugins/Runtime/StateTree/Source/StateTreeEditorModule/Private/StateTreeEditingSubsystem.cpp
+++ b/Engine/Plugins/Runtime/StateTree/Source/StateTreeEditorModule/Private/StateTreeEditingSubsystem.cpp
@@ -192,9 +192,10 @@ void UStateTreeEditingSubsystem::ValidateStateTree(TNotNull<UStateTree*> InState
 
 			// Trash the previous EditorData
 			const FName TrashStateTreeName = MakeUniqueObjectName(GetTransientPackage(), UStateTree::StaticClass(), *FString::Printf(TEXT("TRASH_%s"), *UStateTree::StaticClass()->GetName()));
-			UStateTree* TransientOuter = NewObject<UStateTree>(GetTransientPackage(), TrashStateTreeName, RF_Transient);
 			const FName TrashSchemaName = *FString::Printf(TEXT("TRASH_%s"), *PreviousEditorData->GetName());
 			constexpr ERenameFlags RenameFlags = REN_DoNotDirty | REN_DontCreateRedirectors;
+
+			UStateTree* TransientOuter = NewObject<UStateTree>(GetTransientPackage(), TrashStateTreeName, RF_Transient);
 			PreviousEditorData->Rename(*TrashSchemaName.ToString(), TransientOuter, RenameFlags);
 			PreviousEditorData->SetFlags(RF_Transient);
 		}
@@ -218,11 +219,12 @@ void UStateTreeEditingSubsystem::ValidateStateTree(TNotNull<UStateTree*> InState
 				// The current EditorSchema is not of the correct type. The data needs to be patched by the schema desired editor data subclass.
 				UStateTreeEditorSchema* PreviousEditorSchema = EditorData->EditorSchema;
 				EditorData->EditorSchema = CastChecked<UStateTreeEditorSchema>(StaticDuplicateObject(PreviousEditorSchema, EditorData, FName(), RF_Transactional, EditorSchemaClass.Get()));
-
-				// Trash the previous EditorData
-				const FName TrashName = MakeUniqueObjectName(GetTransientPackage(), UStateTreeEditorSchema::StaticClass(), *FString::Printf(TEXT("TRASH_%s"), *PreviousEditorSchema->GetName()));
+				
+				// Trash the previous EditorDataSchema
+				const FName TrashSchemaName = *FString::Printf(TEXT("TRASH_%s"), *PreviousEditorSchema->GetName());
 				constexpr ERenameFlags RenameFlags = REN_DoNotDirty | REN_DontCreateRedirectors;
-				PreviousEditorSchema->Rename(*TrashName.ToString(), PreviousEditorSchema, RenameFlags);
+
+				PreviousEditorSchema->Rename(*TrashSchemaName.ToString(), GetTransientPackage(), RenameFlags);
 				PreviousEditorSchema->SetFlags(RF_Transient);
 			}
 		};
diff --git a/Engine/Plugins/Runtime/StateTree/Source/StateTreeEditorModule/Public/StateTreeEditorModule.h b/Engine/Plugins/Runtime/StateTree/Source/StateTreeEditorModule/Public/StateTreeEditorModule.h
--- a/Engine/Plugins/Runtime/StateTree/Source/StateTreeEditorModule/Public/StateTreeEditorModule.h
+++ b/Engine/Plugins/Runtime/StateTree/Source/StateTreeEditorModule/Public/StateTreeEditorModule.h
@@ -125,11 +125,11 @@ public:
 	UE_API TNonNullSubclassOf<UStateTreeEditorData> GetEditorDataClass(TNonNullSubclassOf<const UStateTreeSchema> Schema) const;
 
 	/** Register the editor schema type for a specific schema. */
-	void RegisterEditorSchemaClass(TNonNullSubclassOf<const UStateTreeSchema> Schema, TNonNullSubclassOf<const UStateTreeEditorSchema> EditorSchema);
+	UE_API void RegisterEditorSchemaClass(TNonNullSubclassOf<const UStateTreeSchema> Schema, TNonNullSubclassOf<const UStateTreeEditorSchema> EditorSchema);
 	/** Unregister the editor schema type for a specific schema. */
-	void UnregisterEditorSchemaClass(TNonNullSubclassOf<const UStateTreeSchema> Schema);
+	UE_API void UnregisterEditorSchemaClass(TNonNullSubclassOf<const UStateTreeSchema> Schema);
 	/** Get the editor schema type for a specific schema. */
-	TNonNullSubclassOf<UStateTreeEditorSchema> GetEditorSchemaClass(TNonNullSubclassOf<const UStateTreeSchema> Schema) const;
+	UE_API TNonNullSubclassOf<UStateTreeEditorSchema> GetEditorSchemaClass(TNonNullSubclassOf<const UStateTreeSchema> Schema) const;
 
 private:
 	void HandleRequestAssetRegistryTags(TNotNull<const UStateTree*> InStateTree, FAssetRegistryTagsContext& Context);
diff --git a/Engine/Plugins/Runtime/StateTree/Source/StateTreeEditorModule/Public/StateTreeEditorSchema.h b/Engine/Plugins/Runtime/StateTree/Source/StateTreeEditorModule/Public/StateTreeEditorSchema.h
--- a/Engine/Plugins/Runtime/StateTree/Source/StateTreeEditorModule/Public/StateTreeEditorSchema.h
+++ b/Engine/Plugins/Runtime/StateTree/Source/StateTreeEditorModule/Public/StateTreeEditorSchema.h
@@ -8,11 +8,26 @@
 #define UE_API STATETREEEDITORMODULE_API
 
 class UStateTree;
+class UStateTreeEditorData;
 namespace UE::StateTree::Compiler
 {
 	struct FPostInternalContext;
 }
 
+/** Rules specific to editing Transitions. */
+UENUM()
+enum class EStateTreeTransitionEditingRules : uint32
+{
+	/** Previous (UE 5.7) rules */
+	None = 0,
+	/**
+	 * When enabled, user can specify re-activation in the editor. See FStateTreeTransition::ReactivateTargetState for more info.
+	 */
+	AllowReactivation = 1 << 0,
+
+	Default = None,
+};
+
 /**
  * Schema describing how a StateTree is edited.
  */
@@ -52,6 +67,12 @@ public:
 	{
 		return true;
 	}
+
+	/** Get editing rules for transitions */
+	virtual EStateTreeTransitionEditingRules GetTransitionEditingRules(TNotNull<UStateTreeEditorData*> StateTree) const
+	{
+		return EStateTreeTransitionEditingRules::Default;
+	}
 };
 
 #undef UE_API
diff --git a/Engine/Plugins/Runtime/StateTree/Source/StateTreeEditorModule/Public/StateTreeState.h b/Engine/Plugins/Runtime/StateTree/Source/StateTreeEditorModule/Public/StateTreeState.h
--- a/Engine/Plugins/Runtime/StateTree/Source/StateTreeEditorModule/Public/StateTreeState.h
+++ b/Engine/Plugins/Runtime/StateTree/Source/StateTreeEditorModule/Public/StateTreeState.h
@@ -160,6 +160,13 @@ PRAGMA_ENABLE_DEPRECATION_WARNINGS
 	UPROPERTY(EditDefaultsOnly, Category = "Transition", meta = (BaseStruct = "/Script/StateTreeModule.StateTreeConditionBase", BaseClass = "/Script/StateTreeModule.StateTreeConditionBlueprintBase"))
 	TArray<FStateTreeEditorNode> Conditions;
 
+	/** 
+	 * The target state and all the following states will be reactivated. Their instance data will be reinstantiated and will receive EnterState. 
+	 * Note: Only enabled when target state is within the nested state (Parent or self). As this only changes behavior for nested states.
+	 */
+	UPROPERTY(EditDefaultsOnly, Category = "Transition", meta = (HideEditConditionToggle))
+	EStateTreeTransitionChangeTypeRules ReactivateTargetState = EStateTreeTransitionChangeTypeRules::Default;
+
 	/** True if the Transition is Enabled (i.e. not explicitly disabled in the asset). */
 	UPROPERTY(EditDefaultsOnly, Category = "Debug")
 	bool bTransitionEnabled = true;
diff --git a/Engine/Plugins/Runtime/StateTree/Source/StateTreeModule/Private/StateTreeExecutionContext.cpp b/Engine/Plugins/Runtime/StateTree/Source/StateTreeModule/Private/StateTreeExecutionContext.cpp
--- a/Engine/Plugins/Runtime/StateTree/Source/StateTreeModule/Private/StateTreeExecutionContext.cpp
+++ b/Engine/Plugins/Runtime/StateTree/Source/StateTreeModule/Private/StateTreeExecutionContext.cpp
@@ -5211,7 +5211,7 @@ bool FStateTreeExecutionContext::RequestTransitionInternal(
 	SelectStateArgs.TransitionEvent = Args.TransitionEvent;
 	SelectStateArgs.Fallback = Args.Fallback;
 	SelectStateArgs.Behavior = ESelectStateBehavior::StateTransition;
-	SelectStateArgs.SelectionRules = SourceFrame.StateTree->StateSelectionRules;
+	SelectStateArgs.SelectionRules = SourceFrame.StateTree->StateSelectionRules | Args.SelectionRules;
 
 	TSharedRef<FSelectStateResult> StateSelectionResult = MakeShared<FSelectStateResult>();
 	TGuardValue<TSharedPtr<FSelectStateResult>> TemporaryStorageScope(CurrentlyProcessedTemporaryStorage, StateSelectionResult);
@@ -5658,6 +5658,7 @@ bool FStateTreeExecutionContext::TriggerTransitions()
 							TransitionArgs.Priority = Transition.Priority;
 							TransitionArgs.TransitionEvent = DelayedTransition.CapturedEvent;
 							TransitionArgs.Fallback = Transition.Fallback;
+							Transition.ModifyStateSelectionRulesPerTransition(OUT TransitionArgs.SelectionRules);
 							if (RequestTransitionInternal(CurrentFrame, Handler.StateID, TargetState, TransitionArgs))
 							{
 								// If the transition was successfully requested with a specific event, consume and remove the event, it's been used.
@@ -5751,6 +5752,7 @@ PRAGMA_ENABLE_DEPRECATION_WARNINGS
 						TransitionArgs.Priority = Transition.Priority;
 						TransitionArgs.TransitionEvent = TransitionEvent;
 						TransitionArgs.Fallback = Transition.Fallback;
+						Transition.ModifyStateSelectionRulesPerTransition(OUT TransitionArgs.SelectionRules);
 						if (RequestTransitionInternal(CurrentFrame, Handler.StateID, TargetState, TransitionArgs))
 						{
 							// If the transition was successfully requested with a specific event, consume and remove the event, it's been used.
@@ -5915,6 +5917,7 @@ PRAGMA_ENABLE_DEPRECATION_WARNINGS
 									FTransitionArguments TransitionArgs;
 									TransitionArgs.Priority = EStateTreeTransitionPriority::Normal;
 									TransitionArgs.Fallback = Transition.Fallback;
+									Transition.ModifyStateSelectionRulesPerTransition(OUT TransitionArgs.SelectionRules);
 									if (RequestTransitionInternal(CurrentFrame, CurrentStateID, TargetState, TransitionArgs))
 									{
 										RequestedTransition->Source = FStateTreeTransitionSource(CurrentFrame.StateTree, FStateTreeIndex16(TransitionIndex), Transition.State, Transition.Priority);
@@ -7852,6 +7855,7 @@ bool FStateTreeExecutionContext::SelectStateInternal_TryFollowTransitions(
 				SelectStateArgs.TransitionEvent = SelectedStateTransitionEvent;
 				SelectStateArgs.Fallback = Transition.Fallback;
 				SelectStateArgs.SelectionRules = NextStateTree->StateSelectionRules;
+				Transition.ModifyStateSelectionRulesPerTransition(OUT SelectStateArgs.SelectionRules);
 
 				OutSelectionResult->SelectedStates.Reset();
 				OutSelectionResult->SelectedFrames.Reset();
diff --git a/Engine/Plugins/Runtime/StateTree/Source/StateTreeModule/Public/StateTreeExecutionContext.h b/Engine/Plugins/Runtime/StateTree/Source/StateTreeModule/Public/StateTreeExecutionContext.h
--- a/Engine/Plugins/Runtime/StateTree/Source/StateTreeModule/Public/StateTreeExecutionContext.h
+++ b/Engine/Plugins/Runtime/StateTree/Source/StateTreeModule/Public/StateTreeExecutionContext.h
@@ -1173,6 +1173,7 @@ protected:
 		EStateTreeTransitionPriority Priority = EStateTreeTransitionPriority::Normal;
 		FStateTreeSharedEvent TransitionEvent;
 		EStateTreeSelectionFallback Fallback = EStateTreeSelectionFallback::None;
+		EStateTreeStateSelectionRules SelectionRules = EStateTreeStateSelectionRules::None;
 	};
 
 	/** Requests transition to a specified state. */
diff --git a/Engine/Plugins/Runtime/StateTree/Source/StateTreeModule/Public/StateTreeTypes.h b/Engine/Plugins/Runtime/StateTree/Source/StateTreeModule/Public/StateTreeTypes.h
--- a/Engine/Plugins/Runtime/StateTree/Source/StateTreeModule/Public/StateTreeTypes.h
+++ b/Engine/Plugins/Runtime/StateTree/Source/StateTreeModule/Public/StateTreeTypes.h
@@ -601,6 +601,18 @@ struct FCompactEventDesc
 	STATETREEMODULE_API bool DoesEventMatchDesc(const FStateTreeEvent& Event) const;
 };
 
+/**
+ * Additional rules for if state was changed or sustained per transition. 
+ * This is exposed to designers as "Reactivation" w/ "Override" & "Denied" options.
+ */
+UENUM()
+enum class EStateTreeTransitionChangeTypeRules : uint8
+{
+	ForceChanged,
+	ForceSustained,
+	Default,
+};
+
 /**
  *  Runtime representation of a StateTree transition.
  */
@@ -619,6 +631,25 @@ struct FCompactStateTransition
 		return !Delay.IsEmpty();
 	}
 
+	/** Toggles flags for state selection based on this transition's configuration */
+	void ModifyStateSelectionRulesPerTransition(EStateTreeStateSelectionRules& OutModifiedStateSelectionRules) const
+	{
+		switch (ChangeTypeTargetStateRule)
+		{
+		case EStateTreeTransitionChangeTypeRules::ForceChanged:
+			OutModifiedStateSelectionRules |= EStateTreeStateSelectionRules::ReselectedStateCreatesNewStates;
+			return;
+		case EStateTreeTransitionChangeTypeRules::ForceSustained:
+			OutModifiedStateSelectionRules &= ~EStateTreeStateSelectionRules::ReselectedStateCreatesNewStates;
+			return;
+		case EStateTreeTransitionChangeTypeRules::Default:
+			return;
+		default:
+			checkf(false, TEXT("Unsupported EStateTreeTransitionChangeTypeRules: %u"), static_cast<uint8>(ChangeTypeTargetStateRule));
+			return;
+		}
+	}
+
 	/** Event Description */
 	UPROPERTY()
 	FCompactEventDesc RequiredEvent;
@@ -658,13 +689,13 @@ struct FCompactStateTransition
 	UPROPERTY()
 	uint8 ConditionsNum = 0;
 
-	/** Indicates if the transition is enabled and should be considered. */
+	/** Additional rules for state changed or sustain. Their instance data will be reinstantiated. */
 	UPROPERTY()
-	uint8 bTransitionEnabled : 1 = true;
+	EStateTreeTransitionChangeTypeRules ChangeTypeTargetStateRule = EStateTreeTransitionChangeTypeRules::Default;
 
-	/** The target state and all the following states will be reactivated. Their instance data will be reinstantiated and will receive EnterState. */
+	/** Indicates if the transition is enabled and should be considered. */
 	UPROPERTY()
-	uint8 bReactivateTargetState : 1 = false;
+	uint8 bTransitionEnabled : 1 = true;
 
 	/** If set to true, the required event is consumed (later state selection cannot react to it) if state selection can be made. */
 	UPROPERTY()
diff --git a/Engine/Plugins/Runtime/StateTree/Source/StateTreeTestSuite/Private/StateTreeTestTypes.h b/Engine/Plugins/Runtime/StateTree/Source/StateTreeTestSuite/Private/StateTreeTestTypes.h
--- a/Engine/Plugins/Runtime/StateTree/Source/StateTreeTestSuite/Private/StateTreeTestTypes.h
+++ b/Engine/Plugins/Runtime/StateTree/Source/StateTreeTestSuite/Private/StateTreeTestTypes.h
@@ -614,6 +614,7 @@ struct FTestTask_Stand : public FStateTreeTaskBase
 		FInstanceDataType& InstanceData = Context.GetInstanceData(*this);
 		
 		InstanceData.CurrentTick++;
+		TestLog.Log(Name, FString::Printf(TEXT("CurrentTick=%u"), InstanceData.CurrentTick));
 		
 		return (InstanceData.CurrentTick >= TicksToCompletion) ? TickCompletionResult : EStateTreeRunStatus::Running;
 	};
@@ -653,6 +654,22 @@ struct FTestTask_StandNoTick : public FTestTask_Stand
 	}
 };
 
+USTRUCT()
+struct FTestTask_StandNoStateChangeOnReselect : public FTestTask_Stand
+{
+	GENERATED_BODY()
+
+	FTestTask_StandNoStateChangeOnReselect()
+	{
+		bShouldStateChangeOnReselect = false;
+	}
+
+	FTestTask_StandNoStateChangeOnReselect(const FName InName)
+		: FTestTask_Stand(InName)
+	{
+		bShouldStateChangeOnReselect = false;
+	}
+};
 
 USTRUCT()
 struct FTestTask_IntegersOutput_InstanceData
diff --git a/Engine/Plugins/Runtime/StateTree/Source/StateTreeTestSuite/Private/StateTreeTransitionTest.cpp b/Engine/Plugins/Runtime/StateTree/Source/StateTreeTestSuite/Private/StateTreeTransitionTest.cpp
--- a/Engine/Plugins/Runtime/StateTree/Source/StateTreeTestSuite/Private/StateTreeTransitionTest.cpp
+++ b/Engine/Plugins/Runtime/StateTree/Source/StateTreeTestSuite/Private/StateTreeTransitionTest.cpp
@@ -454,6 +454,104 @@ struct FStateTreeTest_TransitionDelay : FStateTreeTestBase
 };
 IMPLEMENT_STATE_TREE_INSTANT_TEST(FStateTreeTest_TransitionDelay, "System.StateTree.Transition.Delay");
 
+struct FStateTreeTest_TransitionReactivateTargetState : FStateTreeTestBase
+{
+	virtual bool InstantTest() override
+	{
+		//Root
+		// State0
+
+		UStateTree& StateTree = NewStateTree();
+		UStateTreeEditorData& EditorData = *Cast<UStateTreeEditorData>(StateTree.EditorData);
+		FInstancedPropertyBag& RootPropertyBag = GetRootPropertyBag(EditorData);
+		RootPropertyBag.AddProperty(FName(TEXT("Int")), EPropertyBagPropertyType::Int32);
+		RootPropertyBag.SetValueInt32(FName(TEXT("Int")), 1);
+
+		UStateTreeState& Root = EditorData.AddSubTree(FName(TEXT("Root")));
+		UStateTreeState& State0 = Root.AddChildState(FName(TEXT("State0")));
+
+		// Note: Task does not reselect by default, so will not trigger entry state on reselect
+		auto& Task0 = State0.AddTask<FTestTask_StandNoStateChangeOnReselect>(FName(TEXT("Task0")));
+		Task0.GetNode().TicksToCompletion = 100;
+
+		// Reactivate transtion back to root, by default will not fire due to condition. But should reactivate / re-trigger enter state when activated.
+		FStateTreeTransition& ReactivateRootTransition = State0.AddTransition(EStateTreeTransitionTrigger::OnTick, EStateTreeTransitionType::GotoState, &Root);
+		TStateTreeEditorNode<FStateTreeCompareIntCondition>& TransIntCond = ReactivateRootTransition.AddCondition<FStateTreeCompareIntCondition>(EGenericAICheck::Equal);
+		TransIntCond.GetInstanceData().Right = 0;
+		EditorData.AddPropertyBinding(
+			FPropertyBindingPath(EditorData.GetRootParametersGuid(), TEXT("Int")),
+			FPropertyBindingPath(TransIntCond.ID, TEXT("Left")));
+		ReactivateRootTransition.ReactivateTargetState = EStateTreeTransitionChangeTypeRules::ForceChanged;
+
+		// Transition back to root, will always fire, but will not reactivate / re-trigger enter state
+		FStateTreeTransition& DefaultRootTransition = State0.AddTransition(EStateTreeTransitionTrigger::OnTick, EStateTreeTransitionType::GotoState, &Root);
+
+		FStateTreeCompilerLog Log;
+		FStateTreeCompiler Compiler(Log);
+		const bool bResult = Compiler.Compile(StateTree);
+		AITEST_TRUE("StateTree should get compiled", bResult);
+
+		EStateTreeRunStatus Status = EStateTreeRunStatus::Unset;
+		FStateTreeInstanceData InstanceData;
+		FTestStateTreeExecutionContext Exec(StateTree, StateTree, InstanceData);
+		const bool bInitSucceeded = Exec.IsValid();
+		AITEST_TRUE("StateTree should init", bInitSucceeded);
+
+		const TCHAR* TickStr(TEXT("Tick"));
+		const TCHAR* EnterStateStr(TEXT("EnterState"));
+		const TCHAR* ExitStateStr(TEXT("ExitState"));
+		const TCHAR* StateCompletedStr(TEXT("StateCompleted"));
+
+		FInstancedPropertyBag GlobalParameters = StateTree.GetDefaultParameters();
+
+		// Start and enter state
+		Status = Exec.Start();
+		AITEST_TRUE(TEXT("StateTree Task0 should enter state"), Exec.Expect(Task0.GetName(), EnterStateStr));
+		Exec.LogClear();
+
+		// Transition from State1A -> State1A non-reactivate
+		Status = Exec.Tick(0.1f);
+		AITEST_TRUE(TEXT("StateTree Task0 1st tick count 1"), Exec.Expect(Task0.GetName(), TEXT("CurrentTick=1")));
+		AITEST_FALSE(TEXT("StateTree Task0 should not re-enter state 1st"), Exec.Expect(Task0.GetName(), EnterStateStr));
+		Exec.LogClear();
+
+		// Transition from State1A -> State1A non-reactivate - should still hold after second tick 
+		Status = Exec.Tick(0.1f);
+		AITEST_TRUE(TEXT("StateTree Task0 2nd tick count 2"), Exec.Expect(Task0.GetName(), TEXT("CurrentTick=2")));
+		AITEST_FALSE(TEXT("StateTree Task0 should not re-enter state 2nd"), Exec.Expect(Task0.GetName(), EnterStateStr));
+		Exec.LogClear();
+
+		GlobalParameters.SetValueInt32(FName(TEXT("Int")), 0);
+		InstanceData.GetMutableStorage().SetGlobalParameters(GlobalParameters);
+
+		// Transition from State1A -> State1A reactivate
+		Status = Exec.Tick(0.1f);
+		AITEST_TRUE(TEXT("StateTree Task0 tick before 1st re-enter tick count 3"), Exec.Expect(Task0.GetName(), TEXT("CurrentTick=3")));
+		AITEST_TRUE(TEXT("StateTree Task0 should re-enter state 1st"), Exec.Expect(Task0.GetName(), EnterStateStr));
+		Exec.LogClear();
+
+		// Transition from State1A -> State1A reactivate  should still hold after second tick 
+		Status = Exec.Tick(0.1f);
+		AITEST_TRUE(TEXT("StateTree Task0 tick after 1st re-enter tick count 1"), Exec.Expect(Task0.GetName(), TEXT("CurrentTick=1")));
+		AITEST_TRUE(TEXT("StateTree Task0 should re-enter state 2nd"), Exec.Expect(Task0.GetName(), EnterStateStr));
+		Exec.LogClear();
+
+		GlobalParameters.SetValueInt32(FName(TEXT("Int")), 1);
+		InstanceData.GetMutableStorage().SetGlobalParameters(GlobalParameters);
+
+		// Transition from State1A -> State1A final non-reactivate after condition reset
+		Status = Exec.Tick(0.1f);
+		AITEST_TRUE(TEXT("StateTree Task0 tick after 2nd re-enter tick count 1"), Exec.Expect(Task0.GetName(), TEXT("CurrentTick=1")));
+		AITEST_FALSE(TEXT("StateTree Task0 should not re-enter state final"), Exec.Expect(Task0.GetName(), EnterStateStr));
+		Exec.LogClear();
+
+		Exec.Stop();
+
+		return true;
+	}
+};
+IMPLEMENT_STATE_TREE_INSTANT_TEST(FStateTreeTest_TransitionReactivateTargetState, "System.StateTree.Transition.ReactivateTargetState");
+
 struct FStateTreeTest_TransitionDelayZero : FStateTreeTestBase
 {
 	virtual bool InstantTest() override
