Fix rename failing on to be descarded ST editordata schema object.

#tests editor, cook

diff --git a/Engine/Plugins/Runtime/StateTree/Source/StateTreeEditorModule/Private/StateTreeEditingSubsystem.cpp b/Engine/Plugins/Runtime/StateTree/Source/StateTreeEditorModule/Private/StateTreeEditingSubsystem.cpp
--- a/Engine/Plugins/Runtime/StateTree/Source/StateTreeEditorModule/Private/StateTreeEditingSubsystem.cpp
+++ b/Engine/Plugins/Runtime/StateTree/Source/StateTreeEditorModule/Private/StateTreeEditingSubsystem.cpp
@@ -13,6 +13,7 @@
 #include "StateTreeEditorDataExtension.h"
 #include "StateTreeObjectHash.h"
 #include "StateTreeTaskBase.h"
+#include "UObject/LinkerLoad.h"
 #include "UObject/UObjectGlobals.h"
 
 #include UE_INLINE_GENERATED_CPP_BY_NAME(StateTreeEditingSubsystem)
@@ -195,6 +196,10 @@ void UStateTreeEditingSubsystem::ValidateStateTree(TNotNull<UStateTree*> InState
 			const FName TrashSchemaName = *FString::Printf(TEXT("TRASH_%s"), *PreviousEditorData->GetName());
 			constexpr ERenameFlags RenameFlags = REN_DoNotDirty | REN_DontCreateRedirectors;
 
+			// @TODO: UE-348529. This rename below is potentially unsafe. While we `ConditionallyPostload` the editordata in `UStateTree::PostLoad`.
+			// Since this object is intended to be discarded the correct method is to call `FLinkerLoad::InvalidateExport`.
+			// Consider moving `UE::StateTree::Compiler::RenameObjectToTransientPackage` to a public static method in `StateTreeCompiler.h` to 
+			// do this since it calls the aforementioned method. Vefify that this temp transient outer object below isn't needed first.
 			UStateTree* TransientOuter = NewObject<UStateTree>(GetTransientPackage(), TrashStateTreeName, RF_Transient);
 			PreviousEditorData->Rename(*TrashSchemaName.ToString(), TransientOuter, RenameFlags);
 			PreviousEditorData->SetFlags(RF_Transient);
@@ -224,8 +229,12 @@ void UStateTreeEditingSubsystem::ValidateStateTree(TNotNull<UStateTree*> InState
 				const FName TrashSchemaName = *FString::Printf(TEXT("TRASH_%s"), *PreviousEditorSchema->GetName());
 				constexpr ERenameFlags RenameFlags = REN_DoNotDirty | REN_DontCreateRedirectors;
 
-				PreviousEditorSchema->Rename(*TrashSchemaName.ToString(), GetTransientPackage(), RenameFlags);
 				PreviousEditorSchema->SetFlags(RF_Transient);
+				PreviousEditorSchema->ClearFlags(RF_Public | RF_Standalone | RF_ArchetypeObject);
+
+				// We are no longer interested in loading the old schema, so invalidate export.
+				FLinkerLoad::InvalidateExport(PreviousEditorSchema);
+				PreviousEditorSchema->Rename(*TrashSchemaName.ToString(), GetTransientPackage(), RenameFlags);
 			}
 		};
 
