[PropertyBindings] Add promotion copy for arrays. 2 array properties of different types can now be bound if their inner types can be promoted or demoted.

diff --git a/Engine/Plugins/Runtime/PropertyBindingUtils/Source/PropertyBindingUtils/Private/PropertyBindingBindingCollection.cpp b/Engine/Plugins/Runtime/PropertyBindingUtils/Source/PropertyBindingUtils/Private/PropertyBindingBindingCollection.cpp
--- a/Engine/Plugins/Runtime/PropertyBindingUtils/Source/PropertyBindingUtils/Private/PropertyBindingBindingCollection.cpp
+++ b/Engine/Plugins/Runtime/PropertyBindingUtils/Source/PropertyBindingUtils/Private/PropertyBindingBindingCollection.cpp
@@ -629,10 +629,15 @@ bool FPropertyBindingBindingCollection::ResolveCopyInfoBetweenIndirections(const
 	OutCopyInfo.Type = EPropertyCopyType::None;
 
 	// Flip Copy source and copy target if the copy direction is reversed
-	const FProperty* CopySourceProperty = OutCopyInfo.bCopyFromTargetToSource ? TargetProperty : SourceProperty;
-	const UStruct* CopySourceStruct = OutCopyInfo.bCopyFromTargetToSource ? InTargetIndirection.GetContainerStruct() : InSourceIndirection.GetContainerStruct();
-	const FProperty* CopyTargetProperty = OutCopyInfo.bCopyFromTargetToSource ? SourceProperty : TargetProperty;
-	const UStruct* CopyTargetStruct = OutCopyInfo.bCopyFromTargetToSource ? InSourceIndirection.GetContainerStruct() : InTargetIndirection.GetContainerStruct();
+	const FProperty* CopySourceProperty = SourceProperty;
+	const UStruct* CopySourceStruct = InSourceIndirection.GetContainerStruct();
+	const FProperty* CopyTargetProperty = TargetProperty;
+	const UStruct* CopyTargetStruct = InTargetIndirection.GetContainerStruct();
+	if (OutCopyInfo.bCopyFromTargetToSource)
+	{
+		Swap(CopySourceProperty, CopyTargetProperty);
+		Swap(CopySourceStruct, CopyTargetStruct);
+	}
 
 	if (CopySourceProperty == nullptr)
 	{
@@ -678,7 +683,20 @@ bool FPropertyBindingBindingCollection::ResolveCopyInfoBetweenIndirections(const
 		}
 	}
 
-	const UE::PropertyBinding::EPropertyCompatibility Compatibility = UE::PropertyBinding::GetPropertyCompatibility(CopySourceProperty, CopyTargetProperty);
+	UE::PropertyBinding::EPropertyCompatibility Compatibility = UE::PropertyBinding::GetPropertyCompatibility(CopySourceProperty, CopyTargetProperty);
+
+	if (Compatibility == UE::PropertyBinding::EPropertyCompatibility::Promotable)
+	{
+		const FArrayProperty* CopyTargetArrayProperty = CastField<FArrayProperty>(CopyTargetProperty);
+		const FArrayProperty* CopySourceArrayProperty = CastField<FArrayProperty>(CopySourceProperty);
+		if (CopyTargetArrayProperty && CopySourceArrayProperty)
+		{
+			OutCopyInfo.bPromoteArrayCopy = true;
+			CopyTargetProperty = CopyTargetArrayProperty->Inner;
+			CopySourceProperty = CopySourceArrayProperty->Inner;
+			Compatibility = UE::PropertyBinding::GetPropertyCompatibility(CopySourceProperty, CopyTargetProperty);
+		}
+	}
 
 	// Extract underlying types for enums
 	if (const FEnumProperty* EnumPropertyA = CastField<const FEnumProperty>(CopySourceProperty))
@@ -710,8 +728,10 @@ bool FPropertyBindingBindingCollection::ResolveCopyInfoBetweenIndirections(const
 		}
 		else if (CastField<FObjectPropertyBase>(CopyTargetProperty))
 		{
-			if (CopySourceProperty->IsA<FSoftObjectProperty>()
-				&& CopyTargetProperty->IsA<FSoftObjectProperty>())
+			const bool bBothAreSoftObjectProperty = CopySourceProperty->IsA<FSoftObjectProperty>() && CopyTargetProperty->IsA<FSoftObjectProperty>();
+			const bool bBothAreWeakObjectProperty = CopySourceProperty->IsA<FWeakObjectProperty>() && CopyTargetProperty->IsA<FWeakObjectProperty>();
+			const bool bBothAreLazyObjectProperty = CopySourceProperty->IsA<FLazyObjectProperty>() && CopyTargetProperty->IsA<FLazyObjectProperty>();
+			if (bBothAreSoftObjectProperty || bBothAreWeakObjectProperty || bBothAreLazyObjectProperty)
 			{
 				// Use CopyComplex when copying soft object to another soft object so that we do not try to dereference the object (just copies the path).
 				// This handles soft class too.
@@ -909,7 +929,7 @@ uint8* FPropertyBindingBindingCollection::GetAddress(FPropertyBindingDataView In
 		}
 		case EPropertyBindingPropertyAccessType::WeakObject:
 		{
-			TWeakObjectPtr<UObject>& WeakObjectPtr = *reinterpret_cast<TWeakObjectPtr<UObject>*>(Address + Indirection->Offset);
+			FWeakObjectPtr& WeakObjectPtr = *reinterpret_cast<FWeakObjectPtr*>(Address + Indirection->Offset);
 			UObject* Object = WeakObjectPtr.Get();
 			Address = reinterpret_cast<uint8*>(Object);
 			break;
@@ -1018,11 +1038,12 @@ uint8* FPropertyBindingBindingCollection::GetAddress(FPropertyBindingDataView In
 
 bool FPropertyBindingBindingCollection::PerformCopy(const FPropertyBindingCopyInfo& Copy, uint8* SourceAddress, uint8* TargetAddress) const
 {
-	const FProperty* SourceLeafProperty = Copy.bCopyFromTargetToSource ? Copy.TargetLeafProperty : Copy.SourceLeafProperty;
-	const FProperty* TargetLeafProperty = Copy.bCopyFromTargetToSource ? Copy.SourceLeafProperty : Copy.TargetLeafProperty;
+	const FProperty* SourceLeafProperty = Copy.SourceLeafProperty;
+	const FProperty* TargetLeafProperty = Copy.TargetLeafProperty;
 	if (Copy.bCopyFromTargetToSource)
 	{
 		Swap(SourceAddress, TargetAddress);
+		Swap(SourceLeafProperty, TargetLeafProperty);
 	}
 
 	// 'SourceAddress' can only be nullptr for object copy (e.g., EPropertyCopyType::CopyObject)
@@ -1042,145 +1063,174 @@ bool FPropertyBindingBindingCollection::PerformCopy(const FPropertyBindingCopyIn
 
 	check(TargetLeafProperty);
 	
-	switch (Copy.Type)
-	{
-	case EPropertyCopyType::CopyPlain:
-		FMemory::Memcpy(TargetAddress, SourceAddress, Copy.CopySize);
-		break;
-	case EPropertyCopyType::CopyComplex:
-		TargetLeafProperty->CopyCompleteValue(TargetAddress, SourceAddress);
-		break;
-	case EPropertyCopyType::CopyBool:
-		CastFieldChecked<const FBoolProperty>(TargetLeafProperty)->SetPropertyValue(TargetAddress, CastFieldChecked<const FBoolProperty>(SourceLeafProperty)->GetPropertyValue(SourceAddress));
-		break;
-	case EPropertyCopyType::CopyStruct:
-		// If SourceProperty == nullptr (pointing to the struct source directly), the GetAddress() did the right thing and is pointing to the beginning of the struct. 
-		CastFieldChecked<const FStructProperty>(TargetLeafProperty)->Struct->CopyScriptStruct(TargetAddress, SourceAddress);
-		break;
-	case EPropertyCopyType::CopyObject:
-		if (SourceLeafProperty == nullptr || SourceAddress == nullptr)
-		{
-			// Source is pointing at object directly.
-			CastFieldChecked<const FObjectPropertyBase>(TargetLeafProperty)->SetObjectPropertyValue(TargetAddress, reinterpret_cast<UObject*>(SourceAddress));
-		}
-		else
+	auto PerformCopyInternal = [Self = this, &Copy](const FProperty* SourceLeafProperty, uint8* SourceAddress, const FProperty* TargetLeafProperty, uint8* TargetAddress)
 		{
-			CastFieldChecked<const FObjectPropertyBase>(TargetLeafProperty)->SetObjectPropertyValue(TargetAddress, CastFieldChecked<const FObjectPropertyBase>(SourceLeafProperty)->GetObjectPropertyValue(SourceAddress));
-		}
-		break;
-	case EPropertyCopyType::CopyName:
-		CastFieldChecked<const FNameProperty>(TargetLeafProperty)->SetPropertyValue(TargetAddress, CastFieldChecked<const FNameProperty>(SourceLeafProperty)->GetPropertyValue(SourceAddress));
-		break;
-	case EPropertyCopyType::CopyFixedArray:
+			switch (Copy.Type)
+			{
+			case EPropertyCopyType::CopyPlain:
+				FMemory::Memcpy(TargetAddress, SourceAddress, Copy.CopySize);
+				break;
+			case EPropertyCopyType::CopyComplex:
+				TargetLeafProperty->CopyCompleteValue(TargetAddress, SourceAddress);
+				break;
+			case EPropertyCopyType::CopyBool:
+				CastFieldChecked<const FBoolProperty>(TargetLeafProperty)->SetPropertyValue(TargetAddress, CastFieldChecked<const FBoolProperty>(SourceLeafProperty)->GetPropertyValue(SourceAddress));
+				break;
+			case EPropertyCopyType::CopyStruct:
+				// If SourceProperty == nullptr (pointing to the struct source directly), the GetAddress() did the right thing and is pointing to the beginning of the struct. 
+				CastFieldChecked<const FStructProperty>(TargetLeafProperty)->Struct->CopyScriptStruct(TargetAddress, SourceAddress);
+				break;
+			case EPropertyCopyType::CopyObject:
+				if (SourceLeafProperty == nullptr || SourceAddress == nullptr)
+				{
+					// Source is pointing at object directly.
+					CastFieldChecked<const FObjectPropertyBase>(TargetLeafProperty)->SetObjectPropertyValue(TargetAddress, reinterpret_cast<UObject*>(SourceAddress));
+				}
+				else
+				{
+					CastFieldChecked<const FObjectPropertyBase>(TargetLeafProperty)->SetObjectPropertyValue(TargetAddress, CastFieldChecked<const FObjectPropertyBase>(SourceLeafProperty)->GetObjectPropertyValue(SourceAddress));
+				}
+				break;
+			case EPropertyCopyType::CopyName:
+				CastFieldChecked<const FNameProperty>(TargetLeafProperty)->SetPropertyValue(TargetAddress, CastFieldChecked<const FNameProperty>(SourceLeafProperty)->GetPropertyValue(SourceAddress));
+				break;
+			case EPropertyCopyType::CopyFixedArray:
+			{
+				// Copy into fixed sized array (EditFixedSize). Resizable arrays are copied as Complex, and regular fixed sizes arrays via the regular copies (dim specifies array size).
+				const FArrayProperty* SourceArrayProperty = CastFieldChecked<const FArrayProperty>(SourceLeafProperty);
+				const FArrayProperty* TargetArrayProperty = CastFieldChecked<const FArrayProperty>(TargetLeafProperty);
+				FScriptArrayHelper SourceArrayHelper(SourceArrayProperty, SourceAddress);
+				FScriptArrayHelper TargetArrayHelper(TargetArrayProperty, TargetAddress);
+			
+				const int32 MinSize = FMath::Min(SourceArrayHelper.Num(), TargetArrayHelper.Num());
+				for (int32 ElementIndex = 0; ElementIndex < MinSize; ++ElementIndex)
+				{
+					TargetArrayProperty->Inner->CopySingleValue(TargetArrayHelper.GetRawPtr(ElementIndex), SourceArrayHelper.GetRawPtr(ElementIndex));
+				}
+				break;
+			}
+			case EPropertyCopyType::StructReference:
+			{
+				checkf(Self->PropertyReferenceCopyFunc, TEXT("Not expecting EPropertyCopyType::StructReference if copy functor was not provided"));
+				const FStructProperty* SourceStructProperty = CastFieldChecked<const FStructProperty>(SourceLeafProperty);
+				if (ensure(SourceStructProperty))
+				{
+					Self->PropertyReferenceCopyFunc(*SourceStructProperty, SourceAddress, TargetAddress);
+				}
+				break;
+			}
+			// Bool promotions
+			case EPropertyCopyType::PromoteBoolToByte:
+				*reinterpret_cast<uint8*>(TargetAddress) = (uint8)CastFieldChecked<const FBoolProperty>(SourceLeafProperty)->GetPropertyValue(SourceAddress);
+				break;
+			case EPropertyCopyType::PromoteBoolToInt32:
+				*reinterpret_cast<int32*>(TargetAddress) = (int32)CastFieldChecked<const FBoolProperty>(SourceLeafProperty)->GetPropertyValue(SourceAddress);
+				break;
+			case EPropertyCopyType::PromoteBoolToUInt32:
+				*reinterpret_cast<uint32*>(TargetAddress) = (uint32)CastFieldChecked<const FBoolProperty>(SourceLeafProperty)->GetPropertyValue(SourceAddress);
+				break;
+			case EPropertyCopyType::PromoteBoolToInt64:
+				*reinterpret_cast<int64*>(TargetAddress) = (int64)CastFieldChecked<const FBoolProperty>(SourceLeafProperty)->GetPropertyValue(SourceAddress);
+				break;
+			case EPropertyCopyType::PromoteBoolToFloat:
+				*reinterpret_cast<float*>(TargetAddress) = (float)CastFieldChecked<const FBoolProperty>(SourceLeafProperty)->GetPropertyValue(SourceAddress);
+				break;
+			case EPropertyCopyType::PromoteBoolToDouble:
+				*reinterpret_cast<double*>(TargetAddress) = (double)CastFieldChecked<const FBoolProperty>(SourceLeafProperty)->GetPropertyValue(SourceAddress);
+				break;
+		
+			// Byte promotions	
+			case EPropertyCopyType::PromoteByteToInt32:
+				*reinterpret_cast<int32*>(TargetAddress) = (int32)*reinterpret_cast<const uint8*>(SourceAddress);
+				break;
+			case EPropertyCopyType::PromoteByteToUInt32:
+				*reinterpret_cast<uint32*>(TargetAddress) = (uint32)*reinterpret_cast<const uint8*>(SourceAddress);
+				break;
+			case EPropertyCopyType::PromoteByteToInt64:
+				*reinterpret_cast<int64*>(TargetAddress) = (int64)*reinterpret_cast<const uint8*>(SourceAddress);
+				break;
+			case EPropertyCopyType::PromoteByteToFloat:
+				*reinterpret_cast<float*>(TargetAddress) = (float)*reinterpret_cast<const uint8*>(SourceAddress);
+				break;
+			case EPropertyCopyType::PromoteByteToDouble:
+				*reinterpret_cast<double*>(TargetAddress) = (double)*reinterpret_cast<const uint8*>(SourceAddress);
+				break;
+
+			// Int32 promotions
+			case EPropertyCopyType::PromoteInt32ToInt64:
+				*reinterpret_cast<int64*>(TargetAddress) = (int64)*reinterpret_cast<const int32*>(SourceAddress);
+				break;
+			case EPropertyCopyType::PromoteInt32ToFloat:
+				*reinterpret_cast<float*>(TargetAddress) = (float)*reinterpret_cast<const int32*>(SourceAddress);
+				break;
+			case EPropertyCopyType::PromoteInt32ToDouble:
+				*reinterpret_cast<double*>(TargetAddress) = (double)*reinterpret_cast<const int32*>(SourceAddress);
+				break;
+
+			// Uint32 promotions
+			case EPropertyCopyType::PromoteUInt32ToInt64:
+				*reinterpret_cast<int64*>(TargetAddress) = (int64)*reinterpret_cast<const uint32*>(SourceAddress);
+				break;
+			case EPropertyCopyType::PromoteUInt32ToFloat:
+				*reinterpret_cast<float*>(TargetAddress) = (float)*reinterpret_cast<const uint32*>(SourceAddress);
+				break;
+			case EPropertyCopyType::PromoteUInt32ToDouble:
+				*reinterpret_cast<double*>(TargetAddress) = (double)*reinterpret_cast<const uint32*>(SourceAddress);
+				break;
+
+			// Float promotions
+			case EPropertyCopyType::PromoteFloatToInt32:
+				*reinterpret_cast<int32*>(TargetAddress) = (int32)*reinterpret_cast<const float*>(SourceAddress);
+				break;
+			case EPropertyCopyType::PromoteFloatToInt64:
+				*reinterpret_cast<int64*>(TargetAddress) = (int64)*reinterpret_cast<const float*>(SourceAddress);
+				break;
+			case EPropertyCopyType::PromoteFloatToDouble:
+				*reinterpret_cast<double*>(TargetAddress) = (double)*reinterpret_cast<const float*>(SourceAddress);
+				break;
+
+			// Double promotions
+			case EPropertyCopyType::DemoteDoubleToInt32:
+				*reinterpret_cast<int32*>(TargetAddress) = (int32)*reinterpret_cast<const double*>(SourceAddress);
+				break;
+			case EPropertyCopyType::DemoteDoubleToInt64:
+				*reinterpret_cast<int64*>(TargetAddress) = (int64)*reinterpret_cast<const double*>(SourceAddress);
+				break;
+			case EPropertyCopyType::DemoteDoubleToFloat:
+				*reinterpret_cast<float*>(TargetAddress) = (float)*reinterpret_cast<const double*>(SourceAddress);
+				break;
+
+			default:
+				ensureMsgf(false, TEXT("FStateTreePropertyBindings::PerformCopy: Unhandled copy type %s between '%s' and '%s'"),
+					*StaticEnum<EPropertyCopyType>()->GetValueAsString(Copy.Type), *SourceLeafProperty->GetNameCPP(), *TargetLeafProperty->GetNameCPP());
+				break;
+			}
+		};
+
+	if (!Copy.bPromoteArrayCopy)
+	{
+		PerformCopyInternal(SourceLeafProperty, SourceAddress, TargetLeafProperty, TargetAddress);
+	}
+	else
 	{
-		// Copy into fixed sized array (EditFixedSize). Resizable arrays are copied as Complex, and regular fixed sizes arrays via the regular copies (dim specifies array size).
 		const FArrayProperty* SourceArrayProperty = CastFieldChecked<const FArrayProperty>(SourceLeafProperty);
 		const FArrayProperty* TargetArrayProperty = CastFieldChecked<const FArrayProperty>(TargetLeafProperty);
 		FScriptArrayHelper SourceArrayHelper(SourceArrayProperty, SourceAddress);
 		FScriptArrayHelper TargetArrayHelper(TargetArrayProperty, TargetAddress);
-			
-		const int32 MinSize = FMath::Min(SourceArrayHelper.Num(), TargetArrayHelper.Num());
-		for (int32 ElementIndex = 0; ElementIndex < MinSize; ++ElementIndex)
+
+		const int32 NumberOfElements = SourceArrayHelper.Num();
+		if (TargetArrayProperty->HasAnyPropertyFlags(CPF_EditFixedSize) && TargetArrayHelper.Num() != NumberOfElements)
 		{
-			TargetArrayProperty->Inner->CopySingleValue(TargetArrayHelper.GetRawPtr(ElementIndex), SourceArrayHelper.GetRawPtr(ElementIndex));
+			// The target is of fixed size. The size of the source need to match. It would be bad to copy half or none of the property.
+			PROPERTY_BINDINGS_LOG(GetLogOwner(), Verbose, TEXT("%hs skipped: the target array size do not match '%d'."), __FUNCTION__, NumberOfElements);
+			return false;
 		}
-		break;
-	}
-	case EPropertyCopyType::StructReference:
-	{
-		checkf(PropertyReferenceCopyFunc, TEXT("Not expecting EPropertyCopyType::StructReference if copy functor was not provided"));
-		const FStructProperty* SourceStructProperty = CastFieldChecked<const FStructProperty>(SourceLeafProperty);
-		if (ensure(SourceStructProperty))
+
+		TargetArrayHelper.Resize(NumberOfElements);
+		for (int32 ElementIndex = 0; ElementIndex < NumberOfElements; ++ElementIndex)
 		{
-			PropertyReferenceCopyFunc(*SourceStructProperty, SourceAddress, TargetAddress);
+			PerformCopyInternal(SourceArrayProperty->Inner, SourceArrayHelper.GetRawPtr(ElementIndex), TargetArrayProperty->Inner, TargetArrayHelper.GetRawPtr(ElementIndex));
 		}
-		break;
-	}
-	// Bool promotions
-	case EPropertyCopyType::PromoteBoolToByte:
-		*reinterpret_cast<uint8*>(TargetAddress) = (uint8)CastFieldChecked<const FBoolProperty>(SourceLeafProperty)->GetPropertyValue(SourceAddress);
-		break;
-	case EPropertyCopyType::PromoteBoolToInt32:
-		*reinterpret_cast<int32*>(TargetAddress) = (int32)CastFieldChecked<const FBoolProperty>(SourceLeafProperty)->GetPropertyValue(SourceAddress);
-		break;
-	case EPropertyCopyType::PromoteBoolToUInt32:
-		*reinterpret_cast<uint32*>(TargetAddress) = (uint32)CastFieldChecked<const FBoolProperty>(SourceLeafProperty)->GetPropertyValue(SourceAddress);
-		break;
-	case EPropertyCopyType::PromoteBoolToInt64:
-		*reinterpret_cast<int64*>(TargetAddress) = (int64)CastFieldChecked<const FBoolProperty>(SourceLeafProperty)->GetPropertyValue(SourceAddress);
-		break;
-	case EPropertyCopyType::PromoteBoolToFloat:
-		*reinterpret_cast<float*>(TargetAddress) = (float)CastFieldChecked<const FBoolProperty>(SourceLeafProperty)->GetPropertyValue(SourceAddress);
-		break;
-	case EPropertyCopyType::PromoteBoolToDouble:
-		*reinterpret_cast<double*>(TargetAddress) = (double)CastFieldChecked<const FBoolProperty>(SourceLeafProperty)->GetPropertyValue(SourceAddress);
-		break;
-		
-	// Byte promotions	
-	case EPropertyCopyType::PromoteByteToInt32:
-		*reinterpret_cast<int32*>(TargetAddress) = (int32)*reinterpret_cast<const uint8*>(SourceAddress);
-		break;
-	case EPropertyCopyType::PromoteByteToUInt32:
-		*reinterpret_cast<uint32*>(TargetAddress) = (uint32)*reinterpret_cast<const uint8*>(SourceAddress);
-		break;
-	case EPropertyCopyType::PromoteByteToInt64:
-		*reinterpret_cast<int64*>(TargetAddress) = (int64)*reinterpret_cast<const uint8*>(SourceAddress);
-		break;
-	case EPropertyCopyType::PromoteByteToFloat:
-		*reinterpret_cast<float*>(TargetAddress) = (float)*reinterpret_cast<const uint8*>(SourceAddress);
-		break;
-	case EPropertyCopyType::PromoteByteToDouble:
-		*reinterpret_cast<double*>(TargetAddress) = (double)*reinterpret_cast<const uint8*>(SourceAddress);
-		break;
-
-	// Int32 promotions
-	case EPropertyCopyType::PromoteInt32ToInt64:
-		*reinterpret_cast<int64*>(TargetAddress) = (int64)*reinterpret_cast<const int32*>(SourceAddress);
-		break;
-	case EPropertyCopyType::PromoteInt32ToFloat:
-		*reinterpret_cast<float*>(TargetAddress) = (float)*reinterpret_cast<const int32*>(SourceAddress);
-		break;
-	case EPropertyCopyType::PromoteInt32ToDouble:
-		*reinterpret_cast<double*>(TargetAddress) = (double)*reinterpret_cast<const int32*>(SourceAddress);
-		break;
-
-	// Uint32 promotions
-	case EPropertyCopyType::PromoteUInt32ToInt64:
-		*reinterpret_cast<int64*>(TargetAddress) = (int64)*reinterpret_cast<const uint32*>(SourceAddress);
-		break;
-	case EPropertyCopyType::PromoteUInt32ToFloat:
-		*reinterpret_cast<float*>(TargetAddress) = (float)*reinterpret_cast<const uint32*>(SourceAddress);
-		break;
-	case EPropertyCopyType::PromoteUInt32ToDouble:
-		*reinterpret_cast<double*>(TargetAddress) = (double)*reinterpret_cast<const uint32*>(SourceAddress);
-		break;
-
-	// Float promotions
-	case EPropertyCopyType::PromoteFloatToInt32:
-		*reinterpret_cast<int32*>(TargetAddress) = (int32)*reinterpret_cast<const float*>(SourceAddress);
-		break;
-	case EPropertyCopyType::PromoteFloatToInt64:
-		*reinterpret_cast<int64*>(TargetAddress) = (int64)*reinterpret_cast<const float*>(SourceAddress);
-		break;
-	case EPropertyCopyType::PromoteFloatToDouble:
-		*reinterpret_cast<double*>(TargetAddress) = (double)*reinterpret_cast<const float*>(SourceAddress);
-		break;
-
-	// Double promotions
-	case EPropertyCopyType::DemoteDoubleToInt32:
-		*reinterpret_cast<int32*>(TargetAddress) = (int32)*reinterpret_cast<const double*>(SourceAddress);
-		break;
-	case EPropertyCopyType::DemoteDoubleToInt64:
-		*reinterpret_cast<int64*>(TargetAddress) = (int64)*reinterpret_cast<const double*>(SourceAddress);
-		break;
-	case EPropertyCopyType::DemoteDoubleToFloat:
-		*reinterpret_cast<float*>(TargetAddress) = (float)*reinterpret_cast<const double*>(SourceAddress);
-		break;
-
-	default:
-		ensureMsgf(false, TEXT("FStateTreePropertyBindings::PerformCopy: Unhandled copy type %s between '%s' and '%s'"),
-			*StaticEnum<EPropertyCopyType>()->GetValueAsString(Copy.Type), *SourceLeafProperty->GetNameCPP(), *TargetLeafProperty->GetNameCPP());
-		break;
 	}
 
 	return true;
@@ -1228,36 +1278,53 @@ void FPropertyBindingBindingCollection::PerformResetObjects(const FPropertyBindi
 	check(Copy.TargetLeafProperty);
 	check(TargetAddress);
 
-	switch (Copy.Type)
+	auto PerformResetObjectsInternal = [Self = this, Type = Copy.Type](const FProperty* TargetLeafProperty, uint8* TargetAddress)
+		{
+			switch (Type)
+			{
+			case EPropertyCopyType::CopyComplex:
+				TargetLeafProperty->ClearValue(TargetAddress);
+				break;
+			case EPropertyCopyType::CopyStruct:
+				CastFieldChecked<const FStructProperty>(TargetLeafProperty)->Struct->ClearScriptStruct(TargetAddress);
+				break;
+			case EPropertyCopyType::CopyObject:
+				CastFieldChecked<const FObjectPropertyBase>(TargetLeafProperty)->SetObjectPropertyValue(TargetAddress, nullptr);
+				break;
+			case EPropertyCopyType::StructReference:
+				checkf(Self->PropertyReferenceResetFunc, TEXT("Not expecting EPropertyCopyType::StructReference if reset object functor was not provided"));
+				Self->PropertyReferenceResetFunc(TargetAddress);
+				break;
+			case EPropertyCopyType::CopyFixedArray:
+			{
+				// Copy into fixed sized array (EditFixedSize). Resizable arrays are copied as Complex, and regular fixed sizes arrays via the regular copies (dim specifies array size).
+				const FArrayProperty* TargetArrayProperty = CastFieldChecked<const FArrayProperty>(TargetLeafProperty);
+				FScriptArrayHelper TargetArrayHelper(TargetArrayProperty, TargetAddress);
+				for (int32 ElementIndex = 0; ElementIndex < TargetArrayHelper.Num(); ++ElementIndex)
+				{
+					TargetArrayProperty->Inner->ClearValue(TargetArrayHelper.GetRawPtr(ElementIndex));
+				}
+				break;
+			}
+			default:
+				break;
+			}
+		};
+
+	if (!Copy.bPromoteArrayCopy)
 	{
-	case EPropertyCopyType::CopyComplex:
-		Copy.TargetLeafProperty->ClearValue(TargetAddress);
-		break;
-	case EPropertyCopyType::CopyStruct:
-		CastFieldChecked<const FStructProperty>(Copy.TargetLeafProperty)->Struct->ClearScriptStruct(TargetAddress);
-		break;
-	case EPropertyCopyType::CopyObject:
-		CastFieldChecked<const FObjectPropertyBase>(Copy.TargetLeafProperty)->SetObjectPropertyValue(TargetAddress, nullptr);
-		break;
-	case EPropertyCopyType::StructReference:
-		checkf(PropertyReferenceResetFunc, TEXT("Not expecting EPropertyCopyType::StructReference if reset object functor was not provided"));
-		PropertyReferenceResetFunc(TargetAddress);
-		break;
-	case EPropertyCopyType::CopyName:
-		break;
-	case EPropertyCopyType::CopyFixedArray:
+		PerformResetObjectsInternal(Copy.TargetLeafProperty, TargetAddress);
+	}
+	else
 	{
-		// Copy into fixed sized array (EditFixedSize). Resizable arrays are copied as Complex, and regular fixed sizes arrays via the regular copies (dim specifies array size).
 		const FArrayProperty* TargetArrayProperty = CastFieldChecked<const FArrayProperty>(Copy.TargetLeafProperty);
 		FScriptArrayHelper TargetArrayHelper(TargetArrayProperty, TargetAddress);
-		for (int32 ElementIndex = 0; ElementIndex < TargetArrayHelper.Num(); ++ElementIndex)
+		const int32 NumberOfElements = TargetArrayHelper.Num();
+		for (int32 ElementIndex = 0; ElementIndex < NumberOfElements; ++ElementIndex)
 		{
-			TargetArrayProperty->Inner->ClearValue(TargetArrayHelper.GetRawPtr(ElementIndex));
+			PerformResetObjectsInternal(TargetArrayProperty->Inner, TargetArrayHelper.GetRawPtr(ElementIndex));
 		}
-		break;
-	}
-	default:
-		break;
+
 	}
 }
 
diff --git a/Engine/Plugins/Runtime/PropertyBindingUtils/Source/PropertyBindingUtils/Private/PropertyBindingPath.cpp b/Engine/Plugins/Runtime/PropertyBindingUtils/Source/PropertyBindingUtils/Private/PropertyBindingPath.cpp
--- a/Engine/Plugins/Runtime/PropertyBindingUtils/Source/PropertyBindingUtils/Private/PropertyBindingPath.cpp
+++ b/Engine/Plugins/Runtime/PropertyBindingUtils/Source/PropertyBindingUtils/Private/PropertyBindingPath.cpp
@@ -573,7 +573,7 @@ bool FPropertyBindingPath::ResolveIndirectionsWithValue(const FPropertyBindingDa
 			{
 				if (bWithValue)
 				{
-					const TWeakObjectPtr<UObject>& WeakObjectPtr = *reinterpret_cast<const TWeakObjectPtr<UObject>*>(CurrentAddress + Offset);
+					const FWeakObjectPtr& WeakObjectPtr = *reinterpret_cast<const FWeakObjectPtr*>(CurrentAddress + Offset);
 					const UObject* Object = WeakObjectPtr.Get();
 					CurrentAddress = reinterpret_cast<const uint8*>(Object);
 
diff --git a/Engine/Plugins/Runtime/PropertyBindingUtils/Source/PropertyBindingUtils/Private/PropertyBindingTypes.cpp b/Engine/Plugins/Runtime/PropertyBindingUtils/Source/PropertyBindingUtils/Private/PropertyBindingTypes.cpp
--- a/Engine/Plugins/Runtime/PropertyBindingUtils/Source/PropertyBindingUtils/Private/PropertyBindingTypes.cpp
+++ b/Engine/Plugins/Runtime/PropertyBindingUtils/Source/PropertyBindingUtils/Private/PropertyBindingTypes.cpp
@@ -126,6 +126,18 @@ UE::PropertyBinding::EPropertyCompatibility UE::PropertyBinding::GetPropertyComp
 				return EPropertyCompatibility::Promotable;
 			}
 		}
+		else
+		{
+			const FArrayProperty* FromArrayProperty = CastField<const FArrayProperty>(FromProperty);
+			const FArrayProperty* ToArrayProperty = CastField<const FArrayProperty>(ToProperty);
+			if (FromArrayProperty && ToArrayProperty)
+			{
+				if (UE::PropertyBinding::GetPropertyCompatibility(FromArrayProperty->Inner, ToArrayProperty->Inner) != EPropertyCompatibility::Incompatible)
+				{
+					return EPropertyCompatibility::Promotable;
+				}
+			}
+		}
 	}
 
 	return EPropertyCompatibility::Incompatible;
diff --git a/Engine/Plugins/Runtime/PropertyBindingUtils/Source/PropertyBindingUtils/Public/PropertyBindingPath.h b/Engine/Plugins/Runtime/PropertyBindingUtils/Source/PropertyBindingUtils/Public/PropertyBindingPath.h
--- a/Engine/Plugins/Runtime/PropertyBindingUtils/Source/PropertyBindingUtils/Public/PropertyBindingPath.h
+++ b/Engine/Plugins/Runtime/PropertyBindingUtils/Source/PropertyBindingUtils/Public/PropertyBindingPath.h
@@ -560,7 +560,14 @@ struct FPropertyBindingCopyInfo
 	UPROPERTY()
 	EPropertyCopyType Type = EPropertyCopyType::None;
 
-	/** If true, Copy will be performed from target property to the source property. */
+	/**
+	 * The source property and the target property are array properties.
+	 * The inner type of the source property can be promoted to the inner type of the target property.
+	 */
+	UPROPERTY()
+	bool bPromoteArrayCopy = false;
+
+	/** Copy will be performed from target property to the source property. */
 	UPROPERTY()
 	bool bCopyFromTargetToSource = false;
 };
diff --git a/Engine/Plugins/Runtime/StateTree/Source/StateTreeTestSuite/Private/StateTreeBindingTest.cpp b/Engine/Plugins/Runtime/StateTree/Source/StateTreeTestSuite/Private/StateTreeBindingTest.cpp
--- a/Engine/Plugins/Runtime/StateTree/Source/StateTreeTestSuite/Private/StateTreeBindingTest.cpp
+++ b/Engine/Plugins/Runtime/StateTree/Source/StateTreeTestSuite/Private/StateTreeBindingTest.cpp
@@ -66,9 +66,21 @@ struct FStateTreeTest_BindingsCompiler : FStateTreeTestBase
 		PropertyBindings.Add(MakeBinding(SourceADesc.ID, TEXT("CArray"), TargetDesc.ID, TEXT("CArray")));
 
 		int32 CopyBatchIndex = INDEX_NONE;
-		const bool bCompileBatchResult = BindingCompiler.CompileBatch(TargetDesc, PropertyBindings, FStateTreeIndex16::Invalid, FStateTreeIndex16::Invalid, CopyBatchIndex);
-		AITEST_TRUE(TEXT("CompileBatch should succeed"), bCompileBatchResult);
-		AITEST_NOT_EQUAL(TEXT("CopyBatchIndex should not be INDEX_NONE"), CopyBatchIndex, (int32)INDEX_NONE);
+		{
+			const bool bCompileBatchResult = BindingCompiler.CompileBatch(TargetDesc, PropertyBindings, FStateTreeIndex16::Invalid, FStateTreeIndex16::Invalid, CopyBatchIndex);
+			AITEST_TRUE(TEXT("CompileBatch should succeed"), bCompileBatchResult);
+			AITEST_NOT_EQUAL(TEXT("CopyBatchIndex should not be INDEX_NONE"), CopyBatchIndex, (int32)INDEX_NONE);
+		}
+
+		PropertyBindings.Reset();
+		PropertyBindings.Add(MakeBinding(SourceADesc.ID, TEXT("ArrayOfFloat"), TargetDesc.ID, TEXT("ArrayOfDouble")));
+		PropertyBindings.Add(MakeBinding(SourceADesc.ID, TEXT("ArrayOfDouble"), TargetDesc.ID, TEXT("ArrayOfFloat")));
+		int32 CopyBatchIndex2 = INDEX_NONE;
+		{
+			const bool bCompileBatchResult = BindingCompiler.CompileBatch(TargetDesc, PropertyBindings, FStateTreeIndex16::Invalid, FStateTreeIndex16::Invalid, CopyBatchIndex2);
+			AITEST_TRUE(TEXT("CompileBatch 2 should succeed"), bCompileBatchResult);
+			AITEST_NOT_EQUAL(TEXT("CopyBatchIndex2 should not be INDEX_NONE"), CopyBatchIndex2, (int32)INDEX_NONE);
+		}
 
 		BindingCompiler.Finalize();
 
@@ -88,6 +100,12 @@ struct FStateTreeTest_BindingsCompiler : FStateTreeTestBase
 		SourceA.CArray[0].A = 1;
 		SourceA.CArray[0].B = 2;
 
+		SourceA.ArrayOfFloat.Add(1.0f);
+		SourceA.ArrayOfFloat.Add(2.0f);
+		SourceA.ArrayOfFloat.Add(3.0f);
+		SourceA.ArrayOfDouble.Add(11.0);
+		SourceA.ArrayOfDouble.Add(12.0);
+
 		FStateTreeTest_PropertyCopy SourceB;
 		SourceB.Item.A = 41;
 		SourceB.Item.B = 42;
@@ -105,12 +123,14 @@ struct FStateTreeTest_BindingsCompiler : FStateTreeTestBase
 		SourceViews[SourceBIndex] = FStateTreeDataView(FStructView::Make(SourceB));
 		FPropertyBindingDataView TargetView(FStructView::Make(Target));
 
-		bool bCopyResult = true;
-		for (const FPropertyBindingCopyInfo& Copy : Bindings.Super::GetBatchCopies(FPropertyBindingIndex16(CopyBatchIndex)))
 		{
-			bCopyResult &= Bindings.Super::CopyProperty(Copy, SourceViews[Copy.SourceDataHandle.Get<FStateTreeDataHandle>().GetIndex()], TargetView);
+			bool bCopyResult = true;
+			for (const FPropertyBindingCopyInfo& Copy : Bindings.Super::GetBatchCopies(FPropertyBindingIndex16(CopyBatchIndex)))
+			{
+				bCopyResult &= Bindings.Super::CopyProperty(Copy, SourceViews[Copy.SourceDataHandle.Get<FStateTreeDataHandle>().GetIndex()], TargetView);
+			}
+			AITEST_TRUE(TEXT("CopyTo should succeed"), bCopyResult);
 		}
-		AITEST_TRUE(TEXT("CopyTo should succeed"), bCopyResult);
 
 		// Due to binding sorting, we expect them to execute in this order (sorted based on target access, earliest to latest)
 		// SourceA.CArray -> Target.CArray
@@ -148,11 +168,439 @@ struct FStateTreeTest_BindingsCompiler : FStateTreeTestBase
 		AITEST_EQUAL(TEXT("Expect Target fixed size Array to not have changed size."), Target.FixedArray.Num(), FixedArraySize);
 		AITEST_NOT_EQUAL(TEXT("Expect the count of constructed FStateTreeTest_PropertyStructB to be smaller after calling ResetObjects"), FStateTreeTest_PropertyStructB::NumConstructed, NumAllocated_FStateTreeTest_PropertyStructB_BeforeReset);
 
+
+		{
+			bool bCopyResult = true;
+			for (const FPropertyBindingCopyInfo& Copy : Bindings.Super::GetBatchCopies(FPropertyBindingIndex16(CopyBatchIndex2)))
+			{
+				bCopyResult &= Bindings.Super::CopyProperty(Copy, SourceViews[Copy.SourceDataHandle.Get<FStateTreeDataHandle>().GetIndex()], TargetView);
+			}
+			AITEST_TRUE(TEXT("CopyTo should succeed"), bCopyResult);
+		}
+
+		// SourceA.ArrayOfFloat -> Target.ArrayOfDouble
+		// SourceA.ArrayOfDouble -> Target.ArrayOfFloat
+		bool bDoubleEqualFloat = Target.ArrayOfDouble.Num() == SourceA.ArrayOfFloat.Num();
+		if (bDoubleEqualFloat)
+		{
+			for (int32 Index = 0; Index < Target.ArrayOfDouble.Num(); ++Index)
+			{
+				if ((int32)SourceA.ArrayOfFloat[Index] != (int32)Target.ArrayOfDouble[Index])
+				{
+					bDoubleEqualFloat = false;
+					break;
+				}
+			}
+		}
+		AITEST_TRUE(TEXT("Expect Target.ArrayOfDouble copied from SourceA.ArrayOfFloat"), bDoubleEqualFloat);
+
+		bool bFloatEqualDouble = Target.ArrayOfFloat.Num() == SourceA.ArrayOfDouble.Num();
+		if (bFloatEqualDouble)
+		{
+			for (int32 Index = 0; Index < Target.ArrayOfFloat.Num(); ++Index)
+			{
+				if ((int32)SourceA.ArrayOfDouble[Index] != (int32)Target.ArrayOfFloat[Index])
+				{
+					bFloatEqualDouble = false;
+					break;
+				}
+			}
+		}
+		AITEST_TRUE(TEXT("Expect Target.ArrayOfFloat copied from SourceA.ArrayOfDouble"), bFloatEqualDouble);
+
 		return true;
 	}
 };
 IMPLEMENT_STATE_TREE_INSTANT_TEST(FStateTreeTest_BindingsCompiler, "System.StateTree.Binding.BindingsCompiler");
 
+struct FStateTreeTest_BindingsCompiler_WithObjectPtr : FStateTreeTestBase
+{
+	virtual bool InstantTest() override
+	{
+		FStateTreeCompilerLog Log;
+		FStateTreePropertyBindings Bindings;
+		FStateTreePropertyBindingCompiler BindingCompiler;
+
+		const bool bInitResult = BindingCompiler.Init(Bindings, Log);
+		AITEST_TRUE(TEXT("Expect init to succeed"), bInitResult);
+
+		FStateTreeBindableStructDesc SourceADesc;
+		SourceADesc.Name = FName(TEXT("SourceA"));
+		SourceADesc.Struct = TBaseStructure<FStateTreeTest_PropertyStructObjectPtr>::Get();
+		SourceADesc.DataSource = EStateTreeBindableStructSource::Parameter;
+		SourceADesc.DataHandle = FStateTreeDataHandle(EStateTreeDataSourceType::ContextData, 0); // Used as index to SourceViews below.
+		SourceADesc.ID = FGuid::NewGuid();
+
+		FStateTreeBindableStructDesc TargetDesc;
+		TargetDesc.Name = FName(TEXT("Target"));
+		TargetDesc.Struct = TBaseStructure<FStateTreeTest_PropertyStructObjectPtr>::Get();
+		TargetDesc.DataSource = EStateTreeBindableStructSource::Parameter;
+		TargetDesc.ID = FGuid::NewGuid();
+
+		const int32 SourceAIndex = BindingCompiler.AddSourceStruct(SourceADesc);
+
+		auto CompileBatch=[this, &BindingCompiler , &TargetDesc](TArrayView<FStateTreePropertyPathBinding> PropertyBindings, int32& CopyBatchIndex)
+			{
+				CopyBatchIndex = INDEX_NONE;
+				const bool bCompileBatchResult = BindingCompiler.CompileBatch(TargetDesc, PropertyBindings, FStateTreeIndex16::Invalid, FStateTreeIndex16::Invalid, CopyBatchIndex);
+				AITEST_TRUE(TEXT("CompileBatch should succeed"), bCompileBatchResult);
+				AITEST_NOT_EQUAL(TEXT("CopyBatchIndex should not be INDEX_NONE"), CopyBatchIndex, (int32)INDEX_NONE);
+				return true;
+			};
+
+		// A to ...
+		int32 CopyBatchIndex1 = INDEX_NONE;
+		{
+			TArray<FStateTreePropertyPathBinding> PropertyBindings;
+			PropertyBindings.Add(MakeBinding(SourceADesc.ID, TEXT("A"), TargetDesc.ID, TEXT("A")));
+			PropertyBindings.Add(MakeBinding(SourceADesc.ID, TEXT("A"), TargetDesc.ID, TEXT("B")));
+			PropertyBindings.Add(MakeBinding(SourceADesc.ID, TEXT("A"), TargetDesc.ID, TEXT("C")));
+			PropertyBindings.Add(MakeBinding(SourceADesc.ID, TEXT("A"), TargetDesc.ID, TEXT("D")));
+			PropertyBindings.Add(MakeBinding(SourceADesc.ID, TEXT("A"), TargetDesc.ID, TEXT("ArrayA[1]")));
+			PropertyBindings.Add(MakeBinding(SourceADesc.ID, TEXT("A"), TargetDesc.ID, TEXT("ArrayB[1]")));
+			PropertyBindings.Add(MakeBinding(SourceADesc.ID, TEXT("A"), TargetDesc.ID, TEXT("ArrayC[1]")));
+			PropertyBindings.Add(MakeBinding(SourceADesc.ID, TEXT("A"), TargetDesc.ID, TEXT("ArrayD[1]")));
+			if (!CompileBatch(PropertyBindings, CopyBatchIndex1))
+			{
+				return false;
+			}
+		}
+		// ArrayA[] to ...
+		int32 CopyBatchIndex2 = INDEX_NONE;
+		{
+			TArray<FStateTreePropertyPathBinding> PropertyBindings;
+			PropertyBindings.Add(MakeBinding(SourceADesc.ID, TEXT("ArrayA[2]"), TargetDesc.ID, TEXT("A")));
+			PropertyBindings.Add(MakeBinding(SourceADesc.ID, TEXT("ArrayA[2]"), TargetDesc.ID, TEXT("B")));
+			PropertyBindings.Add(MakeBinding(SourceADesc.ID, TEXT("ArrayA[2]"), TargetDesc.ID, TEXT("C")));
+			PropertyBindings.Add(MakeBinding(SourceADesc.ID, TEXT("ArrayA[2]"), TargetDesc.ID, TEXT("D")));
+			PropertyBindings.Add(MakeBinding(SourceADesc.ID, TEXT("ArrayA[2]"), TargetDesc.ID, TEXT("ArrayA[1]")));
+			PropertyBindings.Add(MakeBinding(SourceADesc.ID, TEXT("ArrayA[2]"), TargetDesc.ID, TEXT("ArrayB[1]")));
+			PropertyBindings.Add(MakeBinding(SourceADesc.ID, TEXT("ArrayA[2]"), TargetDesc.ID, TEXT("ArrayC[1]")));
+			PropertyBindings.Add(MakeBinding(SourceADesc.ID, TEXT("ArrayA[2]"), TargetDesc.ID, TEXT("ArrayD[1]")));
+			if (!CompileBatch(PropertyBindings, CopyBatchIndex2))
+			{
+				return false;
+			}
+		}
+		// B to ...
+		int32 CopyBatchIndex3 = INDEX_NONE;
+		{
+			TArray<FStateTreePropertyPathBinding> PropertyBindings;
+			PropertyBindings.Add(MakeBinding(SourceADesc.ID, TEXT("B"), TargetDesc.ID, TEXT("A")));
+			PropertyBindings.Add(MakeBinding(SourceADesc.ID, TEXT("B"), TargetDesc.ID, TEXT("B")));
+			PropertyBindings.Add(MakeBinding(SourceADesc.ID, TEXT("B"), TargetDesc.ID, TEXT("C")));
+			PropertyBindings.Add(MakeBinding(SourceADesc.ID, TEXT("B"), TargetDesc.ID, TEXT("D")));
+			PropertyBindings.Add(MakeBinding(SourceADesc.ID, TEXT("B"), TargetDesc.ID, TEXT("ArrayA[1]")));
+			PropertyBindings.Add(MakeBinding(SourceADesc.ID, TEXT("B"), TargetDesc.ID, TEXT("ArrayB[1]")));
+			PropertyBindings.Add(MakeBinding(SourceADesc.ID, TEXT("B"), TargetDesc.ID, TEXT("ArrayC[1]")));
+			PropertyBindings.Add(MakeBinding(SourceADesc.ID, TEXT("B"), TargetDesc.ID, TEXT("ArrayD[1]")));
+			if (!CompileBatch(PropertyBindings, CopyBatchIndex3))
+			{
+				return false;
+			}
+		}
+		// ArrayB[] to ...
+		int32 CopyBatchIndex4 = INDEX_NONE;
+		{
+			TArray<FStateTreePropertyPathBinding> PropertyBindings;
+			PropertyBindings.Add(MakeBinding(SourceADesc.ID, TEXT("ArrayB[2]"), TargetDesc.ID, TEXT("A")));
+			PropertyBindings.Add(MakeBinding(SourceADesc.ID, TEXT("ArrayB[2]"), TargetDesc.ID, TEXT("B")));
+			PropertyBindings.Add(MakeBinding(SourceADesc.ID, TEXT("ArrayB[2]"), TargetDesc.ID, TEXT("C")));
+			PropertyBindings.Add(MakeBinding(SourceADesc.ID, TEXT("ArrayB[2]"), TargetDesc.ID, TEXT("D")));
+			PropertyBindings.Add(MakeBinding(SourceADesc.ID, TEXT("ArrayB[2]"), TargetDesc.ID, TEXT("ArrayA[1]")));
+			PropertyBindings.Add(MakeBinding(SourceADesc.ID, TEXT("ArrayB[2]"), TargetDesc.ID, TEXT("ArrayB[1]")));
+			PropertyBindings.Add(MakeBinding(SourceADesc.ID, TEXT("ArrayB[2]"), TargetDesc.ID, TEXT("ArrayC[1]")));
+			PropertyBindings.Add(MakeBinding(SourceADesc.ID, TEXT("ArrayB[2]"), TargetDesc.ID, TEXT("ArrayD[1]")));
+			if (!CompileBatch(PropertyBindings, CopyBatchIndex4))
+			{
+				return false;
+			}
+		}
+		// C to ...
+		int32 CopyBatchIndex5 = INDEX_NONE;
+		{
+			TArray<FStateTreePropertyPathBinding> PropertyBindings;
+			PropertyBindings.Add(MakeBinding(SourceADesc.ID, TEXT("C"), TargetDesc.ID, TEXT("A")));
+			PropertyBindings.Add(MakeBinding(SourceADesc.ID, TEXT("C"), TargetDesc.ID, TEXT("B")));
+			PropertyBindings.Add(MakeBinding(SourceADesc.ID, TEXT("C"), TargetDesc.ID, TEXT("C")));
+			PropertyBindings.Add(MakeBinding(SourceADesc.ID, TEXT("C"), TargetDesc.ID, TEXT("D")));
+			PropertyBindings.Add(MakeBinding(SourceADesc.ID, TEXT("C"), TargetDesc.ID, TEXT("ArrayA[1]")));
+			PropertyBindings.Add(MakeBinding(SourceADesc.ID, TEXT("C"), TargetDesc.ID, TEXT("ArrayB[1]")));
+			PropertyBindings.Add(MakeBinding(SourceADesc.ID, TEXT("C"), TargetDesc.ID, TEXT("ArrayC[1]")));
+			PropertyBindings.Add(MakeBinding(SourceADesc.ID, TEXT("C"), TargetDesc.ID, TEXT("ArrayD[1]")));
+			if (!CompileBatch(PropertyBindings, CopyBatchIndex5))
+			{
+				return false;
+			}
+		}
+		// ArrayC[] to ...
+		int32 CopyBatchIndex6 = INDEX_NONE;
+		{
+			TArray<FStateTreePropertyPathBinding> PropertyBindings;
+			PropertyBindings.Add(MakeBinding(SourceADesc.ID, TEXT("ArrayC[2]"), TargetDesc.ID, TEXT("A")));
+			PropertyBindings.Add(MakeBinding(SourceADesc.ID, TEXT("ArrayC[2]"), TargetDesc.ID, TEXT("B")));
+			PropertyBindings.Add(MakeBinding(SourceADesc.ID, TEXT("ArrayC[2]"), TargetDesc.ID, TEXT("C")));
+			PropertyBindings.Add(MakeBinding(SourceADesc.ID, TEXT("ArrayC[2]"), TargetDesc.ID, TEXT("D")));
+			PropertyBindings.Add(MakeBinding(SourceADesc.ID, TEXT("ArrayC[2]"), TargetDesc.ID, TEXT("ArrayA[1]")));
+			PropertyBindings.Add(MakeBinding(SourceADesc.ID, TEXT("ArrayC[2]"), TargetDesc.ID, TEXT("ArrayB[1]")));
+			PropertyBindings.Add(MakeBinding(SourceADesc.ID, TEXT("ArrayC[2]"), TargetDesc.ID, TEXT("ArrayC[1]")));
+			PropertyBindings.Add(MakeBinding(SourceADesc.ID, TEXT("ArrayC[2]"), TargetDesc.ID, TEXT("ArrayD[1]")));
+			if (!CompileBatch(PropertyBindings, CopyBatchIndex6))
+			{
+				return false;
+			}
+		}
+		// D to ...
+		int32 CopyBatchIndex7 = INDEX_NONE;
+		{
+			TArray<FStateTreePropertyPathBinding> PropertyBindings;
+			PropertyBindings.Add(MakeBinding(SourceADesc.ID, TEXT("D"), TargetDesc.ID, TEXT("A")));
+			PropertyBindings.Add(MakeBinding(SourceADesc.ID, TEXT("D"), TargetDesc.ID, TEXT("B")));
+			PropertyBindings.Add(MakeBinding(SourceADesc.ID, TEXT("D"), TargetDesc.ID, TEXT("C")));
+			PropertyBindings.Add(MakeBinding(SourceADesc.ID, TEXT("D"), TargetDesc.ID, TEXT("D")));
+			PropertyBindings.Add(MakeBinding(SourceADesc.ID, TEXT("D"), TargetDesc.ID, TEXT("ArrayA[1]")));
+			PropertyBindings.Add(MakeBinding(SourceADesc.ID, TEXT("D"), TargetDesc.ID, TEXT("ArrayB[1]")));
+			PropertyBindings.Add(MakeBinding(SourceADesc.ID, TEXT("D"), TargetDesc.ID, TEXT("ArrayC[1]")));
+			PropertyBindings.Add(MakeBinding(SourceADesc.ID, TEXT("D"), TargetDesc.ID, TEXT("ArrayD[1]")));
+			if (!CompileBatch(PropertyBindings, CopyBatchIndex7))
+			{
+				return false;
+			}
+		}
+		// ArrayD[] to ...
+		int32 CopyBatchIndex8 = INDEX_NONE;
+		{
+			TArray<FStateTreePropertyPathBinding> PropertyBindings;
+			PropertyBindings.Add(MakeBinding(SourceADesc.ID, TEXT("ArrayD[2]"), TargetDesc.ID, TEXT("A")));
+			PropertyBindings.Add(MakeBinding(SourceADesc.ID, TEXT("ArrayD[2]"), TargetDesc.ID, TEXT("B")));
+			PropertyBindings.Add(MakeBinding(SourceADesc.ID, TEXT("ArrayD[2]"), TargetDesc.ID, TEXT("C")));
+			PropertyBindings.Add(MakeBinding(SourceADesc.ID, TEXT("ArrayD[2]"), TargetDesc.ID, TEXT("D")));
+			PropertyBindings.Add(MakeBinding(SourceADesc.ID, TEXT("ArrayD[2]"), TargetDesc.ID, TEXT("ArrayA[1]")));
+			PropertyBindings.Add(MakeBinding(SourceADesc.ID, TEXT("ArrayD[2]"), TargetDesc.ID, TEXT("ArrayB[1]")));
+			PropertyBindings.Add(MakeBinding(SourceADesc.ID, TEXT("ArrayD[2]"), TargetDesc.ID, TEXT("ArrayC[1]")));
+			PropertyBindings.Add(MakeBinding(SourceADesc.ID, TEXT("ArrayD[2]"), TargetDesc.ID, TEXT("ArrayD[1]")));
+			if (!CompileBatch(PropertyBindings, CopyBatchIndex8))
+			{
+				return false;
+			}
+		}
+		// ArrayA to Array...
+		int32 CopyBatchIndex9 = INDEX_NONE;
+		{
+			TArray<FStateTreePropertyPathBinding> PropertyBindings;
+			PropertyBindings.Add(MakeBinding(SourceADesc.ID, TEXT("ArrayA"), TargetDesc.ID, TEXT("ArrayA")));
+			PropertyBindings.Add(MakeBinding(SourceADesc.ID, TEXT("ArrayA"), TargetDesc.ID, TEXT("ArrayB")));
+			PropertyBindings.Add(MakeBinding(SourceADesc.ID, TEXT("ArrayA"), TargetDesc.ID, TEXT("ArrayC")));
+			PropertyBindings.Add(MakeBinding(SourceADesc.ID, TEXT("ArrayA"), TargetDesc.ID, TEXT("ArrayD")));
+			if (!CompileBatch(PropertyBindings, CopyBatchIndex9))
+			{
+				return false;
+			}
+		}
+		// ArrayB to Array...
+		int32 CopyBatchIndex10 = INDEX_NONE;
+		{
+			TArray<FStateTreePropertyPathBinding> PropertyBindings;
+			PropertyBindings.Add(MakeBinding(SourceADesc.ID, TEXT("ArrayB"), TargetDesc.ID, TEXT("ArrayA")));
+			PropertyBindings.Add(MakeBinding(SourceADesc.ID, TEXT("ArrayB"), TargetDesc.ID, TEXT("ArrayB")));
+			PropertyBindings.Add(MakeBinding(SourceADesc.ID, TEXT("ArrayB"), TargetDesc.ID, TEXT("ArrayC")));
+			PropertyBindings.Add(MakeBinding(SourceADesc.ID, TEXT("ArrayB"), TargetDesc.ID, TEXT("ArrayD")));
+			if (!CompileBatch(PropertyBindings, CopyBatchIndex10))
+			{
+				return false;
+			}
+		}
+		// ArrayC to Array...
+		int32 CopyBatchIndex11 = INDEX_NONE;
+		{
+			TArray<FStateTreePropertyPathBinding> PropertyBindings;
+			PropertyBindings.Add(MakeBinding(SourceADesc.ID, TEXT("ArrayC"), TargetDesc.ID, TEXT("ArrayA")));
+			PropertyBindings.Add(MakeBinding(SourceADesc.ID, TEXT("ArrayC"), TargetDesc.ID, TEXT("ArrayB")));
+			PropertyBindings.Add(MakeBinding(SourceADesc.ID, TEXT("ArrayC"), TargetDesc.ID, TEXT("ArrayC")));
+			PropertyBindings.Add(MakeBinding(SourceADesc.ID, TEXT("ArrayC"), TargetDesc.ID, TEXT("ArrayD")));
+			if (!CompileBatch(PropertyBindings, CopyBatchIndex11))
+			{
+				return false;
+			}
+		}
+		// ArrayD to Array...
+		int32 CopyBatchIndex12 = INDEX_NONE;
+		{
+			TArray<FStateTreePropertyPathBinding> PropertyBindings;
+			PropertyBindings.Add(MakeBinding(SourceADesc.ID, TEXT("ArrayD"), TargetDesc.ID, TEXT("ArrayA")));
+			PropertyBindings.Add(MakeBinding(SourceADesc.ID, TEXT("ArrayD"), TargetDesc.ID, TEXT("ArrayB")));
+			PropertyBindings.Add(MakeBinding(SourceADesc.ID, TEXT("ArrayD"), TargetDesc.ID, TEXT("ArrayC")));
+			PropertyBindings.Add(MakeBinding(SourceADesc.ID, TEXT("ArrayD"), TargetDesc.ID, TEXT("ArrayD")));
+			if (!CompileBatch(PropertyBindings, CopyBatchIndex12))
+			{
+				return false;
+			}
+		}
+		BindingCompiler.Finalize();
+
+		const bool bResolveResult = Bindings.ResolvePaths();
+		AITEST_TRUE(TEXT("ResolvePaths should succeed"), bResolveResult);
+
+		FStateTreeTest_PropertyStructObjectPtr SourceA;
+		SourceA.A = TBaseStructure<FVector>::Get();
+		SourceA.B = TBaseStructure<FRotator>::Get();
+		SourceA.C = TBaseStructure<FMatrix>::Get();
+		SourceA.D = TBaseStructure<FTransform>::Get();
+		constexpr int32 NumberOfElements = 4;
+		SourceA.ArrayA.SetNum(NumberOfElements, EAllowShrinking::No);
+		SourceA.ArrayB.SetNum(NumberOfElements, EAllowShrinking::No);
+		SourceA.ArrayC.SetNum(NumberOfElements, EAllowShrinking::No);
+		SourceA.ArrayD.SetNum(NumberOfElements, EAllowShrinking::No);
+		
+		SourceA.ArrayA[2] = TBaseStructure<FVector>::Get();
+		SourceA.ArrayB[2] = TBaseStructure<FRotator>::Get();
+		SourceA.ArrayC[2] = TBaseStructure<FMatrix>::Get();
+		SourceA.ArrayD[2] = TBaseStructure<FTransform>::Get();
+
+		FStateTreeTest_PropertyStruct SourceB;
+
+		FStateTreeTest_PropertyStructObjectPtr Target;
+		Target.ArrayA.SetNum(2, EAllowShrinking::No);
+		Target.ArrayB.SetNum(2, EAllowShrinking::No);
+		Target.ArrayC.SetNum(2, EAllowShrinking::No);
+		Target.ArrayD.SetNum(2, EAllowShrinking::No);
+
+		auto ResetTargetData = [&Target]()
+			{
+				Target.A = nullptr;
+				Target.B = nullptr;
+				Target.C = nullptr;
+				Target.D = nullptr;
+				for (int32 Index = 0; Index < Target.ArrayA.Num(); ++Index)
+				{
+					Target.ArrayA[Index] = nullptr;
+					Target.ArrayB[Index] = nullptr;
+					Target.ArrayC[Index] = nullptr;
+					Target.ArrayD[Index] = nullptr;
+				}
+				return true;
+			};
+		ResetTargetData();
+
+		auto TestCopy = [this, &Target](UObject* Object, const FStringView SourceStr)
+			{
+				AITEST_EQUAL(FString::Printf(TEXT("Expect A to be copied from %s"), SourceStr.GetData()), Target.A.Get(), Object);
+				AITEST_EQUAL(FString::Printf(TEXT("Expect B to be copied from %s"), SourceStr.GetData()), Target.B.Get(), Object);
+				AITEST_EQUAL(FString::Printf(TEXT("Expect C to be copied from %s"), SourceStr.GetData()), Target.C.Get(), Object);
+				AITEST_EQUAL(FString::Printf(TEXT("Expect D to be copied from %s"), SourceStr.GetData()), Target.D.Get(), Object);
+				AITEST_EQUAL(FString::Printf(TEXT("Expect A[] to be copied from %s"), SourceStr.GetData()), Target.ArrayA[1].Get(), Object);
+				AITEST_EQUAL(FString::Printf(TEXT("Expect B[] to be copied from %s"), SourceStr.GetData()), Target.ArrayB[1].Get(), Object);
+				AITEST_EQUAL(FString::Printf(TEXT("Expect C[] to be copied from %s"), SourceStr.GetData()), Target.ArrayC[1].Get(), Object);
+				AITEST_EQUAL(FString::Printf(TEXT("Expect D[] to be copied from %s"), SourceStr.GetData()), Target.ArrayD[1].Get(), Object);
+				return true;
+			};
+
+		AITEST_TRUE(TEXT("SourceAIndex should be less than max number of source structs."), SourceAIndex < Bindings.GetNumBindableStructDescriptors());
+
+		TArray<FStateTreeDataView> SourceViews;
+		SourceViews.SetNum(Bindings.GetNumBindableStructDescriptors());
+		SourceViews[SourceAIndex] = FStateTreeDataView(FStructView::Make(SourceA));
+		FPropertyBindingDataView TargetView(FStructView::Make(Target));
+
+		auto CopyAndTestBatch = [this, &ResetTargetData, &TestCopy, &Bindings, &SourceViews, &TargetView](int32 CopyBatchIndex, UObject* Object, FStringView SourceStr)
+			{
+				ResetTargetData();
+				bool bCopyResult = true;
+				for (const FPropertyBindingCopyInfo& Copy : Bindings.Super::GetBatchCopies(FPropertyBindingIndex16(CopyBatchIndex)))
+				{
+					bCopyResult &= Bindings.Super::CopyProperty(Copy, SourceViews[Copy.SourceDataHandle.Get<FStateTreeDataHandle>().GetIndex()], TargetView);
+				}
+				AITEST_TRUE(TEXT("CopyTo should succeed"), bCopyResult);
+				return TestCopy(Object, SourceStr);
+			};
+
+		if (!CopyAndTestBatch(CopyBatchIndex1, SourceA.A.Get(), TEXT("SourceA.A")))
+		{
+			return false;
+		}
+		if (!CopyAndTestBatch(CopyBatchIndex2, SourceA.A.Get(), TEXT("SourceA.ArrayA")))
+		{
+			return false;
+		}
+		if (!CopyAndTestBatch(CopyBatchIndex3, SourceA.B.Get(), TEXT("SourceA.B")))
+		{
+			return false;
+		}
+		if (!CopyAndTestBatch(CopyBatchIndex4, SourceA.B.Get(), TEXT("SourceA.ArrayB")))
+		{
+			return false;
+		}
+		if (!CopyAndTestBatch(CopyBatchIndex5, SourceA.C.Get(), TEXT("SourceA.C")))
+		{
+			return false;
+		}
+		if (!CopyAndTestBatch(CopyBatchIndex6, SourceA.C.Get(), TEXT("SourceA.ArrayC")))
+		{
+			return false;
+		}
+		if (!CopyAndTestBatch(CopyBatchIndex7, SourceA.D.Get(), TEXT("SourceA.D")))
+		{
+			return false;
+		}
+		if (!CopyAndTestBatch(CopyBatchIndex8, SourceA.D.Get(), TEXT("SourceA.ArrayD")))
+		{
+			return false;
+		}
+
+		// Test promote array
+		for (int32 Index = 0; Index < NumberOfElements; ++Index)
+		{
+			SourceA.ArrayA[Index] = TBaseStructure<FVector>::Get();
+			SourceA.ArrayB[Index] = TBaseStructure<FRotator>::Get();
+			SourceA.ArrayC[Index] = TBaseStructure<FMatrix>::Get();
+			SourceA.ArrayD[Index] = TBaseStructure<FTransform>::Get();
+		}
+
+		auto TestCopyArray = [this, &Target, NumberOfElements](UObject* Object, const FStringView SourceStr)
+			{
+				for (int32 Index = 0; Index < NumberOfElements; ++Index)
+				{
+					AITEST_EQUAL(FString::Printf(TEXT("Expect A[] to be copied from %s"), SourceStr.GetData()), Target.ArrayA[Index].Get(), Object);
+					AITEST_EQUAL(FString::Printf(TEXT("Expect B[] to be copied from %s"), SourceStr.GetData()), Target.ArrayB[Index].Get(), Object);
+					AITEST_EQUAL(FString::Printf(TEXT("Expect C[] to be copied from %s"), SourceStr.GetData()), Target.ArrayC[Index].Get(), Object);
+					AITEST_EQUAL(FString::Printf(TEXT("Expect D[] to be copied from %s"), SourceStr.GetData()), Target.ArrayD[Index].Get(), Object);
+				}
+				return true;
+			};
+
+		auto CopyAndTestBatchArray = [this, &ResetTargetData, &TestCopyArray, &Bindings, &SourceViews, &TargetView](int32 CopyBatchIndex, UObject* Object, FStringView SourceStr)
+			{
+				ResetTargetData();
+				bool bCopyResult = true;
+				for (const FPropertyBindingCopyInfo& Copy : Bindings.Super::GetBatchCopies(FPropertyBindingIndex16(CopyBatchIndex)))
+				{
+					bCopyResult &= Bindings.Super::CopyProperty(Copy, SourceViews[Copy.SourceDataHandle.Get<FStateTreeDataHandle>().GetIndex()], TargetView);
+				}
+				AITEST_TRUE(TEXT("CopyTo should succeed"), bCopyResult);
+				return TestCopyArray(Object, SourceStr);
+			};
+		if (!CopyAndTestBatchArray(CopyBatchIndex9, SourceA.A.Get(), TEXT("SourceA.ArrayA")))
+		{
+			return false;
+		}
+		if (!CopyAndTestBatchArray(CopyBatchIndex10, SourceA.B.Get(), TEXT("SourceA.ArrayB")))
+		{
+			return false;
+		}
+		if (!CopyAndTestBatchArray(CopyBatchIndex11, SourceA.C.Get(), TEXT("SourceA.ArrayC")))
+		{
+			return false;
+		}
+		if (!CopyAndTestBatchArray(CopyBatchIndex12, SourceA.D.Get(), TEXT("SourceA.ArrayD")))
+		{
+			return false;
+		}
+
+		return true;
+	}
+};
+IMPLEMENT_STATE_TREE_INSTANT_TEST(FStateTreeTest_BindingsCompiler_WithObjectPtr, "System.StateTree.Binding.BindingsCompilerWithObjectPtr");
+
 struct FStateTreeTest_PropertyFunctions : FStateTreeTestBase
 {
 	virtual bool InstantTest() override
diff --git a/Engine/Plugins/Runtime/StateTree/Source/StateTreeTestSuite/Private/StateTreeTestTypes.h b/Engine/Plugins/Runtime/StateTree/Source/StateTreeTestSuite/Private/StateTreeTestTypes.h
--- a/Engine/Plugins/Runtime/StateTree/Source/StateTreeTestSuite/Private/StateTreeTestTypes.h
+++ b/Engine/Plugins/Runtime/StateTree/Source/StateTreeTestSuite/Private/StateTreeTestTypes.h
@@ -827,6 +827,36 @@ struct FStateTreeTest_PropertyStructB
 	static inline int32 NumConstructed = 0;
 };
 
+USTRUCT()
+struct FStateTreeTest_PropertyStructObjectPtr
+{
+	GENERATED_BODY()
+
+	UPROPERTY(EditAnywhere, Category = "")
+	TObjectPtr<UObject> A;
+
+	UPROPERTY(EditAnywhere, Category = "")
+	TWeakObjectPtr<UObject> B;
+
+	UPROPERTY(EditAnywhere, Category = "")
+	TSoftObjectPtr<UObject> C;
+
+	UPROPERTY(EditAnywhere, Category = "")
+	TLazyObjectPtr<UObject> D;
+
+	UPROPERTY(EditAnywhere, Category = "")
+	TArray<TObjectPtr<UObject>> ArrayA;
+
+	UPROPERTY(EditAnywhere, Category = "")
+	TArray<TWeakObjectPtr<UObject>> ArrayB;
+
+	UPROPERTY(EditAnywhere, Category = "")
+	TArray<TSoftObjectPtr<UObject>> ArrayC;
+
+	UPROPERTY(EditAnywhere, Category = "")
+	TArray<TLazyObjectPtr<UObject>> ArrayD;
+};
+
 USTRUCT()
 struct FStateTreeTest_PropertyStruct
 {
@@ -919,6 +949,13 @@ struct FStateTreeTest_PropertyCopy
 
 	UPROPERTY(EditAnywhere, Category = "")
 	FStateTreeTest_PropertyStruct CArray[3];
+
+	UPROPERTY(EditAnywhere, Category = "")
+	TArray<float> ArrayOfFloat;
+
+	UPROPERTY(EditAnywhere, Category = "")
+	TArray<double> ArrayOfDouble;
+
 };
 
 USTRUCT()
