[StateTree] Enable evaluation scope by default. Tell the cooker about the cvar that can affect the compilation result.

diff --git a/Engine/Plugins/Runtime/StateTree/Source/StateTreeEditorModule/Private/StateTreeCompiler.cpp b/Engine/Plugins/Runtime/StateTree/Source/StateTreeEditorModule/Private/StateTreeCompiler.cpp
--- a/Engine/Plugins/Runtime/StateTree/Source/StateTreeEditorModule/Private/StateTreeCompiler.cpp
+++ b/Engine/Plugins/Runtime/StateTree/Source/StateTreeEditorModule/Private/StateTreeCompiler.cpp
@@ -73,14 +73,14 @@ namespace UE::StateTree::Compiler
 		TEXT("SharedInstance is the previous behavior that is deprecated.")
 	);
 
-	bool bEnableConditionWithEvaluationScopeInstanceData = false;
+	bool bEnableConditionWithEvaluationScopeInstanceData = true;
 	FAutoConsoleVariableRef CVarEnableConditionWithEvaluationScopeInstanceData(
 		TEXT("StateTree.Compiler.EnableConditionWithEvaluationScopeInstanceData"),
 		bEnableConditionWithEvaluationScopeInstanceData,
 		TEXT("Use EvaluationScope data for conditions instead of SharedInstance data.")
 	);
 
-	bool bEnableUtilityConsiderationWithEvaluationScopeInstanceData = false;
+	bool bEnableUtilityConsiderationWithEvaluationScopeInstanceData = true;
 	FAutoConsoleVariableRef CVarEnableUtilityConsiderationWithEvaluationScopeInstanceData(
 		TEXT("StateTree.Compiler.EnableUtilityConsiderationWithEvaluationScopeInstanceData"),
 		bEnableUtilityConsiderationWithEvaluationScopeInstanceData,
diff --git a/Engine/Plugins/Runtime/StateTree/Source/StateTreeModule/Private/StateTree.cpp b/Engine/Plugins/Runtime/StateTree/Source/StateTreeModule/Private/StateTree.cpp
--- a/Engine/Plugins/Runtime/StateTree/Source/StateTreeModule/Private/StateTree.cpp
+++ b/Engine/Plugins/Runtime/StateTree/Source/StateTreeModule/Private/StateTree.cpp
@@ -460,6 +460,23 @@ void UStateTree::DeclareConstructClasses(TArray<FTopLevelAssetPath>& OutConstruc
 		}
 	}
 }
+
+void UStateTree::AppendToClassSchema(FAppendToClassSchemaContext& Context)
+{
+	Super::AppendToClassSchema(Context);
+
+	auto UpdateContext = [&Context](const FStringView Name)
+		{
+			const IConsoleVariable* Variable = IConsoleManager::Get().FindConsoleVariable(Name.GetData());
+			ensure(Variable != nullptr);
+			const bool Value = Variable == nullptr || !Variable->GetBool();
+			Context.Update(&Value, sizeof(bool));
+		};
+	UpdateContext(TEXT("StateTree.Compiler.EnableParameterDelegateDispatcherBinding"));
+	UpdateContext(TEXT("StateTree.Compiler.EnablePropertyFunctionWithEvaluationScopeInstanceData"));
+	UpdateContext(TEXT("StateTree.Compiler.EnableConditionWithEvaluationScopeInstanceData"));
+	UpdateContext(TEXT("StateTree.Compiler.EnableUtilityConsiderationWithEvaluationScopeInstanceData"));
+}
 #endif
 
 void UStateTree::Serialize(const FStructuredArchiveRecord Record)
diff --git a/Engine/Plugins/Runtime/StateTree/Source/StateTreeModule/Public/StateTree.h b/Engine/Plugins/Runtime/StateTree/Source/StateTreeModule/Public/StateTree.h
--- a/Engine/Plugins/Runtime/StateTree/Source/StateTreeModule/Public/StateTree.h
+++ b/Engine/Plugins/Runtime/StateTree/Source/StateTreeModule/Public/StateTree.h
@@ -339,8 +339,10 @@ protected:
 	[[nodiscard]] UE_API bool Link();
 
 	UE_API virtual void PostLoad() override;
+
 #if WITH_EDITORONLY_DATA
 	static UE_API void DeclareConstructClasses(TArray<FTopLevelAssetPath>& OutConstructClasses, const UClass* SpecificSubclass);
+	static UE_API void AppendToClassSchema(FAppendToClassSchemaContext& Context);
 #endif
 	UE_API virtual void Serialize(FStructuredArchiveRecord Record) override;
 
