Fix other instance of state tree renaming issues.

#tests editor

diff --git a/Engine/Plugins/Runtime/StateTree/Source/StateTreeEditorModule/Private/StateTreeEditingSubsystem.cpp b/Engine/Plugins/Runtime/StateTree/Source/StateTreeEditorModule/Private/StateTreeEditingSubsystem.cpp
--- a/Engine/Plugins/Runtime/StateTree/Source/StateTreeEditorModule/Private/StateTreeEditingSubsystem.cpp
+++ b/Engine/Plugins/Runtime/StateTree/Source/StateTreeEditorModule/Private/StateTreeEditingSubsystem.cpp
@@ -3,6 +3,7 @@
 #include "StateTreeEditingSubsystem.h"
 
 #include "SStateTreeView.h"
+#include "StateTree.h"
 #include "StateTreeCompiler.h"
 #include "StateTreeCompilerLog.h"
 #include "StateTreeCompilerManager.h"
@@ -13,7 +14,6 @@
 #include "StateTreeEditorDataExtension.h"
 #include "StateTreeObjectHash.h"
 #include "StateTreeTaskBase.h"
-#include "UObject/LinkerLoad.h"
 #include "UObject/UObjectGlobals.h"
 
 #include UE_INLINE_GENERATED_CPP_BY_NAME(StateTreeEditingSubsystem)
@@ -192,17 +192,8 @@ void UStateTreeEditingSubsystem::ValidateStateTree(TNotNull<UStateTree*> InState
 			StateTree->EditorData = EditorData;
 
 			// Trash the previous EditorData
-			const FName TrashStateTreeName = MakeUniqueObjectName(GetTransientPackage(), UStateTree::StaticClass(), *FString::Printf(TEXT("TRASH_%s"), *UStateTree::StaticClass()->GetName()));
-			const FName TrashSchemaName = *FString::Printf(TEXT("TRASH_%s"), *PreviousEditorData->GetName());
-			constexpr ERenameFlags RenameFlags = REN_DoNotDirty | REN_DontCreateRedirectors;
-
-			// @TODO: UE-348529. This rename below is potentially unsafe. While we `ConditionallyPostload` the editordata in `UStateTree::PostLoad`.
-			// Since this object is intended to be discarded the correct method is to call `FLinkerLoad::InvalidateExport`.
-			// Consider moving `UE::StateTree::Compiler::RenameObjectToTransientPackage` to a public static method in `StateTreeCompiler.h` to 
-			// do this since it calls the aforementioned method. Vefify that this temp transient outer object below isn't needed first.
-			UStateTree* TransientOuter = NewObject<UStateTree>(GetTransientPackage(), TrashStateTreeName, RF_Transient);
-			PreviousEditorData->Rename(*TrashSchemaName.ToString(), TransientOuter, RenameFlags);
-			PreviousEditorData->SetFlags(RF_Transient);
+			const FString TrashEditorDataName = FString::Printf(TEXT("TRASH_%s"), *PreviousEditorData->GetName());
+			UE::StateTree::Compiler::RenameObjectToTransientPackage(PreviousEditorData, TrashEditorDataName);
 		}
 	};
 
@@ -226,15 +217,8 @@ void UStateTreeEditingSubsystem::ValidateStateTree(TNotNull<UStateTree*> InState
 				EditorData->EditorSchema = CastChecked<UStateTreeEditorSchema>(StaticDuplicateObject(PreviousEditorSchema, EditorData, FName(), RF_Transactional, EditorSchemaClass.Get()));
 				
 				// Trash the previous EditorDataSchema
-				const FName TrashSchemaName = *FString::Printf(TEXT("TRASH_%s"), *PreviousEditorSchema->GetName());
-				constexpr ERenameFlags RenameFlags = REN_DoNotDirty | REN_DontCreateRedirectors;
-
-				PreviousEditorSchema->SetFlags(RF_Transient);
-				PreviousEditorSchema->ClearFlags(RF_Public | RF_Standalone | RF_ArchetypeObject);
-
-				// We are no longer interested in loading the old schema, so invalidate export.
-				FLinkerLoad::InvalidateExport(PreviousEditorSchema);
-				PreviousEditorSchema->Rename(*TrashSchemaName.ToString(), GetTransientPackage(), RenameFlags);
+				const FString TrashEditorSchemaName = FString::Printf(TEXT("TRASH_%s"), *PreviousEditorSchema->GetName());
+				UE::StateTree::Compiler::RenameObjectToTransientPackage(PreviousEditorSchema, TrashEditorSchemaName );
 			}
 		};
 
diff --git a/Engine/Plugins/Runtime/StateTree/Source/StateTreeModule/Private/StateTree.cpp b/Engine/Plugins/Runtime/StateTree/Source/StateTreeModule/Private/StateTree.cpp
--- a/Engine/Plugins/Runtime/StateTree/Source/StateTreeModule/Private/StateTree.cpp
+++ b/Engine/Plugins/Runtime/StateTree/Source/StateTreeModule/Private/StateTree.cpp
@@ -225,7 +225,7 @@ void UStateTree::BeginDestroy()
 #if WITH_EDITOR
 namespace UE::StateTree::Compiler
 {
-	void RenameObjectToTransientPackage(UObject* ObjectToRename)
+	void RenameObjectToTransientPackage(TNotNull<UObject*> ObjectToRename, FStringView Name)
 	{
 		const ERenameFlags RenFlags = REN_DoNotDirty | REN_DontCreateRedirectors;
 
@@ -234,7 +234,7 @@ namespace UE::StateTree::Compiler
 
 		// Rename will remove the renamed object's linker when moving to a new package so invalidate the export beforehand
 		FLinkerLoad::InvalidateExport(ObjectToRename);
-		ObjectToRename->Rename(nullptr, GetTransientPackage(), RenFlags);
+		ObjectToRename->Rename(Name.GetData(), GetTransientPackage(), RenFlags);
 	}
 }
 
diff --git a/Engine/Plugins/Runtime/StateTree/Source/StateTreeModule/Public/StateTree.h b/Engine/Plugins/Runtime/StateTree/Source/StateTreeModule/Public/StateTree.h
--- a/Engine/Plugins/Runtime/StateTree/Source/StateTreeModule/Public/StateTree.h
+++ b/Engine/Plugins/Runtime/StateTree/Source/StateTreeModule/Public/StateTree.h
@@ -80,6 +80,18 @@ private:
 
 
 #if WITH_EDITOR
+
+namespace UE::StateTree::Compiler
+{
+	/** 
+	 * Helper method to correctly rename a no longer needed object. 
+	 * 
+	 * @param ObjectToRename - Object to be renamed, will have flags updated & export invalidated
+	 * @param Name - New name of the object. Null by default which means it will only get moved to transient outer.
+	 */
+	extern UE_API void RenameObjectToTransientPackage(TNotNull<UObject*> ObjectToRename, FStringView Name = nullptr);
+}
+
 namespace UE::StateTree::Compiler::Private
 {
 	class FCompilerManagerImpl;
