[Backout] - CL46999606
#rnx
Original CL Desc
-----------------------------------------------------------------
[StateTree] Enable evaluation scope by default. Tell the cooker about the cvar that can affect the compilation result.

diff --git a/Engine/Plugins/Runtime/StateTree/Source/StateTreeEditorModule/Private/StateTreeCompiler.cpp b/Engine/Plugins/Runtime/StateTree/Source/StateTreeEditorModule/Private/StateTreeCompiler.cpp
--- a/Engine/Plugins/Runtime/StateTree/Source/StateTreeEditorModule/Private/StateTreeCompiler.cpp
+++ b/Engine/Plugins/Runtime/StateTree/Source/StateTreeEditorModule/Private/StateTreeCompiler.cpp
@@ -73,14 +73,14 @@ namespace UE::StateTree::Compiler
 		TEXT("SharedInstance is the previous behavior that is deprecated.")
 	);
 
-	bool bEnableConditionWithEvaluationScopeInstanceData = true;
+	bool bEnableConditionWithEvaluationScopeInstanceData = false;
 	FAutoConsoleVariableRef CVarEnableConditionWithEvaluationScopeInstanceData(
 		TEXT("StateTree.Compiler.EnableConditionWithEvaluationScopeInstanceData"),
 		bEnableConditionWithEvaluationScopeInstanceData,
 		TEXT("Use EvaluationScope data for conditions instead of SharedInstance data.")
 	);
 
-	bool bEnableUtilityConsiderationWithEvaluationScopeInstanceData = true;
+	bool bEnableUtilityConsiderationWithEvaluationScopeInstanceData = false;
 	FAutoConsoleVariableRef CVarEnableUtilityConsiderationWithEvaluationScopeInstanceData(
 		TEXT("StateTree.Compiler.EnableUtilityConsiderationWithEvaluationScopeInstanceData"),
 		bEnableUtilityConsiderationWithEvaluationScopeInstanceData,
