[StateTree] Move the temporaries before it is used in EnterState. Previously, they were all moved in UpdateInstanceData. It caused a problem if an async task requested the instance data before the state was added to the Frame. Usually, that could not happen, but global nodes (task and evaluator) can be started temporarily. Their EnterState is called, and they will only be added to the frame when the linked state is added.

diff --git a/Engine/Plugins/Runtime/StateTree/Source/StateTreeModule/Private/StateTreeExecutionContext.cpp b/Engine/Plugins/Runtime/StateTree/Source/StateTreeModule/Private/StateTreeExecutionContext.cpp
--- a/Engine/Plugins/Runtime/StateTree/Source/StateTreeModule/Private/StateTreeExecutionContext.cpp
+++ b/Engine/Plugins/Runtime/StateTree/Source/StateTreeModule/Private/StateTreeExecutionContext.cpp
@@ -2256,6 +2256,9 @@ void FStateTreeExecutionContext::UpdateInstanceData(const TSharedPtr<FSelectStat
 	TArray<FConstStructView, FNonconcurrentLinearArrayAllocator> InstanceStructs;
 	InstanceStructs.Reserve(EstimatedNumStructs);
 
+	// Temporaries that are moved to the new instance data buffer.
+	// Node instance data are moved in EnterState when the ActiveNodeIndex is incremented.
+	// The array is clear at the end of EnterState.
 	TArray<FInstancedStruct*, FNonconcurrentLinearArrayAllocator> TempInstanceStructs;
 	TempInstanceStructs.Reserve(EstimatedNumStructs);
 
@@ -2330,10 +2333,6 @@ void FStateTreeExecutionContext::UpdateInstanceData(const TSharedPtr<FSelectStat
 						const FStateTreeEvaluatorBase& Eval = CurrentStateTree->Nodes[EvalIndex].Get<const FStateTreeEvaluatorBase>();
 						const FConstStructView EvalInstanceData = CurrentStateTree->DefaultInstanceData.GetStruct(Eval.InstanceTemplateIndex.Get());
 						InstanceStructs[BaseIndex + Eval.InstanceDataHandle.GetIndex()] = EvalInstanceData;
-						if (!bIsFrameCommon)
-						{
-							TempInstanceStructs[BaseIndex + Eval.InstanceDataHandle.GetIndex()] = FindInstanceTempData(CurrentFrameID, Eval.InstanceDataHandle);
-						}
 					}
 
 					// Global tasks
@@ -2342,10 +2341,6 @@ void FStateTreeExecutionContext::UpdateInstanceData(const TSharedPtr<FSelectStat
 						const FStateTreeTaskBase& Task = CurrentStateTree->Nodes[TaskIndex].Get<const FStateTreeTaskBase>();
 						const FConstStructView TaskInstanceData = CurrentStateTree->DefaultInstanceData.GetStruct(Task.InstanceTemplateIndex.Get());
 						InstanceStructs[BaseIndex + Task.InstanceDataHandle.GetIndex()] = TaskInstanceData;
-						if (!bIsFrameCommon)
-						{
-							TempInstanceStructs[BaseIndex + Task.InstanceDataHandle.GetIndex()] = FindInstanceTempData(CurrentFrameID, Task.InstanceDataHandle);
-						}
 					}
 
 					if (bIsFrameCommon)
@@ -2474,10 +2469,6 @@ void FStateTreeExecutionContext::UpdateInstanceData(const TSharedPtr<FSelectStat
 						const FStateTreeTaskBase& Task = CurrentStateTree->Nodes[TaskIndex].Get<const FStateTreeTaskBase>();
 						const FConstStructView TaskInstanceData = CurrentStateTree->DefaultInstanceData.GetStruct(Task.InstanceTemplateIndex.Get());
 						InstanceStructs[CurrentFrameBaseIndex + Task.InstanceDataHandle.GetIndex()] = TaskInstanceData;
-						if (!bIsStateCommon)
-						{
-							TempInstanceStructs[CurrentFrameBaseIndex + Task.InstanceDataHandle.GetIndex()] = FindInstanceTempData(CurrentFrameID, Task.InstanceDataHandle);
-						}
 					}
 
 					if (bIsStateCommon)
@@ -2517,8 +2508,57 @@ void FStateTreeExecutionContext::UpdateInstanceData(const TSharedPtr<FSelectStat
 	InstanceData.Append(Owner,
 		MakeArrayView(InstanceStructs.GetData() + NumCommonInstanceData, InstanceStructs.Num() - NumCommonInstanceData),
 		MakeArrayView(TempInstanceStructs.GetData() + NumCommonInstanceData, TempInstanceStructs.Num() - NumCommonInstanceData));
+}
 
-	InstanceData.ResetTemporaryInstances();
+FStateTreeDataView FStateTreeExecutionContext::MoveTemporaryToInstance(const FStateTreeExecutionFrame& CurrentFrame, const FStateTreeDataHandle Handle)
+{
+	check(Handle.GetSource() == EStateTreeDataSourceType::ActiveInstanceData
+		|| Handle.GetSource() == EStateTreeDataSourceType::ActiveInstanceDataObject
+		|| Handle.GetSource() == EStateTreeDataSourceType::GlobalInstanceData
+		|| Handle.GetSource() == EStateTreeDataSourceType::GlobalInstanceDataObject);
+
+	using namespace UE::StateTree;
+	using namespace UE::StateTree::ExecutionContext;
+	using namespace UE::StateTree::ExecutionContext::Private;
+
+	TArrayView<FStateTreeTemporaryInstanceData> TempInstances = Storage.GetMutableTemporaryInstances();
+	auto FindInstanceTempData = [&TempInstances](FActiveFrameID FrameID, FStateTreeDataHandle DataHandle)
+		{
+			FStateTreeTemporaryInstanceData* TempData = TempInstances.FindByPredicate([&FrameID, &DataHandle](const FStateTreeTemporaryInstanceData& Data)
+				{
+					return Data.FrameID == FrameID && Data.DataHandle == DataHandle;
+				});
+			return TempData ? &TempData->Instance : nullptr;
+		};
+
+	FStateTreeDataView Struct = UE::StateTree::InstanceData::GetDataView(Storage, nullptr, nullptr, CurrentFrame, Handle.ToStructSource());
+	FInstancedStruct* Source = FindInstanceTempData(CurrentFrame.FrameID, Handle);
+
+	if (Struct.IsValid())
+	{
+		// ToStructSource should have converted it to FStateTreeInstanceObjectWrapper
+		check(Struct.GetStruct() && Cast<const UClass>(Struct.GetStruct()) == nullptr);
+
+		if (Source && Source->IsValid())
+		{
+			check(Struct.GetStruct() == Source->GetScriptStruct());
+
+			FMemory::Memswap(Struct.GetMutableMemory(), Source->GetMutableMemory(), Struct.GetStruct()->GetStructureSize());
+			Source->Reset();
+		}
+		else if (Struct.GetStruct()->IsChildOf(FStateTreeInstanceObjectWrapper::StaticStruct()))
+		{
+			FStateTreeInstanceObjectWrapper* Wrapper = Struct.GetMutablePtr<FStateTreeInstanceObjectWrapper>();
+			if (Wrapper->InstanceObject)
+			{
+				const bool bDuplicate = true;
+				Wrapper->InstanceObject = UE::StateTree::InstanceData::Private::CopyNodeInstance(Wrapper->InstanceObject, &Owner, bDuplicate);
+			}
+			Struct = FStateTreeDataView(Wrapper->InstanceObject);
+		}
+	}
+
+	return Struct;
 }
 
 FStateTreeDataView FStateTreeExecutionContext::GetDataView(const FStateTreeExecutionFrame* ParentFrame, const FStateTreeExecutionFrame& CurrentFrame, const FStateTreeDataHandle Handle)
@@ -3375,6 +3415,8 @@ EStateTreeRunStatus FStateTreeExecutionContext::EnterState(const TSharedPtr<FSel
 
 	CSV_SCOPED_TIMING_STAT_EXCLUSIVE(StateTree_EnterState);
 
+	ON_SCOPE_EXIT{ InstanceData.ResetTemporaryInstances(); };
+
 	if (!ArgsPtr.IsValid())
 	{
 		TSharedPtr<FSelectStateResult> EmptySelectStateResult;
@@ -3461,17 +3503,27 @@ EStateTreeRunStatus FStateTreeExecutionContext::EnterState(const TSharedPtr<FSel
 			CurrentFrame.ActiveTasksStatus = FrameInfo ? FStateTreeTasksCompletionStatus(*FrameInfo) : FStateTreeTasksCompletionStatus();
 		}
 
-		if (bIsCurrentFrameNew && !bFrameAdded && !CurrentFrame.bHaveEntered)
+		if (bIsCurrentFrameNew)
 		{
-			// EnterState on Changed global was called during selection with StartTemporaryEvaluatorsAndGlobalTasks
-			CurrentTransition.CurrentState = FStateTreeStateHandle::Invalid;
-			CurrentTransition.ChangeType = EStateTreeStateChangeType::Sustained;
+			if (!bFrameAdded && !CurrentFrame.bHaveEntered)
+			{
+				// EnterState on Changed global was called during selection with StartTemporaryEvaluatorsAndGlobalTasks
+				CurrentTransition.CurrentState = FStateTreeStateHandle::Invalid;
+				CurrentTransition.ChangeType = EStateTreeStateChangeType::Sustained;
 
-			EStateTreeRunStatus GlobalTaskRunStatus = StartGlobalsForFrameOnActiveInstances(CurrentParentFrame, CurrentFrame, CurrentTransition);
-			Result = GetPriorityRunStatus(Result, GlobalTaskRunStatus);
-			if (Result == EStateTreeRunStatus::Failed)
+				constexpr bool bOnActiveInstances = true;
+				constexpr bool bMoveFromTemporary = true;
+				EStateTreeRunStatus GlobalTaskRunStatus = StartGlobalsForFrameInternal<bOnActiveInstances, bMoveFromTemporary>(CurrentParentFrame, CurrentFrame, CurrentTransition);
+				Result = GetPriorityRunStatus(Result, GlobalTaskRunStatus);
+				if (Result == EStateTreeRunStatus::Failed)
+				{
+					break;
+				}
+			}
+			else
 			{
-				break;
+				// New frame might have temporaries that need to move.
+				MoveGlobalTemporariesToInstances(CurrentFrame);
 			}
 		}
 
@@ -3535,7 +3587,7 @@ EStateTreeRunStatus FStateTreeExecutionContext::EnterState(const TSharedPtr<FSel
 			{
 				const int32 AssetTaskIndex = CurrentState.TasksBegin + StateTaskIndex;
 				const FStateTreeTaskBase& Task = CurrentStateTree->Nodes[AssetTaskIndex].Get<const FStateTreeTaskBase>();
-				const FStateTreeDataView TaskInstanceView = GetDataView(CurrentParentFrame, CurrentFrame, Task.InstanceDataHandle);
+				const FStateTreeDataView TaskInstanceView = MoveTemporaryToInstance(CurrentFrame, Task.InstanceDataHandle);
 
 				// Ignore disabled task
 				if (Task.bTaskEnabled == false)
@@ -4014,19 +4066,44 @@ EStateTreeRunStatus FStateTreeExecutionContext::TickEvaluatorsAndGlobalTasksForF
 	return Result;
 }
 
+void FStateTreeExecutionContext::MoveGlobalTemporariesToInstances(FStateTreeExecutionFrame& Frame)
+{
+	TNotNull<const UStateTree*> CurrentStateTree = Frame.StateTree;
+
+	// Move evaluators
+	const int32 EvaluatorsEnd = CurrentStateTree->EvaluatorsBegin + CurrentStateTree->EvaluatorsNum;
+	for (int32 EvalIndex = CurrentStateTree->EvaluatorsBegin; EvalIndex < EvaluatorsEnd; ++EvalIndex)
+	{
+		const FStateTreeEvaluatorBase& Eval = CurrentStateTree->Nodes[EvalIndex].Get<const FStateTreeEvaluatorBase>();
+		MoveTemporaryToInstance(Frame, Eval.InstanceDataHandle);
+	}
+
+	// Start Global tasks
+	// Only move data for tasks that successfully create their EnterState.
+	const int32 ActiveNodeIndexMax = Frame.ActiveNodeIndex.IsValid() ? Frame.ActiveNodeIndex.AsInt32() + 1 : 0;
+	const int32 AssetTaskMax = FMath::Min(ActiveNodeIndexMax, static_cast<int32>(CurrentStateTree->GlobalTasksBegin + CurrentStateTree->GlobalTasksNum));
+	for (int32 AssetTaskIndex = CurrentStateTree->GlobalTasksBegin; AssetTaskIndex < AssetTaskMax; ++AssetTaskIndex)
+	{
+		const FStateTreeTaskBase& Task = CurrentStateTree->Nodes[AssetTaskIndex].Get<const FStateTreeTaskBase>();
+		MoveTemporaryToInstance(Frame, Task.InstanceDataHandle);
+	}
+}
+
 EStateTreeRunStatus FStateTreeExecutionContext::StartGlobalsForFrameOnActiveInstances(const FStateTreeExecutionFrame* ParentFrame, FStateTreeExecutionFrame& Frame, FStateTreeTransitionResult& Transition)
 {
 	constexpr bool bOnActiveInstances = true;
-	return StartGlobalsForFrameInternal<bOnActiveInstances>(ParentFrame, Frame, Transition);
+	constexpr bool bMoveFromTemporary = false;
+	return StartGlobalsForFrameInternal<bOnActiveInstances, bMoveFromTemporary>(ParentFrame, Frame, Transition);
 }
 
 EStateTreeRunStatus FStateTreeExecutionContext::StartGlobalsForFrameWithValidation(const FStateTreeExecutionFrame* ParentFrame, FStateTreeExecutionFrame& Frame, FStateTreeTransitionResult& Transition)
 {
 	constexpr bool bOnActiveInstances = false;
-	return StartGlobalsForFrameInternal<bOnActiveInstances>(ParentFrame, Frame, Transition);
+	constexpr bool bMoveFromTemporary = false;
+	return StartGlobalsForFrameInternal<bOnActiveInstances, bMoveFromTemporary>(ParentFrame, Frame, Transition);
 }
 
-template<bool bOnActiveInstances>
+template<bool bOnActiveInstances, bool bMoveFromTemporary>
 EStateTreeRunStatus FStateTreeExecutionContext::StartGlobalsForFrameInternal(const FStateTreeExecutionFrame* CurrentParentFrame, FStateTreeExecutionFrame& CurrentFrame, FStateTreeTransitionResult& Transition)
 {
 	if (!CurrentFrame.bIsGlobalFrame)
@@ -4056,9 +4133,20 @@ EStateTreeRunStatus FStateTreeExecutionContext::StartGlobalsForFrameInternal(con
 	for (int32 EvalIndex = CurrentStateTree->EvaluatorsBegin; EvalIndex < EvaluatorsEnd; ++EvalIndex)
 	{
 		const FStateTreeEvaluatorBase& Eval = CurrentStateTree->Nodes[EvalIndex].Get<const FStateTreeEvaluatorBase>();
-		FStateTreeDataView EvalInstanceView = bOnActiveInstances
-			? GetDataView(CurrentParentFrame, CurrentFrame, Eval.InstanceDataHandle)
-			: GetDataViewOrTemporary(CurrentParentFrame, CurrentFrame, Eval.InstanceDataHandle);
+		FStateTreeDataView EvalInstanceView;
+		if constexpr (bMoveFromTemporary)
+		{
+			EvalInstanceView = MoveTemporaryToInstance(CurrentFrame, Eval.InstanceDataHandle);
+		}
+		else if constexpr (bOnActiveInstances)
+		{
+			EvalInstanceView = GetDataView(CurrentParentFrame, CurrentFrame, Eval.InstanceDataHandle);
+		}
+		else
+		{
+			EvalInstanceView = GetDataViewOrTemporary(CurrentParentFrame, CurrentFrame, Eval.InstanceDataHandle);
+		}
+
 		if constexpr (!bOnActiveInstances)
 		{
 			if (!EvalInstanceView.IsValid())
@@ -4114,6 +4202,12 @@ EStateTreeRunStatus FStateTreeExecutionContext::StartGlobalsForFrameInternal(con
 		const int32 AssetTaskIndex = CurrentStateTree->GlobalTasksBegin + GlobalTaskIndex;
 		const FStateTreeTaskBase& Task = CurrentStateTree->Nodes[AssetTaskIndex].Get<const FStateTreeTaskBase>();
 
+		FStateTreeDataView TaskDataView;
+		if constexpr (bMoveFromTemporary)
+		{
+			TaskDataView = MoveTemporaryToInstance(CurrentFrame, Task.InstanceDataHandle);
+		}
+
 		// Ignore disabled task
 		if (Task.bTaskEnabled == false)
 		{
@@ -4121,9 +4215,12 @@ EStateTreeRunStatus FStateTreeExecutionContext::StartGlobalsForFrameInternal(con
 			continue;
 		}
 
-		FStateTreeDataView TaskDataView = bOnActiveInstances
-			? GetDataView(CurrentParentFrame, CurrentFrame, Task.InstanceDataHandle)
-			: GetDataViewOrTemporary(CurrentParentFrame, CurrentFrame, Task.InstanceDataHandle);
+		if constexpr (!bMoveFromTemporary)
+		{
+			TaskDataView = bOnActiveInstances
+				? GetDataView(CurrentParentFrame, CurrentFrame, Task.InstanceDataHandle)
+				: GetDataViewOrTemporary(CurrentParentFrame, CurrentFrame, Task.InstanceDataHandle);
+		}
 		if constexpr (!bOnActiveInstances)
 		{
 			if (!TaskDataView.IsValid())
diff --git a/Engine/Plugins/Runtime/StateTree/Source/StateTreeModule/Private/StateTreeTypes.cpp b/Engine/Plugins/Runtime/StateTree/Source/StateTreeModule/Private/StateTreeTypes.cpp
--- a/Engine/Plugins/Runtime/StateTree/Source/StateTreeModule/Private/StateTreeTypes.cpp
+++ b/Engine/Plugins/Runtime/StateTree/Source/StateTreeModule/Private/StateTreeTypes.cpp
@@ -179,6 +179,25 @@ FStateTreeDataHandle FStateTreeDataHandle::ToObjectSource() const
 	}
 }
 
+FStateTreeDataHandle FStateTreeDataHandle::ToStructSource() const
+{
+	switch (Source)
+	{
+	case EStateTreeDataSourceType::GlobalInstanceDataObject:
+		return FStateTreeDataHandle(EStateTreeDataSourceType::GlobalInstanceData, Index, StateHandle);
+	case EStateTreeDataSourceType::ActiveInstanceDataObject:
+		return FStateTreeDataHandle(EStateTreeDataSourceType::ActiveInstanceData, Index, StateHandle);
+	case EStateTreeDataSourceType::SharedInstanceDataObject:
+		return FStateTreeDataHandle(EStateTreeDataSourceType::SharedInstanceData, Index, StateHandle);
+	case EStateTreeDataSourceType::EvaluationScopeInstanceDataObject:
+		return FStateTreeDataHandle(EStateTreeDataSourceType::EvaluationScopeInstanceData, Index, StateHandle);
+	case EStateTreeDataSourceType::ExecutionRuntimeDataObject:
+		return FStateTreeDataHandle(EStateTreeDataSourceType::ExecutionRuntimeData, Index, StateHandle);
+	default:
+		return *this;
+	}
+}
+
 FString FStateTreeDataHandle::Describe() const
 {
 	switch (Source)
diff --git a/Engine/Plugins/Runtime/StateTree/Source/StateTreeModule/Public/StateTreeExecutionContext.h b/Engine/Plugins/Runtime/StateTree/Source/StateTreeModule/Public/StateTreeExecutionContext.h
--- a/Engine/Plugins/Runtime/StateTree/Source/StateTreeModule/Public/StateTreeExecutionContext.h
+++ b/Engine/Plugins/Runtime/StateTree/Source/StateTreeModule/Public/StateTreeExecutionContext.h
@@ -988,7 +988,10 @@ protected:
 	UE_API EStateTreeRunStatus TickEvaluatorsAndGlobalTasksForFrame(float DeltaTime, bool bTickGlobalTasks, int32 FrameIndex, const FStateTreeExecutionFrame* CurrentParentFrame, const TNotNull<FStateTreeExecutionFrame*> CurrentFrame);
 
 private:
-	template<bool bOnActiveInstances>
+	/** Move the temporary data for each evaluators and tasks to instance data. */
+	void MoveGlobalTemporariesToInstances(FStateTreeExecutionFrame& Frame);
+
+	template<bool bOnActiveInstances, bool bMoveFromTemporary>
 	EStateTreeRunStatus StartGlobalsForFrameInternal(const FStateTreeExecutionFrame* ParentFrame, FStateTreeExecutionFrame& Frame, FStateTreeTransitionResult& Transition);
 
 protected:
@@ -1403,6 +1406,12 @@ PRAGMA_ENABLE_DEPRECATION_WARNINGS
 	/** @return String describing all events that are currently being processed  for logging and debug. */
 	[[nodiscard]] UE_API FString DebugGetEventsAsString() const;
 
+	/**
+	 * Move the data from the temporary buffer to the allocated instance data relative to given frame.
+	 * @return instanced data of the specified handle relative to given frame.
+	 */
+	UE_API FStateTreeDataView MoveTemporaryToInstance(const FStateTreeExecutionFrame& Frame, const FStateTreeDataHandle Handle);
+
 	/** @return data view of the specified handle relative to given frame. */
 	[[nodiscard]] UE_API FStateTreeDataView GetDataView(const FStateTreeExecutionFrame* ParentFrame, const FStateTreeExecutionFrame& CurrentFrame, const FStateTreeDataHandle Handle);
 	[[nodiscard]] UE_API FStateTreeDataView GetDataView(const FStateTreeExecutionFrame* ParentFrame, const FStateTreeExecutionFrame& CurrentFrame, const FPropertyBindingCopyInfo& CopyInfo);
diff --git a/Engine/Plugins/Runtime/StateTree/Source/StateTreeModule/Public/StateTreeTypes.h b/Engine/Plugins/Runtime/StateTree/Source/StateTreeModule/Public/StateTreeTypes.h
--- a/Engine/Plugins/Runtime/StateTree/Source/StateTreeModule/Public/StateTreeTypes.h
+++ b/Engine/Plugins/Runtime/StateTree/Source/StateTreeModule/Public/StateTreeTypes.h
@@ -456,6 +456,7 @@ struct FStateTreeDataHandle
 	STATETREEMODULE_API bool IsObjectSource() const;
 
 	STATETREEMODULE_API FStateTreeDataHandle ToObjectSource() const;
+	STATETREEMODULE_API FStateTreeDataHandle ToStructSource() const;
 
 	STATETREEMODULE_API FString Describe() const;
 
diff --git a/Engine/Plugins/Runtime/StateTree/Source/StateTreeTestSuite/Private/StateTreeLinkedStateTest.cpp b/Engine/Plugins/Runtime/StateTree/Source/StateTreeTestSuite/Private/StateTreeLinkedStateTest.cpp
--- a/Engine/Plugins/Runtime/StateTree/Source/StateTreeTestSuite/Private/StateTreeLinkedStateTest.cpp
+++ b/Engine/Plugins/Runtime/StateTree/Source/StateTreeTestSuite/Private/StateTreeLinkedStateTest.cpp
@@ -4,6 +4,7 @@
 #include "StateTreeTestBase.h"
 #include "StateTreeTestTypes.h"
 
+#include "StateTreeAsyncExecutionContext.h"
 #include "StateTreeCompilerLog.h"
 #include "StateTreeEditorData.h"
 #include "StateTreeCompiler.h"
@@ -1489,6 +1490,143 @@ struct FStateTreeTest_Linked_GlobalParameter : FStateTreeTestBase
 };
 IMPLEMENT_STATE_TREE_INSTANT_TEST(FStateTreeTest_Linked_GlobalParameter, "System.StateTree.LinkedAsset.GlobalParameter");
 
+struct FStateTreeTest_Linked_AccessGlobalInstanceData : FStateTreeTestBase
+{
+	virtual bool InstantTest() override
+	{
+		//Tree 1
+		//	Root
+		//		State1
+		//			State2 (linked Tree2)
+		//Tree 2 (with global tasks)
+		//	Root
+
+		FSimpleDelegate State1OnEnter;
+		bool bDelegateCalled = false;
+
+		// Tree 2
+		UStateTree& StateTree2 = NewStateTree();
+		{
+			UStateTreeEditorData& EditorData2 = *Cast<UStateTreeEditorData>(StateTree2.EditorData);
+
+			// Global tasks
+			TStateTreeEditorNode<FTestTask_PrintValue>& GlobalTask1 = EditorData2.AddGlobalTask<FTestTask_PrintValue>("Tree2GlobalTask1");
+			{
+				GlobalTask1.GetInstanceData().Value = 77;
+			}
+			{
+				TStateTreeEditorNode<FTestTask_PrintValue>& GlobalTask2 = EditorData2.AddGlobalTask<FTestTask_PrintValue>("Tree2GlobalTask2");
+				GlobalTask2.GetInstanceData().Value = 88;
+				GlobalTask2.GetNode().CustomEnterStateFunc = [this, &State1OnEnter, &bDelegateCalled](FStateTreeExecutionContext& ExecContext, const FTestTask_PrintValue* Node)
+					{
+						FTestTask_PrintValue::FInstanceDataType& InstanceData = ExecContext.GetInstanceData<FTestTask_PrintValue::FInstanceDataType>(*Node);
+						if (!GetTestRunner().TestTrue(TEXT("Tree2GlobalTask2 instance data is not 88"), InstanceData.Value == 88))
+						{
+							return;
+						}
+						State1OnEnter.BindLambda([this, &bDelegateCalled, WeakContext = ExecContext.MakeWeakExecutionContext()]()
+							{
+								FStateTreeStrongReadOnlyExecutionContext StrongContext = WeakContext.MakeStrongReadOnlyExecutionContext();
+								const FTestTask_PrintValue::FInstanceDataType* InstanceData = StrongContext.GetInstanceDataPtr<const FTestTask_PrintValue::FInstanceDataType>();
+								if (!GetTestRunner().TestTrue(TEXT("Tree2GlobalTask2 instance data is nullptr"), InstanceData != nullptr))
+								{
+									return;
+								}
+								if (!GetTestRunner().TestTrue(TEXT("Tree2GlobalTask2 instance data is not 88"), InstanceData && InstanceData->Value == 88))
+								{
+									return;
+								}
+								bDelegateCalled = true;
+							});
+
+					};
+			}
+
+			UStateTreeState& Root = EditorData2.AddSubTree("Tree2StateRoot");
+			{
+				TStateTreeEditorNode<FTestTask_PrintValue>& Task1 = Root.AddTask<FTestTask_PrintValue>("Tree2StateRootTask1");
+				Task1.GetInstanceData().Value = 99;
+			}
+
+			FStateTreeCompilerLog Log;
+			FStateTreeCompiler Compiler2(Log);
+			const bool bResult2 = Compiler2.Compile(StateTree2);
+			AITEST_TRUE(TEXT("StateTree2 should get compiled"), bResult2);
+		}
+
+		// Tree 1
+		UStateTree& StateTree1 = NewStateTree();
+		FGuid Tree1GlobalParameter_ValueID_Int;
+		{
+			UStateTreeEditorData& EditorData1 = *Cast<UStateTreeEditorData>(StateTree1.EditorData);
+
+			// Global tasks
+			TStateTreeEditorNode<FTestTask_PrintValue>& GlobalTask1 = EditorData1.AddGlobalTask<FTestTask_PrintValue>("Tree1GlobalTask1");
+			{
+				GlobalTask1.GetInstanceData().Value = 11;
+			}
+
+			UStateTreeState& Root = EditorData1.AddSubTree("Tree1StateRoot");
+			{
+				TStateTreeEditorNode<FTestTask_PrintValue>& Task1 = Root.AddTask<FTestTask_PrintValue>("Tree1StateRootTask1");
+				Task1.GetInstanceData().Value = 22;
+			}
+			{
+				UStateTreeState& State1 = Root.AddChildState("Tree1State1");
+				TStateTreeEditorNode<FTestTask_PrintValue>& Task1 = State1.AddTask<FTestTask_PrintValue>("Tree1State1Task1");
+				Task1.GetInstanceData().Value = 33;
+				Task1.GetNode().CustomEnterStateFunc = [this, &State1OnEnter](FStateTreeExecutionContext& ExecContext, const FTestTask_PrintValue* Node)
+					{
+						FTestTask_PrintValue::FInstanceDataType& InstanceData = ExecContext.GetInstanceData<FTestTask_PrintValue::FInstanceDataType>(*Node);
+						if (!GetTestRunner().TestTrue(TEXT("Tree1State1Task1 instance data is not 33"), InstanceData.Value == 33))
+						{
+							return;
+						}
+						State1OnEnter.ExecuteIfBound();
+					};
+				{
+					UStateTreeState& State2 = State1.AddChildState("Tree1State2", EStateTreeStateType::LinkedAsset);
+					State2.SetLinkedStateAsset(&StateTree2);
+				}
+			}
+
+			FStateTreeCompilerLog Log;
+			FStateTreeCompiler Compiler1(Log);
+			const bool bResult1 = Compiler1.Compile(StateTree1);
+			AITEST_TRUE(TEXT("StateTree1 should get compiled"), bResult1);
+		}
+
+		FStateTreeInstanceData InstanceData;
+		{
+			FTestStateTreeExecutionContext Exec(StateTree1, StateTree1, InstanceData);
+			const bool bInitSucceeded = Exec.IsValid();
+			AITEST_TRUE(TEXT("StateTree should init"), bInitSucceeded);
+		}
+		{
+			FTestStateTreeExecutionContext Exec(StateTree1, StateTree1, InstanceData);
+			const EStateTreeRunStatus Status = Exec.Start();
+			AITEST_EQUAL(TEXT("Start should complete with Running"), Status, EStateTreeRunStatus::Running);
+			AITEST_TRUE(TEXT("In correct states"), Exec.ExpectInActiveStates("Tree1StateRoot", "Tree1State1", "Tree1State2", "Tree2StateRoot"));
+			AITEST_TRUE(TEXT("Start should enter Tree1GlobalTask1"), Exec.Expect("Tree1GlobalTask1", TEXT("EnterState11")));
+			AITEST_TRUE(TEXT("Start should enter Tree1StateRootTask1"), Exec.Expect("Tree1StateRootTask1", TEXT("EnterState22")));
+			AITEST_TRUE(TEXT("Start should enter Tree1State1Task1"), Exec.Expect("Tree1State1Task1", TEXT("EnterState33")));
+			AITEST_TRUE(TEXT("Start should enter Tree2GlobalTask1"), Exec.Expect("Tree2GlobalTask1", TEXT("EnterState77")));
+			AITEST_TRUE(TEXT("Start should enter Tree2GlobalTask2"), Exec.Expect("Tree2GlobalTask2", TEXT("EnterState88")));
+			AITEST_TRUE(TEXT("Start should enter Tree2StateRootTask1"), Exec.Expect("Tree2StateRootTask1", TEXT("EnterState99")));
+
+			AITEST_TRUE(TEXT("Start should enter bDelegateCalled"), bDelegateCalled);
+			Exec.LogClear();
+		}
+		{
+			FTestStateTreeExecutionContext Exec(StateTree1, StateTree1, InstanceData);
+			Exec.Stop();
+		}
+
+		return true;
+	}
+};
+IMPLEMENT_STATE_TREE_INSTANT_TEST(FStateTreeTest_Linked_AccessGlobalInstanceData, "System.StateTree.LinkedAsset.AccessGlobalInstanceData");
+
 struct FStateTreeTest_Linked_FinishGlobalTasks : FStateTreeTestBase
 {
 	virtual bool InstantTest() override
