[StateTree] Set the alignment when there are only object wrapper nodes.
#was-backout 47028400

diff --git a/Engine/Plugins/Runtime/StateTree/Source/StateTreeModule/Private/StateTreeEvaluationScopeInstanceContainer.cpp b/Engine/Plugins/Runtime/StateTree/Source/StateTreeModule/Private/StateTreeEvaluationScopeInstanceContainer.cpp
--- a/Engine/Plugins/Runtime/StateTree/Source/StateTreeModule/Private/StateTreeEvaluationScopeInstanceContainer.cpp
+++ b/Engine/Plugins/Runtime/StateTree/Source/StateTreeModule/Private/StateTreeEvaluationScopeInstanceContainer.cpp
@@ -21,7 +21,7 @@ namespace Private
 		MemoryRequirement.Size += StructSize;
 		MemoryRequirement.Size += sizeof(DebugStructEndTag); // for struct end tag
 	}
-}
+} // namespace Private
 
 void FEvaluationScopeInstanceContainer::FMemoryRequirementBuilder::Add(TNotNull<const UScriptStruct*> Struct)
 {
@@ -38,6 +38,19 @@ void FEvaluationScopeInstanceContainer::FMemoryRequirementBuilder::Add(TNotNull<
 
 FEvaluationScopeInstanceContainer::FMemoryRequirement FEvaluationScopeInstanceContainer::FMemoryRequirementBuilder::Build()
 {
+	if (MemoryRequirement.NumberOfElements == 0)
+	{
+		return FMemoryRequirement();
+	}
+
+	// There is no container memory for object wrapper. The instance is directly saved in the StructView.
+	//In case all the structs are object wrapper.
+	if (FirstStructAlignment <= 0)
+	{
+		MemoryRequirement.Alignment = alignof(decltype(Private::DebugStructEndTag));
+		FirstStructAlignment = alignof(decltype(Private::DebugStructEndTag));
+	}
+
 	int32 ContainerSize = sizeof(FItem) * MemoryRequirement.NumberOfElements;
 	ContainerSize += sizeof(Private::DebugTableEndTag); // table for end tag
 	ContainerSize = Align(ContainerSize, FirstStructAlignment);
@@ -93,12 +106,13 @@ void FEvaluationScopeInstanceContainer::Add(FStateTreeDataHandle DataHandle, FCo
 	}
 	else
 	{
+		new (ItemContainer + NumberOfElements)FItem();
+
 		const FStateTreeInstanceObjectWrapper& Wrapper = DefaultInstance.Get<const FStateTreeInstanceObjectWrapper>();
 		if (Wrapper.InstanceObject)
 		{
 			UObject* DuplicatedObject = ::DuplicateObject(Wrapper.InstanceObject, GetTransientPackage());
 
-			new (ItemContainer + NumberOfElements)FItem();
 			ItemContainer[NumberOfElements].Instance = FStateTreeDataView(DuplicatedObject);
 			ItemContainer[NumberOfElements].DataHandle = DataHandle;
 		}
@@ -115,7 +129,8 @@ void FEvaluationScopeInstanceContainer::Reset()
 		for (; NumberOfElements > 0; --NumberOfElements)
 		{
 			FItem* ItemContainer = static_cast<FItem*>(Memory);
-			if (!ItemContainer[NumberOfElements - 1].DataHandle.IsObjectSource())
+			const UStruct* InstanceStruct = ItemContainer[NumberOfElements - 1].Instance.GetStruct();
+			if (InstanceStruct && Cast<const UClass>(InstanceStruct) == nullptr)
 			{
 				FStateTreeDataView& Instance = ItemContainer[NumberOfElements - 1].Instance;
 				Instance.GetStruct()->DestroyStruct(Instance.GetMutableMemory());
@@ -183,10 +198,15 @@ void FEvaluationScopeInstanceContainer::TestDebugTags() const
 		FItem* ItemContainer = static_cast<FItem*>(Memory);
 		for (int32 Index = 0; Index < NumberOfElements; ++Index)
 		{
-			const void* MemoryForStructDebugTag = (const uint8*)ItemContainer[Index].Instance.GetMutableMemory()
-				+ ItemContainer[Index].Instance.GetStruct()->GetStructureSize();
-			const int32 StructEndTag = *reinterpret_cast<const uint32*>(MemoryForStructDebugTag);
-			ensure(StructEndTag == Private::DebugStructEndTag);
+			// Object wrapper do not have EndTag, they use the FItem directly.
+			const UStruct* InstanceStruct = ItemContainer[Index].Instance.GetStruct();
+			if (InstanceStruct && Cast<const UClass>(InstanceStruct) == nullptr)
+			{
+				const void* MemoryForStructDebugTag = (const uint8*)ItemContainer[Index].Instance.GetMutableMemory()
+					+ InstanceStruct->GetStructureSize();
+				const int32 StructEndTag = *reinterpret_cast<const uint32*>(MemoryForStructDebugTag);
+				ensure(StructEndTag == Private::DebugStructEndTag);
+			}
 		}
 	}
 #endif
diff --git a/Engine/Plugins/Runtime/StateTree/Source/StateTreeModule/Private/StateTreeExecutionContext.cpp b/Engine/Plugins/Runtime/StateTree/Source/StateTreeModule/Private/StateTreeExecutionContext.cpp
--- a/Engine/Plugins/Runtime/StateTree/Source/StateTreeModule/Private/StateTreeExecutionContext.cpp
+++ b/Engine/Plugins/Runtime/StateTree/Source/StateTreeModule/Private/StateTreeExecutionContext.cpp
@@ -3145,7 +3145,7 @@ bool FStateTreeExecutionContext::CopyBatchInternal(const FStateTreeExecutionFram
 		check(Batch.PropertyFunctionsBegin.IsValid() && Batch.PropertyFunctionsEnd.IsValid());
 
 		const FEvaluationScopeInstanceContainer::FMemoryRequirement& MemoryRequirement = CurrentFrame.StateTree->PropertyFunctionEvaluationScopeMemoryRequirements[BindingsBatch.Get()];
-		if (MemoryRequirement.Size > 0)
+		if (MemoryRequirement.HasMemory())
 		{
 			void* FunctionBindingMemory = FMemory_Alloca_Aligned(MemoryRequirement.Size, MemoryRequirement.Alignment);
 			EvaluationScopeContainer = FEvaluationScopeInstanceContainer(FunctionBindingMemory, MemoryRequirement);
@@ -4837,7 +4837,7 @@ bool FStateTreeExecutionContext::TestAllConditionsInternal(const FStateTreeExecu
 	FEvaluationScopeInstanceContainer EvaluationScopeContainer;
 	bool bEvaluationScopeContainerPushed = false;
 
-	if (MemoryRequirement.Size > 0)
+	if (MemoryRequirement.HasMemory())
 	{
 		void* FunctionBindingMemory = FMemory_Alloca_Aligned(MemoryRequirement.Size, MemoryRequirement.Alignment);
 		EvaluationScopeContainer = FEvaluationScopeInstanceContainer(FunctionBindingMemory, MemoryRequirement);
@@ -4955,7 +4955,7 @@ void FStateTreeExecutionContext::RunAllConditionsInternal(const FStateTreeExecut
 
 		FEvaluationScopeInstanceContainer EvaluationScopeContainer;
 		bool bEvaluationScopeContainerPushed = false;
-		if (CurrentState.EnterConditionEvaluationScopeMemoryRequirement.Size > 0)
+		if (CurrentState.EnterConditionEvaluationScopeMemoryRequirement.HasMemory())
 		{
 			void* ScopedConditionMemory = FMemory_Alloca_Aligned(CurrentState.EnterConditionEvaluationScopeMemoryRequirement.Size, CurrentState.EnterConditionEvaluationScopeMemoryRequirement.Alignment);
 			EvaluationScopeContainer = FEvaluationScopeInstanceContainer(ScopedConditionMemory, CurrentState.EnterConditionEvaluationScopeMemoryRequirement);
@@ -5058,7 +5058,7 @@ float FStateTreeExecutionContext::EvaluateUtilityWithValidation(const FStateTree
 	FEvaluationScopeInstanceContainer EvaluationScopeContainer;
 	bool bEvaluationScopeContainerPushed = false;
 
-	if (MemoryRequirement.Size > 0)
+	if (MemoryRequirement.HasMemory())
 	{
 		void* FunctionBindingMemory = FMemory_Alloca_Aligned(MemoryRequirement.Size, MemoryRequirement.Alignment);
 		EvaluationScopeContainer = FEvaluationScopeInstanceContainer(FunctionBindingMemory, MemoryRequirement);
diff --git a/Engine/Plugins/Runtime/StateTree/Source/StateTreeModule/Private/StateTreeInstanceData.cpp b/Engine/Plugins/Runtime/StateTree/Source/StateTreeModule/Private/StateTreeInstanceData.cpp
--- a/Engine/Plugins/Runtime/StateTree/Source/StateTreeModule/Private/StateTreeInstanceData.cpp
+++ b/Engine/Plugins/Runtime/StateTree/Source/StateTreeModule/Private/StateTreeInstanceData.cpp
@@ -1023,20 +1023,22 @@ void FStateTreeInstanceData::Append(UObject& InOwner, TConstArrayView<FConstStru
 		// The instance struct can be empty, in which case the temporary instance is ignored.
 		// If the source is specified, move it to the instance data.
 		// We assume that if the source is object wrapper, it is already the instance we want.
-		if (Struct.IsValid()
-			&& (Source && Source->IsValid()))
+		if (Struct.IsValid())
 		{
-			check(Struct.GetScriptStruct() == Source->GetScriptStruct());
+			if (Source && Source->IsValid())
+			{
+				check(Struct.GetScriptStruct() == Source->GetScriptStruct());
 				
-			FMemory::Memswap(Struct.GetMemory(), Source->GetMutableMemory(), Struct.GetScriptStruct()->GetStructureSize());
-			Source->Reset();
-		}
-		else if (FStateTreeInstanceObjectWrapper* Wrapper = Struct.GetPtr<FStateTreeInstanceObjectWrapper>())
-		{
-			if (Wrapper->InstanceObject)
+				FMemory::Memswap(Struct.GetMemory(), Source->GetMutableMemory(), Struct.GetScriptStruct()->GetStructureSize());
+				Source->Reset();
+			}
+			else if (FStateTreeInstanceObjectWrapper* Wrapper = Struct.GetPtr<FStateTreeInstanceObjectWrapper>())
 			{
-				const bool bDuplicate = Args.bDuplicateWrappedObject || &InOwner != Wrapper->InstanceObject->GetOuter();
-				Wrapper->InstanceObject = UE::StateTree::InstanceData::Private::CopyNodeInstance(Wrapper->InstanceObject, &InOwner, bDuplicate);
+				if (Wrapper->InstanceObject)
+				{
+					const bool bDuplicate = Args.bDuplicateWrappedObject || &InOwner != Wrapper->InstanceObject->GetOuter();
+					Wrapper->InstanceObject = UE::StateTree::InstanceData::Private::CopyNodeInstance(Wrapper->InstanceObject, &InOwner, bDuplicate);
+				}
 			}
 		}
 	}
diff --git a/Engine/Plugins/Runtime/StateTree/Source/StateTreeModule/Public/StateTreeTypes.h b/Engine/Plugins/Runtime/StateTree/Source/StateTreeModule/Public/StateTreeTypes.h
--- a/Engine/Plugins/Runtime/StateTree/Source/StateTreeModule/Public/StateTreeTypes.h
+++ b/Engine/Plugins/Runtime/StateTree/Source/StateTreeModule/Public/StateTreeTypes.h
@@ -479,6 +479,10 @@ namespace UE::StateTree::InstanceData
 		int32 Size = 0;
 		int32 Alignment = 0;
 		int32 NumberOfElements = 0;
+		bool HasMemory() const
+		{
+			return Size > 0;
+		}
 	};
 }
 
