[StateTree] Move Append to the StateTree compiler for safety and to access private members.

#rnx

diff --git a/Engine/Plugins/Runtime/StateTree/Source/StateTreeEditorModule/Private/StateTreeCompiler.cpp b/Engine/Plugins/Runtime/StateTree/Source/StateTreeEditorModule/Private/StateTreeCompiler.cpp
--- a/Engine/Plugins/Runtime/StateTree/Source/StateTreeEditorModule/Private/StateTreeCompiler.cpp
+++ b/Engine/Plugins/Runtime/StateTree/Source/StateTreeEditorModule/Private/StateTreeCompiler.cpp
@@ -59,6 +59,11 @@ namespace UE::StateTree
 
 namespace UE::StateTree::Compiler
 {
+	// Version used by AppendToClassSchema. Bump this version number when a recompile of all assets is needed.
+	//Only bump when only the compiler changes.
+	//You do not need to bump the version when the memory layout of UStateTree has changed.
+	static constexpr int32 ClassSchemaVersion = 1;
+
 	FAutoConsoleVariable CVarLogEnableParameterDelegateDispatcherBinding(
 		TEXT("StateTree.Compiler.EnableParameterDelegateDispatcherBinding"),
 		false,
@@ -406,6 +411,23 @@ namespace UE::StateTree::Compiler
 
 }; // UE::StateTree::Compiler
 
+void FStateTreeCompiler::AppendToClassSchema(FAppendToClassSchemaContext& Context)
+{
+	Context.Update(&UE::StateTree::Compiler::ClassSchemaVersion, sizeof(UE::StateTree::Compiler::ClassSchemaVersion));
+
+	auto UpdateContext = [&Context](const FStringView Name)
+		{
+			const IConsoleVariable* Variable = IConsoleManager::Get().FindConsoleVariable(Name.GetData());
+			ensure(Variable != nullptr);
+			const bool Value = Variable == nullptr || !Variable->GetBool();
+			Context.Update(&Value, sizeof(bool));
+		};
+	UpdateContext(TEXT("StateTree.Compiler.EnableParameterDelegateDispatcherBinding"));
+	UpdateContext(TEXT("StateTree.Compiler.EnablePropertyFunctionWithEvaluationScopeInstanceData"));
+	UpdateContext(TEXT("StateTree.Compiler.EnableConditionWithEvaluationScopeInstanceData"));
+	UpdateContext(TEXT("StateTree.Compiler.EnableUtilityConsiderationWithEvaluationScopeInstanceData"));
+}
+
 bool FStateTreeCompiler::Compile(UStateTree& InStateTree)
 {
 	return Compile(&InStateTree);
diff --git a/Engine/Plugins/Runtime/StateTree/Source/StateTreeEditorModule/Private/StateTreeEditorModule.cpp b/Engine/Plugins/Runtime/StateTree/Source/StateTreeEditorModule/Private/StateTreeEditorModule.cpp
--- a/Engine/Plugins/Runtime/StateTree/Source/StateTreeEditorModule/Private/StateTreeEditorModule.cpp
+++ b/Engine/Plugins/Runtime/StateTree/Source/StateTreeEditorModule/Private/StateTreeEditorModule.cpp
@@ -91,6 +91,11 @@ void FStateTreeEditorModule::StartupModule()
 			return UStateTreeEditingSubsystem::CalculateStateTreeHash(&InStateTree);
 		});
 	UE::StateTree::Delegates::Private::OnRequestAssetRegistryTags.BindRaw(this, &FStateTreeEditorModule::HandleRequestAssetRegistryTags);
+	UE::StateTree::Delegates::Private::OnAppendToClassSchema.BindLambda([](FAppendToClassSchemaContext& Context)
+		{
+			FStateTreeCompiler::AppendToClassSchema(Context);
+		});
+
 	UE::StateTree::Compiler::FCompilerManager::Startup();
 
 	OnPostEngineInitHandle = FCoreDelegates::OnPostEngineInit.AddLambda([this]()
@@ -172,6 +177,8 @@ void FStateTreeEditorModule::ShutdownModule()
 	NodeClassCache.Reset();
 
 	UE::StateTree::Compiler::FCompilerManager::Shutdown();
+
+	UE::StateTree::Delegates::Private::OnAppendToClassSchema.Unbind();
 	UE::StateTree::Delegates::Private::OnRequestAssetRegistryTags.Unbind();
 	UE::StateTree::Delegates::OnRequestEditorHash.Unbind();
 }
diff --git a/Engine/Plugins/Runtime/StateTree/Source/StateTreeEditorModule/Public/StateTreeCompiler.h b/Engine/Plugins/Runtime/StateTree/Source/StateTreeEditorModule/Public/StateTreeCompiler.h
--- a/Engine/Plugins/Runtime/StateTree/Source/StateTreeEditorModule/Public/StateTreeCompiler.h
+++ b/Engine/Plugins/Runtime/StateTree/Source/StateTreeEditorModule/Public/StateTreeCompiler.h
@@ -6,6 +6,7 @@
 
 #define UE_API STATETREEEDITORMODULE_API
 
+struct FAppendToClassSchemaContext;
 struct FStructView;
 
 enum class EStateTreeExpressionOperand : uint8;
@@ -58,6 +59,8 @@ public:
 	UE_API bool Compile(TNotNull<UStateTree*> StateTree);
 	UE_API bool Compile(UStateTree& InStateTree);
 
+	UE_API static void AppendToClassSchema(FAppendToClassSchemaContext& Context);
+
 private:
 	/** Resolves the state a transition points to, and the optional fallback for failing to enter the state. SourceState is nullptr for global tasks. */
 	bool ResolveTransitionStateAndFallback(const UStateTreeState* SourceState, const FStateTreeStateLink& Link, FStateTreeStateHandle& OutTransitionHandle, EStateTreeSelectionFallback& OutFallback) const;
diff --git a/Engine/Plugins/Runtime/StateTree/Source/StateTreeModule/Internal/StateTreeDelegatesInternal.h b/Engine/Plugins/Runtime/StateTree/Source/StateTreeModule/Internal/StateTreeDelegatesInternal.h
--- a/Engine/Plugins/Runtime/StateTree/Source/StateTreeModule/Internal/StateTreeDelegatesInternal.h
+++ b/Engine/Plugins/Runtime/StateTree/Source/StateTreeModule/Internal/StateTreeDelegatesInternal.h
@@ -6,16 +6,25 @@
 
 #define UE_API STATETREEMODULE_API
 
+#if WITH_EDITOR
 class UStateTree;
+struct FAppendToClassSchemaContext;
 class FAssetRegistryTagsContext;
+#endif
 
 namespace UE::StateTree::Delegates::Private
 {
 
+#if WITH_EDITOR
 /** A callback for UStateTree::GetAssetRegistryTags. */
 DECLARE_DELEGATE_TwoParams(FOnStateTreeRequestAssetRegistryTags, TNotNull<const UStateTree*> /*StateTree*/, FAssetRegistryTagsContext& /*Context*/);
 extern UE_API FOnStateTreeRequestAssetRegistryTags OnRequestAssetRegistryTags;
 
+/** A callback for UStateTree::AppendToClassSchema. */
+DECLARE_DELEGATE_OneParam(FOnAppendToClassSchema, FAppendToClassSchemaContext& /*Context*/);
+extern UE_API FOnAppendToClassSchema OnAppendToClassSchema;
+#endif
+
 }; // UE::StateTree::Delegates
 
 #undef UE_API
diff --git a/Engine/Plugins/Runtime/StateTree/Source/StateTreeModule/Private/StateTree.cpp b/Engine/Plugins/Runtime/StateTree/Source/StateTreeModule/Private/StateTree.cpp
--- a/Engine/Plugins/Runtime/StateTree/Source/StateTreeModule/Private/StateTree.cpp
+++ b/Engine/Plugins/Runtime/StateTree/Source/StateTreeModule/Private/StateTree.cpp
@@ -464,18 +464,7 @@ void UStateTree::DeclareConstructClasses(TArray<FTopLevelAssetPath>& OutConstruc
 void UStateTree::AppendToClassSchema(FAppendToClassSchemaContext& Context)
 {
 	Super::AppendToClassSchema(Context);
-
-	auto UpdateContext = [&Context](const FStringView Name)
-		{
-			const IConsoleVariable* Variable = IConsoleManager::Get().FindConsoleVariable(Name.GetData());
-			ensure(Variable != nullptr);
-			const bool Value = Variable == nullptr || !Variable->GetBool();
-			Context.Update(&Value, sizeof(bool));
-		};
-	UpdateContext(TEXT("StateTree.Compiler.EnableParameterDelegateDispatcherBinding"));
-	UpdateContext(TEXT("StateTree.Compiler.EnablePropertyFunctionWithEvaluationScopeInstanceData"));
-	UpdateContext(TEXT("StateTree.Compiler.EnableConditionWithEvaluationScopeInstanceData"));
-	UpdateContext(TEXT("StateTree.Compiler.EnableUtilityConsiderationWithEvaluationScopeInstanceData"));
+	UE::StateTree::Delegates::Private::OnAppendToClassSchema.ExecuteIfBound(Context);
 }
 #endif
 
diff --git a/Engine/Plugins/Runtime/StateTree/Source/StateTreeModule/Private/StateTreeDelegates.cpp b/Engine/Plugins/Runtime/StateTree/Source/StateTreeModule/Private/StateTreeDelegates.cpp
--- a/Engine/Plugins/Runtime/StateTree/Source/StateTreeModule/Private/StateTreeDelegates.cpp
+++ b/Engine/Plugins/Runtime/StateTree/Source/StateTreeModule/Private/StateTreeDelegates.cpp
@@ -34,5 +34,6 @@ namespace UE::StateTree::Delegates::Private
 {
 #if WITH_EDITOR
 	FOnStateTreeRequestAssetRegistryTags OnRequestAssetRegistryTags;
+	FOnAppendToClassSchema OnAppendToClassSchema;
 #endif
 } //UE::StateTree::Delegates::Private
