[StateTree] Fix bug found post review in ScopeInstanceContainer
#rnx

diff --git a/Engine/Plugins/Runtime/StateTree/Source/StateTreeModule/Private/StateTreeEvaluationScopeInstanceContainer.cpp b/Engine/Plugins/Runtime/StateTree/Source/StateTreeModule/Private/StateTreeEvaluationScopeInstanceContainer.cpp
--- a/Engine/Plugins/Runtime/StateTree/Source/StateTreeModule/Private/StateTreeEvaluationScopeInstanceContainer.cpp
+++ b/Engine/Plugins/Runtime/StateTree/Source/StateTreeModule/Private/StateTreeEvaluationScopeInstanceContainer.cpp
@@ -47,8 +47,8 @@ FEvaluationScopeInstanceContainer::FMemoryRequirement FEvaluationScopeInstanceCo
 	//In case all the structs are object wrapper.
 	if (FirstStructAlignment <= 0)
 	{
-		MemoryRequirement.Alignment = alignof(decltype(Private::DebugStructEndTag));
-		FirstStructAlignment = alignof(decltype(Private::DebugStructEndTag));
+		MemoryRequirement.Alignment = alignof(FItem);
+		FirstStructAlignment = alignof(FItem);
 	}
 
 	int32 ContainerSize = sizeof(FItem) * MemoryRequirement.NumberOfElements;
@@ -126,9 +126,9 @@ void FEvaluationScopeInstanceContainer::Reset()
 
 	if (bStructsHaveDestructor)
 	{
+		FItem* ItemContainer = static_cast<FItem*>(Memory);
 		for (; NumberOfElements > 0; --NumberOfElements)
 		{
-			FItem* ItemContainer = static_cast<FItem*>(Memory);
 			const UStruct* InstanceStruct = ItemContainer[NumberOfElements - 1].Instance.GetStruct();
 			if (InstanceStruct && Cast<const UClass>(InstanceStruct) == nullptr)
 			{
@@ -137,17 +137,17 @@ void FEvaluationScopeInstanceContainer::Reset()
 			}
 			if constexpr (!std::is_trivially_destructible_v<FItem>)
 			{
-				ItemContainer->~FItem();
+				ItemContainer[NumberOfElements - 1].~FItem();
 			}
 		}
 		check(NumberOfElements == 0);
 	}
 	else if constexpr (!std::is_trivially_destructible_v<FItem>)
 	{
+		FItem* ItemContainer = static_cast<FItem*>(Memory);
 		for (; NumberOfElements > 0; --NumberOfElements)
 		{
-			FItem* ItemContainer = static_cast<FItem*>(Memory);
-			ItemContainer->~FItem();
+			ItemContainer[NumberOfElements - 1].~FItem();
 		}
 		check(NumberOfElements == 0);
 	}
