[State Tree] refactored and fixed incorrectly accessing subtree parameter, global parameter and global instance when it exists in a parent frame multiple levels above.

- Now we cache corresponding Parent Frame IDs for accessing global parameter and global instance in temporary data, and cache corresponding index base for accessing global parameter and global instance in active instance data.
- Don't set Subtree and Global Parameters in UpdateInstanceData because they have already been set when creating the frames and they don't change. We update ActiveInstanceIndex, GlobalInstanceIndex and GlobalParameterIndex in UpdateInstanceData.

diff --git a/Engine/Plugins/Runtime/StateTree/Source/StateTreeModule/Private/StateTreeAsyncExecutionContext.cpp b/Engine/Plugins/Runtime/StateTree/Source/StateTreeModule/Private/StateTreeAsyncExecutionContext.cpp
--- a/Engine/Plugins/Runtime/StateTree/Source/StateTreeModule/Private/StateTreeAsyncExecutionContext.cpp
+++ b/Engine/Plugins/Runtime/StateTree/Source/StateTreeModule/Private/StateTreeAsyncExecutionContext.cpp
@@ -332,7 +332,6 @@ FStateTreeDataView TStateTreeStrongExecutionContext<bWithWriteAccess>::GetInstan
 		FStateTreeDataView InstanceDataView = InstanceData::GetDataViewOrTemporary(
 			*Storage,
 			/* SharedInstanceStorage */ nullptr,
-			ActivePath.ParentFrame,
 			*ActivePath.Frame,
 			ActivePath.GetNode().InstanceDataHandle);
 
diff --git a/Engine/Plugins/Runtime/StateTree/Source/StateTreeModule/Private/StateTreeExecutionContext.cpp b/Engine/Plugins/Runtime/StateTree/Source/StateTreeModule/Private/StateTreeExecutionContext.cpp
--- a/Engine/Plugins/Runtime/StateTree/Source/StateTreeModule/Private/StateTreeExecutionContext.cpp
+++ b/Engine/Plugins/Runtime/StateTree/Source/StateTreeModule/Private/StateTreeExecutionContext.cpp
@@ -109,6 +109,7 @@ namespace UE::StateTree::ExecutionContext::Private
 			return Number;
 	}
 
+	// @todo: should use const TConstArrayView instead
 	template<typename FrameType>
 	FrameType* FindExecutionFrame(FActiveFrameID FrameID, const TArrayView<FrameType> Frames, const TArrayView<FrameType> TemporaryFrames)
 	{
@@ -952,9 +953,31 @@ FStateTreeExecutionFrame& FStateTreeExecutionContext::FSelectStateResult::MakeAn
 
 	constexpr bool bIsGlobalFrame = true;
 	FStateTreeExecutionFrame& ExecFrame = MakeAndAddTemporaryFrame(FrameID, FrameHandle, bIsGlobalFrame);
+	check(ExecFrame.StateTree == OtherFrame.StateTree);
+
 	ExecFrame.ExternalDataBaseIndex = OtherFrame.ExternalDataBaseIndex;
-	ExecFrame.GlobalParameterDataHandle = OtherFrame.GlobalParameterDataHandle;
+
+	// new frame is global frame
+	ExecFrame.GlobalInstanceDataFrameID = ExecFrame.FrameID;
 	ExecFrame.GlobalInstanceIndexBase = OtherFrame.GlobalInstanceIndexBase;
+
+	ExecFrame.GlobalParameterDataHandle = OtherFrame.GlobalParameterDataHandle;
+	if (ExecFrame.GlobalParameterDataHandle.HasValidIndex())
+	{
+		// Global Parameter is a state parameter living within parent frame with different state tree
+		ExecFrame.GlobalParameterDataFrameID = OtherFrame.GlobalParameterDataFrameID;
+	}
+	else
+	{
+		// Global Parameter is root global parameter and new frame is global frame
+		ExecFrame.GlobalParameterDataFrameID = ExecFrame.FrameID;
+	}
+	ExecFrame.GlobalParameterIndexBase = OtherFrame.GlobalParameterIndexBase;
+
+	// new frame could be a subtree, but it's always standalone.
+	ExecFrame.StateParameterDataHandle = ExecFrame.StateTree->GetStateFromHandle(ExecFrame.RootState)->ParameterDataHandle;
+	ExecFrame.StateParameterDataFrameID = ExecFrame.FrameID;
+	
 	// Don't stop globals.
 	ExecFrame.bHaveEntered = OtherFrame.bHaveEntered;
 	OtherFrame.bHaveEntered = false;
@@ -1437,7 +1460,14 @@ EStateTreeRunStatus FStateTreeExecutionContext::Start(FStartParameters Parameter
 	constexpr bool bIsGlobalFrame = true;
 	FStateTreeExecutionFrame& InitFrame = SelectStateResult->MakeAndAddTemporaryFrame(InitFrameID, InitFrameHandle, bIsGlobalFrame);
 	InitFrame.ExecutionRuntimeIndexBase = FStateTreeIndex16(Storage.AddExecutionRuntimeData(GetOwner(), InitFrameHandle));
+
 	InitFrame.GlobalParameterDataHandle = FStateTreeDataHandle(EStateTreeDataSourceType::GlobalParameterData);
+	InitFrame.GlobalParameterDataFrameID = InitFrameID;
+
+	InitFrame.GlobalInstanceDataFrameID = InitFrameID;
+
+	InitFrame.StateParameterDataHandle = InitFrame.StateTree->GetStateFromHandle(InitFrame.RootState)->ParameterDataHandle;
+	InitFrame.StateParameterDataFrameID = InitFrameID;
 
 	if (!CollectActiveExternalData(SelectStateResult->TemporaryFrames))
 	{
@@ -2282,7 +2312,8 @@ void FStateTreeExecutionContext::UpdateInstanceData(const TSharedPtr<FSelectStat
 	const UStruct* NextStateParameterDataStruct = nullptr;
 	FStateTreeDataHandle NextStateParameterDataHandle = FStateTreeDataHandle::Invalid;
 
-	FStateTreeDataHandle CurrentGlobalParameterDataHandle = FStateTreeDataHandle(EStateTreeDataSourceType::GlobalParameterData);
+	// start from Root Global Parameter. Update to state parameter if we encounter any linked asset state.
+	FStateTreeIndex16 CurrentGlobalParameterIndexBase = FStateTreeIndex16::Invalid;
 	if (Args)
 	{
 		FActiveFrameID CurrentFrameID;
@@ -2292,10 +2323,11 @@ void FStateTreeExecutionContext::UpdateInstanceData(const TSharedPtr<FSelectStat
 		for (int32 SelectedStateIndex = 0; SelectedStateIndex < Args->SelectedStates.Num(); ++SelectedStateIndex)
 		{
 			const FActiveState& SelectedState = Args->SelectedStates[SelectedStateIndex];
-
+			const FStateTreeExecutionFrame* ParentFrame = nullptr;
 			// Global
 			if (SelectedState.GetFrameID() != CurrentFrameID)
 			{
+				ParentFrame = FindExecutionFrame(CurrentFrameID, MakeArrayView(Exec.ActiveFrames), MakeArrayView(Args->TemporaryFrames));
 				CurrentFrameID = SelectedState.GetFrameID();
 				CurrentFrame = Exec.FindActiveFrame(CurrentFrameID);
 
@@ -2317,8 +2349,10 @@ void FStateTreeExecutionContext::UpdateInstanceData(const TSharedPtr<FSelectStat
 					{
 						// Point to the parameter block set by linked state.
 						check(NextStateParameterDataStruct == CurrentStateTree->GetDefaultParameters().GetPropertyBagStruct());
-						CurrentGlobalParameterDataHandle = NextStateParameterDataHandle;
 						NextStateParameterDataHandle = FStateTreeDataHandle::Invalid; // Mark as used.
+
+						check(ParentFrame);
+						CurrentGlobalParameterIndexBase = ParentFrame->ActiveInstanceIndexBase;
 					}
 
 					const int32 BaseIndex = InstanceStructs.Num();
@@ -2351,9 +2385,11 @@ void FStateTreeExecutionContext::UpdateInstanceData(const TSharedPtr<FSelectStat
 
 				CurrentFrameBaseIndex = InstanceStructs.Num();
 
-				CurrentFrame->GlobalParameterDataHandle = CurrentGlobalParameterDataHandle;
-				CurrentFrame->GlobalInstanceIndexBase = FStateTreeIndex16(CurrentGlobalInstanceIndexBase);
 				CurrentFrame->ActiveInstanceIndexBase = FStateTreeIndex16(CurrentFrameBaseIndex);
+
+				CurrentFrame->GlobalInstanceIndexBase = FStateTreeIndex16(CurrentGlobalInstanceIndexBase);
+				CurrentFrame->GlobalParameterIndexBase = CurrentGlobalParameterIndexBase;
+				CurrentFrame->StateParameterIndexBase = CurrentFrame->ActiveInstanceIndexBase;
 			}
 
 			// State Params and Nodes
@@ -2380,7 +2416,6 @@ void FStateTreeExecutionContext::UpdateInstanceData(const TSharedPtr<FSelectStat
 						{
 							// Parameters are not set by a linked state, create instance data.
 							InstanceStructs[CurrentFrameBaseIndex + CurrentState->ParameterDataHandle.GetIndex()] = ParamsInstanceData;
-							CurrentFrame->StateParameterDataHandle = CurrentState->ParameterDataHandle;
 							bCanHaveTempData = true;
 						}
 						else
@@ -2390,7 +2425,9 @@ void FStateTreeExecutionContext::UpdateInstanceData(const TSharedPtr<FSelectStat
 							const UStruct* StateParameterDataStruct = Params ? Params->Parameters.GetPropertyBagStruct() : nullptr;
 							check(NextStateParameterDataStruct == StateParameterDataStruct);
 
-							CurrentFrame->StateParameterDataHandle = NextStateParameterDataHandle;
+							check(ParentFrame);
+							CurrentFrame->StateParameterIndexBase = ParentFrame->ActiveInstanceIndexBase;
+
 							NextStateParameterDataHandle = FStateTreeDataHandle::Invalid; // Mark as used.
 
 							// This state will not instantiate parameter data, so we don't care about the temp data either.
@@ -2531,7 +2568,7 @@ FStateTreeDataView FStateTreeExecutionContext::MoveTemporaryToInstance(const FSt
 			return TempData ? &TempData->Instance : nullptr;
 		};
 
-	FStateTreeDataView Struct = UE::StateTree::InstanceData::GetDataView(Storage, nullptr, nullptr, CurrentFrame, Handle.ToStructSource());
+	FStateTreeDataView Struct = UE::StateTree::InstanceData::GetDataView(Storage, nullptr, CurrentFrame, Handle.ToStructSource());
 	FInstancedStruct* Source = FindInstanceTempData(CurrentFrame.FrameID, Handle);
 
 	if (Struct.IsValid())
@@ -2622,7 +2659,7 @@ FStateTreeDataView FStateTreeExecutionContext::GetDataView(const FStateTreeExecu
 
 		if (UE::StateTree::InstanceData::Private::IsActiveInstanceHandleSourceValid(Storage, CurrentFrame, Handle))
 		{
-			return UE::StateTree::InstanceData::GetDataView(Storage, CurrentlyProcessedSharedInstanceStorage, ParentFrame, CurrentFrame, Handle);
+			return UE::StateTree::InstanceData::GetDataView(Storage, CurrentlyProcessedSharedInstanceStorage, CurrentFrame, Handle);
 		}
 	}
 
@@ -2631,7 +2668,7 @@ FStateTreeDataView FStateTreeExecutionContext::GetDataView(const FStateTreeExecu
 		return nullptr;
 
 	default:
-		return UE::StateTree::InstanceData::GetDataView(Storage, CurrentlyProcessedSharedInstanceStorage, ParentFrame, CurrentFrame, Handle);
+		return UE::StateTree::InstanceData::GetDataView(Storage, CurrentlyProcessedSharedInstanceStorage, CurrentFrame, Handle);
 	}
 }
 
@@ -2947,7 +2984,7 @@ bool FStateTreeExecutionContext::IsHandleSourceValid(const FStateTreeExecutionFr
 	}
 
 	default:
-		return UE::StateTree::InstanceData::Private::IsHandleSourceValid(Storage, ParentFrame, CurrentFrame, Handle);
+		return UE::StateTree::InstanceData::Private::IsHandleSourceValid(Storage, CurrentFrame, Handle);
 	}
 
 	return false;
@@ -3001,7 +3038,7 @@ FStateTreeDataView FStateTreeExecutionContext::GetTemporaryDataView(const FState
 		return {};
 
 	default:
-		return UE::StateTree::InstanceData::Private::GetTemporaryDataView(Storage, ParentFrame, CurrentFrame, Handle);
+		return UE::StateTree::InstanceData::Private::GetTemporaryDataView(Storage, CurrentFrame, Handle);
 	}
 }
 
@@ -7363,25 +7400,28 @@ bool FStateTreeExecutionContext::SelectStateInternal_Linked(
 	const bool bIsNewFrame = SelectedFrame == nullptr;
 	if (bIsNewFrame)
 	{
-		// Note. Adding to TemporaryFrame can invalidate TargetFrame.
-		FStateTreeIndex16 ExternalDataBaseIndex;
-		FStateTreeDataHandle GlobalParameterDataHandle;
-		FStateTreeIndex16 GlobalInstanceIndexBase;
-		{
-			const FStateTreeExecutionFrame* NextFrame = FindExecutionFrame(OutSelectionResult->SelectedFrames.Last(), MakeConstArrayView(Exec.ActiveFrames), MakeConstArrayView(OutSelectionResult->TemporaryFrames));
-			check(NextFrame);
-			ExternalDataBaseIndex = NextFrame->ExternalDataBaseIndex;
-			GlobalParameterDataHandle = NextFrame->GlobalParameterDataHandle;
-			GlobalInstanceIndexBase = NextFrame->GlobalInstanceIndexBase;
-		}
-
 		constexpr bool bIsGlobalFrame = false;
 		FStateTreeExecutionFrame& NewFrame = OutSelectionResult->MakeAndAddTemporaryFrame(FActiveFrameID(Storage.GenerateUniqueId()), LinkStateFrameHandle, bIsGlobalFrame);
-		NewFrame.ExternalDataBaseIndex = ExternalDataBaseIndex;
-		NewFrame.GlobalInstanceIndexBase = GlobalInstanceIndexBase;
+
+		// Note. Adding to TemporaryFrame can invalidate ParentFrame.
+		const FStateTreeExecutionFrame* NextParentFrame = FindExecutionFrame(OutSelectionResult->SelectedFrames.Last(), MakeConstArrayView(Exec.ActiveFrames), MakeConstArrayView(OutSelectionResult->TemporaryFrames));
+		check(NextParentFrame);
+
+		NewFrame.ExternalDataBaseIndex = NextParentFrame->ExternalDataBaseIndex;
+
+		NewFrame.GlobalInstanceDataFrameID = NextParentFrame->GlobalInstanceDataFrameID;
+		NewFrame.GlobalInstanceIndexBase = NextParentFrame->GlobalInstanceIndexBase;
+
+		NewFrame.GlobalParameterDataHandle = NextParentFrame->GlobalParameterDataHandle;
+		NewFrame.GlobalParameterDataFrameID = NextParentFrame->GlobalParameterDataFrameID;
+		NewFrame.GlobalParameterIndexBase = NextParentFrame->GlobalParameterIndexBase;
+
+		// Points to the linked state's state parameter
+		NewFrame.StateParameterDataHandle = TargetState.ParameterDataHandle;
+		NewFrame.StateParameterDataFrameID = NextParentFrame->FrameID;
+		NewFrame.StateParameterIndexBase = NextParentFrame->ActiveInstanceIndexBase;
+
 		NewFrame.ExecutionRuntimeIndexBase = FStateTreeIndex16(Storage.AddExecutionRuntimeData(GetOwner(), LinkStateFrameHandle));
-		NewFrame.StateParameterDataHandle = TargetState.ParameterDataHandle; // Temporary allocated earlier if did not exists.
-		NewFrame.GlobalParameterDataHandle = GlobalParameterDataHandle;
 
 		SelectedFrame = &NewFrame;
 	}
@@ -7542,14 +7582,26 @@ bool FStateTreeExecutionContext::SelectStateInternal_LinkedAsset(
 			return false;
 		}
 
-		FStateTreeExecutionFrame& NewFrame = OutSelectionResult->MakeAndAddTemporaryFrame(FActiveFrameID(Storage.GenerateUniqueId()), LinkStateFrameHandle, true);
+		constexpr bool bIsGlobalFrame = true;
+		FStateTreeExecutionFrame& NewFrame = OutSelectionResult->MakeAndAddTemporaryFrame(FActiveFrameID(Storage.GenerateUniqueId()), LinkStateFrameHandle, bIsGlobalFrame);
+
+		// Adding to TemporaryFrame can invalidate ParentFrame.
+		const FStateTreeExecutionFrame* NextParentFrame = FindExecutionFrame(OutSelectionResult->SelectedFrames.Last(), MakeConstArrayView(Exec.ActiveFrames), MakeConstArrayView(OutSelectionResult->TemporaryFrames));
+		check(NextParentFrame);
+
 		NewFrame.ExternalDataBaseIndex = ExternalDataBaseIndex;
 		NewFrame.ExecutionRuntimeIndexBase = FStateTreeIndex16(Storage.AddExecutionRuntimeData(GetOwner(), LinkStateFrameHandle));
+
 		// Pass the linked state's parameters as global parameters to the linked asset.
 		NewFrame.GlobalParameterDataHandle = NextState.ParameterDataHandle;
-		// The state parameters will be from the root state.
-		const FCompactStateTreeState& NewFrameRootState = LinkStateFrameHandle.GetStateTree()->States[NewFrame.RootState.Index];
-		NewFrame.StateParameterDataHandle = NewFrameRootState.ParameterDataHandle;
+		NewFrame.GlobalParameterDataFrameID = NextParentFrame->FrameID;
+		NewFrame.GlobalParameterIndexBase = NextParentFrame->ActiveInstanceIndexBase;
+
+		// Global Frame
+		NewFrame.GlobalInstanceDataFrameID = NewFrame.FrameID;
+
+		NewFrame.StateParameterDataHandle = NewFrame.StateTree->GetStateFromHandle(NewFrame.RootState)->ParameterDataHandle;
+		NewFrame.StateParameterDataFrameID = NewFrame.FrameID;
 
 		SelectedFrame = &NewFrame;
 	}
@@ -7558,12 +7610,14 @@ bool FStateTreeExecutionContext::SelectStateInternal_LinkedAsset(
 
 	if (bIsNewFrame)
 	{
-		// Start global tasks and evaluators temporarily, so that their data is available already during select.
+		// Adding to TemporaryFrame can invalidate ParentFrame.
 		const FStateTreeExecutionFrame* NextParentFrame = FindExecutionFrame(OutSelectionResult->SelectedFrames.Last(1), MakeConstArrayView(Exec.ActiveFrames), MakeConstArrayView(OutSelectionResult->TemporaryFrames));
-		FStateTreeExecutionFrame* NextFrame = FindExecutionFrame(OutSelectionResult->SelectedFrames.Last(), MakeArrayView(Exec.ActiveFrames), MakeArrayView(OutSelectionResult->TemporaryFrames));
 		check(NextParentFrame);
-		check(SelectedFrame == NextFrame);
 
+		FStateTreeExecutionFrame* NextFrame = FindExecutionFrame(SelectedFrame->FrameID, MakeArrayView(Exec.ActiveFrames), MakeArrayView(OutSelectionResult->TemporaryFrames));
+		check(NextFrame && SelectedFrame == NextFrame);
+
+		// Start global tasks and evaluators temporarily, so that their data is available already during select.
 		const EStateTreeRunStatus StartResult = StartTemporaryEvaluatorsAndGlobalTasks(NextParentFrame, *NextFrame);
 		if (StartResult != EStateTreeRunStatus::Running)
 		{
@@ -7599,10 +7653,9 @@ bool FStateTreeExecutionContext::SelectStateInternal_LinkedAsset(
 	{
 		if (bIsNewFrame)
 		{
-			const FStateTreeExecutionFrame* NextParentFrame = FindExecutionFrame(OutSelectionResult->SelectedFrames.Last(1), MakeConstArrayView(Exec.ActiveFrames), MakeConstArrayView(OutSelectionResult->TemporaryFrames));
+			// Adding to TemporaryFrame can invalidate ParentFrame
+			const FStateTreeExecutionFrame* NextParentFrame = FindExecutionFrame(OutSelectionResult->SelectedFrames.Last(1), MakeArrayView(Exec.ActiveFrames), MakeArrayView(OutSelectionResult->TemporaryFrames));
 			FStateTreeExecutionFrame* NextFrame = FindExecutionFrame(OutSelectionResult->SelectedFrames.Last(), MakeArrayView(Exec.ActiveFrames), MakeArrayView(OutSelectionResult->TemporaryFrames));
-			check(NextParentFrame);
-			check(NextFrame);
 
 			constexpr EStateTreeRunStatus CompletionStatus = EStateTreeRunStatus::Stopped;
 			StopTemporaryEvaluatorsAndGlobalTasks(NextParentFrame, *NextFrame, CompletionStatus);
diff --git a/Engine/Plugins/Runtime/StateTree/Source/StateTreeModule/Private/StateTreeInstanceData.cpp b/Engine/Plugins/Runtime/StateTree/Source/StateTreeModule/Private/StateTreeInstanceData.cpp
--- a/Engine/Plugins/Runtime/StateTree/Source/StateTreeModule/Private/StateTreeInstanceData.cpp
+++ b/Engine/Plugins/Runtime/StateTree/Source/StateTreeModule/Private/StateTreeInstanceData.cpp
@@ -160,7 +160,6 @@ namespace UE::StateTree::InstanceData::Private
 
 	bool IsHandleSourceValid(
 		const FStateTreeInstanceStorage& InstanceStorage,
-		const FStateTreeExecutionFrame* ParentFrame,
 		const FStateTreeExecutionFrame& CurrentFrame,
 		const FStateTreeDataHandle& Handle)
 	{
@@ -194,23 +193,40 @@ namespace UE::StateTree::InstanceData::Private
 			return true;
 
 		case EStateTreeDataSourceType::GlobalParameterData:
-			return ParentFrame
-				? IsHandleSourceValid(InstanceStorage, nullptr, *ParentFrame, CurrentFrame.GlobalParameterDataHandle)
-				: CurrentFrame.GlobalParameterDataHandle.IsValid();
-
-		case EStateTreeDataSourceType::SubtreeParameterData:
-			if (ParentFrame)
+		{
+			// Defined as State Parameter in one of its parent frame
+			if (CurrentFrame.GlobalParameterDataHandle.HasValidIndex())
 			{
-				// If the current subtree state is not instantiated yet, we cannot assume that the parameter data is instantiated in the parent frame either. 
-				if (!CurrentFrame.ActiveInstanceIndexBase.IsValid())
+				check(CurrentFrame.GlobalParameterDataHandle.GetSource() == EStateTreeDataSourceType::StateParameterData);
+
+				if (const FStateTreeExecutionFrame* ActiveGlobalParameterFrame = InstanceStorage.GetExecutionState().FindActiveFrame(CurrentFrame.GlobalParameterDataFrameID))
 				{
-					return false;
+					check(ActiveGlobalParameterFrame->ActiveInstanceIndexBase == CurrentFrame.GlobalParameterIndexBase);
+
+					return IsActiveInstanceHandleSourceValid(InstanceStorage, *ActiveGlobalParameterFrame, CurrentFrame.GlobalParameterDataHandle);
 				}
-				// Linked subtree, params defined in parent scope.
-				return IsHandleSourceValid(InstanceStorage, nullptr, *ParentFrame, CurrentFrame.StateParameterDataHandle);
+
+				// Parent Frame is temporary, the state parameter can't be active.
+				return false;
 			}
-			// Standalone subtree, params define as state params.
-			return IsActiveInstanceHandleSourceValid(InstanceStorage, CurrentFrame, Handle);
+
+			// Root Global Parameter should always be valid.
+			check(CurrentFrame.GlobalParameterDataHandle.GetSource() == EStateTreeDataSourceType::GlobalParameterData);
+			return true;
+		}
+
+		case EStateTreeDataSourceType::SubtreeParameterData:
+		{
+			if (const FStateTreeExecutionFrame* ActiveParameterFrame = InstanceStorage.GetExecutionState().FindActiveFrame(CurrentFrame.StateParameterDataFrameID))
+			{
+				check(ActiveParameterFrame->ActiveInstanceIndexBase == CurrentFrame.StateParameterIndexBase);
+
+				return IsActiveInstanceHandleSourceValid(InstanceStorage, *ActiveParameterFrame, CurrentFrame.StateParameterDataHandle);
+			}
+
+			// Frame hosting the parameter is not active
+			return false;
+		}
 
 		default:
 			checkf(false, TEXT("Unhandled case or unsupported type for InstanceDataStorage %s"), *UEnum::GetValueAsString(Handle.GetSource()));
@@ -221,52 +237,66 @@ namespace UE::StateTree::InstanceData::Private
 
 	FStateTreeDataView GetTemporaryDataView(
 		FStateTreeInstanceStorage& InstanceStorage,
-		const FStateTreeExecutionFrame* ParentFrame,
 		const FStateTreeExecutionFrame& CurrentFrame,
 		const FStateTreeDataHandle& Handle)
 	{
 		switch (Handle.GetSource())
 		{
 		case EStateTreeDataSourceType::GlobalInstanceData:
+			return InstanceStorage.GetMutableTemporaryStruct(CurrentFrame.GlobalInstanceDataFrameID, Handle);
 		case EStateTreeDataSourceType::ActiveInstanceData:
-			return InstanceStorage.GetMutableTemporaryStruct(CurrentFrame, Handle);
-
+			return InstanceStorage.GetMutableTemporaryStruct(CurrentFrame.FrameID, Handle);
+			
 		case EStateTreeDataSourceType::GlobalInstanceDataObject:
+			return InstanceStorage.GetMutableTemporaryObject(CurrentFrame.GlobalInstanceDataFrameID, Handle);
 		case EStateTreeDataSourceType::ActiveInstanceDataObject:
-			return InstanceStorage.GetMutableTemporaryObject(CurrentFrame, Handle);
+			return InstanceStorage.GetMutableTemporaryObject(CurrentFrame.FrameID, Handle);
 
 		case EStateTreeDataSourceType::GlobalParameterData:
-			if (ParentFrame)
+		{
+			// Defined as State Parameter in one of its parent frame
+			if (CurrentFrame.GlobalParameterDataHandle.HasValidIndex())
 			{
-				if (FCompactStateTreeParameters* Params = InstanceStorage.GetMutableTemporaryStruct(*ParentFrame, CurrentFrame.GlobalParameterDataHandle).GetPtr<FCompactStateTreeParameters>())
+				check(CurrentFrame.GlobalParameterDataHandle.GetSource() == EStateTreeDataSourceType::StateParameterData);
+
+				if (FCompactStateTreeParameters* Params = InstanceStorage.GetMutableTemporaryStruct(CurrentFrame.GlobalParameterDataFrameID, CurrentFrame.GlobalParameterDataHandle).GetPtr<FCompactStateTreeParameters>())
 				{
 					return Params->Parameters.GetMutableValue();
 				}
 			}
+
+			// Root global parameter is never stored in temporary
 			break;
+		}
 
 		case EStateTreeDataSourceType::SubtreeParameterData:
-			if (ParentFrame)
-			{
-				// Linked subtree, params defined in parent scope.
-				if (FCompactStateTreeParameters* Params = InstanceStorage.GetMutableTemporaryStruct(*ParentFrame, CurrentFrame.StateParameterDataHandle).GetPtr<FCompactStateTreeParameters>())
-				{
-					return Params->Parameters.GetMutableValue();
-				}
-			}
-			// Standalone subtree, params define as state params.
-			if (FCompactStateTreeParameters* Params = InstanceStorage.GetMutableTemporaryStruct(CurrentFrame, Handle).GetPtr<FCompactStateTreeParameters>())
+		{
+			check(Handle == CurrentFrame.StateTree->GetStateFromHandle(CurrentFrame.RootState)->ParameterDataHandle);
+
+			// A subtree could be either linked or standalone.
+			check((CurrentFrame.StateParameterDataHandle.GetSource() == EStateTreeDataSourceType::SubtreeParameterData 
+					&& CurrentFrame.FrameID == CurrentFrame.StateParameterDataFrameID
+					&& CurrentFrame.StateParameterDataHandle == Handle)
+				|| (CurrentFrame.StateParameterDataHandle.GetSource() == EStateTreeDataSourceType::StateParameterData 
+					&& CurrentFrame.FrameID != CurrentFrame.StateParameterDataFrameID));
+
+			if (FCompactStateTreeParameters* Params = InstanceStorage.GetMutableTemporaryStruct(CurrentFrame.StateParameterDataFrameID, CurrentFrame.StateParameterDataHandle).GetPtr<FCompactStateTreeParameters>())
 			{
 				return Params->Parameters.GetMutableValue();
 			}
+
 			break;
+		}
 
 		case EStateTreeDataSourceType::StateParameterData:
-			if (FCompactStateTreeParameters* Params = InstanceStorage.GetMutableTemporaryStruct(CurrentFrame, Handle).GetPtr<FCompactStateTreeParameters>())
+		{
+			if (FCompactStateTreeParameters* Params = InstanceStorage.GetMutableTemporaryStruct(CurrentFrame.FrameID, Handle).GetPtr<FCompactStateTreeParameters>())
 			{
 				return Params->Parameters.GetMutableValue();
 			}
+
 			break;
+		}
 
 		default:
 			checkf(false, TEXT("Unhandled case or unsupported type for InstanceDataStorage %s"), *UEnum::GetValueAsString(Handle.GetSource()));
@@ -282,7 +312,6 @@ namespace UE::StateTree::InstanceData
 	FStateTreeDataView GetDataView(
 		FStateTreeInstanceStorage& InstanceStorage,
 		FStateTreeInstanceStorage* SharedInstanceStorage,
-		const FStateTreeExecutionFrame* ParentFrame,
 		const FStateTreeExecutionFrame& CurrentFrame,
 		const FStateTreeDataHandle& Handle)
 	{
@@ -306,24 +335,35 @@ namespace UE::StateTree::InstanceData
 			return SharedInstanceStorage->GetMutableObject(Handle.GetIndex());
 
 		case EStateTreeDataSourceType::GlobalParameterData:
-			// Defined in parent frame or is root state tree parameters
-			if (ParentFrame)
+		{
+			if (CurrentFrame.GlobalParameterDataHandle.HasValidIndex())
 			{
-				return GetDataView(InstanceStorage, SharedInstanceStorage, nullptr, *ParentFrame, CurrentFrame.GlobalParameterDataHandle);
+				// Defined as State Parameter in one of its parent frame, we must have a valid cached Index base for active instances
+				check(CurrentFrame.GlobalParameterIndexBase.IsValid());
+
+				FCompactStateTreeParameters& StateParams = InstanceStorage.GetMutableStruct(CurrentFrame.GlobalParameterIndexBase.Get() + CurrentFrame.GlobalParameterDataHandle.GetIndex()).Get<FCompactStateTreeParameters>();
+				return StateParams.Parameters.GetMutableValue();
 			}
+
+			// Root Global Parameters
 			return InstanceStorage.GetMutableGlobalParameters();
+		}
 
 		case EStateTreeDataSourceType::SubtreeParameterData:
 		{
-			// Defined in parent frame.
-			if (ParentFrame)
-			{
-				// Linked subtree, params defined in parent scope.
-				return GetDataView(InstanceStorage, SharedInstanceStorage, nullptr, *ParentFrame, CurrentFrame.StateParameterDataHandle);
-			}
-			// Standalone subtree, params define as state params.
-			FCompactStateTreeParameters& SubtreeParams = InstanceStorage.GetMutableStruct(CurrentFrame.ActiveInstanceIndexBase.Get() + Handle.GetIndex()).Get<FCompactStateTreeParameters>();
-			return SubtreeParams.Parameters.GetMutableValue();
+			check(CurrentFrame.StateParameterIndexBase.IsValid() && CurrentFrame.StateParameterDataHandle.HasValidIndex());
+
+			check(Handle == CurrentFrame.StateTree->GetStateFromHandle(CurrentFrame.RootState)->ParameterDataHandle);
+
+			// A subtree could be either linked or standalone.
+			check((CurrentFrame.StateParameterDataHandle.GetSource() == EStateTreeDataSourceType::SubtreeParameterData
+					&& CurrentFrame.FrameID == CurrentFrame.StateParameterDataFrameID
+					&& CurrentFrame.StateParameterDataHandle == Handle)
+				|| (CurrentFrame.StateParameterDataHandle.GetSource() == EStateTreeDataSourceType::StateParameterData
+					&& CurrentFrame.FrameID != CurrentFrame.StateParameterDataFrameID));
+
+			FCompactStateTreeParameters& StateParams = InstanceStorage.GetMutableStruct(CurrentFrame.StateParameterIndexBase.Get() + CurrentFrame.StateParameterDataHandle.GetIndex()).Get<FCompactStateTreeParameters>();
+			return StateParams.Parameters.GetMutableValue();
 		}
 
 		case EStateTreeDataSourceType::StateParameterData:
@@ -354,16 +394,15 @@ namespace UE::StateTree::InstanceData
 	FStateTreeDataView GetDataViewOrTemporary(
 		FStateTreeInstanceStorage& InstanceStorage,
 		FStateTreeInstanceStorage* SharedInstanceStorage,
-		const FStateTreeExecutionFrame* ParentFrame,
 		const FStateTreeExecutionFrame& CurrentFrame,
 		const FStateTreeDataHandle& Handle)
 	{
-		if (Private::IsHandleSourceValid(InstanceStorage, ParentFrame, CurrentFrame, Handle))
+		if (Private::IsHandleSourceValid(InstanceStorage, CurrentFrame, Handle))
 		{
-			return GetDataView(InstanceStorage, SharedInstanceStorage, ParentFrame, CurrentFrame, Handle);
+			return GetDataView(InstanceStorage, SharedInstanceStorage, CurrentFrame, Handle);
 		}
 
-		return Private::GetTemporaryDataView(InstanceStorage, ParentFrame, CurrentFrame, Handle);
+		return Private::GetTemporaryDataView(InstanceStorage, CurrentFrame, Handle);
 	}
 } // namespace
 
@@ -575,23 +614,23 @@ FStructView FStateTreeInstanceStorage::AddTemporaryInstance(UObject& InOwner, co
 	return TempInstance->Instance;
 }
 
-FStructView FStateTreeInstanceStorage::GetMutableTemporaryStruct(const FStateTreeExecutionFrame& Frame, const FStateTreeDataHandle DataHandle)
+FStructView FStateTreeInstanceStorage::GetMutableTemporaryStruct(UE::StateTree::FActiveFrameID FrameID, const FStateTreeDataHandle DataHandle)
 {
-	FStateTreeTemporaryInstanceData* ExistingInstance = TemporaryInstances.FindByPredicate([&Frame, &DataHandle](const FStateTreeTemporaryInstanceData& TempInstance)
-	{
-		return TempInstance.FrameID == Frame.FrameID
+	FStateTreeTemporaryInstanceData* ExistingInstance = TemporaryInstances.FindByPredicate([FrameID, &DataHandle](const FStateTreeTemporaryInstanceData& TempInstance)
+		{
+			return TempInstance.FrameID == FrameID
 				&& TempInstance.DataHandle == DataHandle;
-	});
+		});
 	return ExistingInstance ? FStructView(ExistingInstance->Instance) : FStructView();
 }
 
-UObject* FStateTreeInstanceStorage::GetMutableTemporaryObject(const FStateTreeExecutionFrame& Frame, const FStateTreeDataHandle DataHandle)
+UObject* FStateTreeInstanceStorage::GetMutableTemporaryObject(UE::StateTree::FActiveFrameID FrameID, const FStateTreeDataHandle DataHandle)
 {
-	FStateTreeTemporaryInstanceData* ExistingInstance = TemporaryInstances.FindByPredicate([&Frame, &DataHandle](const FStateTreeTemporaryInstanceData& TempInstance)
-	{
-		return TempInstance.FrameID == Frame.FrameID
+	FStateTreeTemporaryInstanceData* ExistingInstance = TemporaryInstances.FindByPredicate([FrameID, &DataHandle](const FStateTreeTemporaryInstanceData& TempInstance)
+		{
+			return TempInstance.FrameID == FrameID
 				&& TempInstance.DataHandle == DataHandle;
-	});
+		});
 	if (ExistingInstance)
 	{
 		const FStateTreeInstanceObjectWrapper& Wrapper = ExistingInstance->Instance.Get<FStateTreeInstanceObjectWrapper>();
diff --git a/Engine/Plugins/Runtime/StateTree/Source/StateTreeModule/Private/StateTreeInstanceDataHelpers.h b/Engine/Plugins/Runtime/StateTree/Source/StateTreeModule/Private/StateTreeInstanceDataHelpers.h
--- a/Engine/Plugins/Runtime/StateTree/Source/StateTreeModule/Private/StateTreeInstanceDataHelpers.h
+++ b/Engine/Plugins/Runtime/StateTree/Source/StateTreeModule/Private/StateTreeInstanceDataHelpers.h
@@ -15,14 +15,12 @@ namespace UE::StateTree::InstanceData::Private
 	/** @return whether the handle is valid relative to the given frame. */
 	[[nodiscard]] bool IsHandleSourceValid(
 		const FStateTreeInstanceStorage& InstanceStorage,
-		const FStateTreeExecutionFrame* ParentFrame,
 		const FStateTreeExecutionFrame& CurrentFrame,
 		const FStateTreeDataHandle& Handle);
 
 	/** @return data view of the specified handle in temporary instance. */
 	[[nodiscard]] FStateTreeDataView GetTemporaryDataView(
 		FStateTreeInstanceStorage& InstanceStorage,
-		const FStateTreeExecutionFrame* ParentFrame,
 		const FStateTreeExecutionFrame& CurrentFrame,
 		const FStateTreeDataHandle& Handle);
 
diff --git a/Engine/Plugins/Runtime/StateTree/Source/StateTreeModule/Public/StateTreeExecutionContext.h b/Engine/Plugins/Runtime/StateTree/Source/StateTreeModule/Public/StateTreeExecutionContext.h
--- a/Engine/Plugins/Runtime/StateTree/Source/StateTreeModule/Public/StateTreeExecutionContext.h
+++ b/Engine/Plugins/Runtime/StateTree/Source/StateTreeModule/Public/StateTreeExecutionContext.h
@@ -772,7 +772,7 @@ PRAGMA_ENABLE_DEPRECATION_WARNINGS
 	UE_DEPRECATED(5.6, "Use UE::StateTree::InstanceData::GetDataView instead.")
 	static FStateTreeDataView GetDataViewFromInstanceStorage(FStateTreeInstanceStorage& InstanceDataStorage, FStateTreeInstanceStorage* CurrentlyProcessedSharedInstanceStorage, const FStateTreeExecutionFrame* ParentFrame, const FStateTreeExecutionFrame& CurrentFrame, const FStateTreeDataHandle Handle)
 	{
-		return UE::StateTree::InstanceData::GetDataView(InstanceDataStorage, CurrentlyProcessedSharedInstanceStorage, ParentFrame, CurrentFrame, Handle);
+		return UE::StateTree::InstanceData::GetDataView(InstanceDataStorage, CurrentlyProcessedSharedInstanceStorage, CurrentFrame, Handle);
 	}
 
 	//@todo deprecate in favor of a find with frameID
diff --git a/Engine/Plugins/Runtime/StateTree/Source/StateTreeModule/Public/StateTreeExecutionTypes.h b/Engine/Plugins/Runtime/StateTree/Source/StateTreeModule/Public/StateTreeExecutionTypes.h
--- a/Engine/Plugins/Runtime/StateTree/Source/StateTreeModule/Public/StateTreeExecutionTypes.h
+++ b/Engine/Plugins/Runtime/StateTree/Source/StateTreeModule/Public/StateTreeExecutionTypes.h
@@ -1020,21 +1020,46 @@ PRAGMA_ENABLE_DEPRECATION_WARNINGS
 	UPROPERTY()
 	TObjectPtr<const UStateTree> StateTree = nullptr;
 
-	/** The root state of the frame (e.g. Root state or a subtree). */
-	UPROPERTY()
-	FStateTreeStateHandle RootState = FStateTreeStateHandle::Root;
-	
+	/** Flag to track the completion of a global task or a task from a state in the ActiveStates. */
+	UPROPERTY(Transient)
+	FStateTreeTasksCompletionStatus ActiveTasksStatus;
+
 	/** Active states in this frame */
 	UPROPERTY()
 	FStateTreeActiveStates ActiveStates;
 
-	/** Flag to track the completion of a global task or a task from a state in the ActiveStates. */
-	UPROPERTY(Transient)
-	FStateTreeTasksCompletionStatus ActiveTasksStatus;
+	/**
+	 * Handle to the root state parameter data.
+	 * Exists in parent frame if this frame is entered by a linked state,
+	 * Otherwise it exists in the current frame.
+	 */
+	UPROPERTY()
+	FStateTreeDataHandle StateParameterDataHandle = FStateTreeDataHandle::Invalid;
+
+	/**
+	 * Handle to the global parameter data,
+	 * Exists in parent frame if it is a state parameter.
+	 * Exists in root frame if it is the root global parameter.
+	 */
+	UPROPERTY()
+	FStateTreeDataHandle GlobalParameterDataHandle = FStateTreeDataHandle::Invalid;
 
 	/** Unique frame ID for this frame. Can be used to identify the frame. */
 	UE::StateTree::FActiveFrameID FrameID;
 
+	/** current or parent frame ID that holds global instance data this frame can have access to. */
+	UE::StateTree::FActiveFrameID GlobalInstanceDataFrameID;
+
+	/** current or parent frame ID that holds global parameter data this frame can have access to. */
+	UE::StateTree::FActiveFrameID GlobalParameterDataFrameID;
+
+	/** current or parent frame ID that holds root state parameter this frame can have access to. */
+	UE::StateTree::FActiveFrameID StateParameterDataFrameID;
+
+	/** The root state of the frame (e.g. Root state or a subtree). */
+	UPROPERTY()
+	FStateTreeStateHandle RootState = FStateTreeStateHandle::Root;
+
 	/**
 	 * The evaluator or task node index that was "entered".
 	 * Used during the Enter and Exit phase. A node can fail EnterState.
@@ -1050,6 +1075,22 @@ PRAGMA_ENABLE_DEPRECATION_WARNINGS
 	UPROPERTY()
 	FStateTreeIndex16 GlobalInstanceIndexBase = FStateTreeIndex16::Invalid;
 
+	/**
+	 * Index within the instance data that GlobalParameterDataHandle is relative to,
+	 * if it is state parameter it lives within the parent frame
+	 * Invalid means GlobalParameterDataHandle points to the root global parameter.
+	 */
+	UPROPERTY()
+	FStateTreeIndex16 GlobalParameterIndexBase = FStateTreeIndex16::Invalid;
+
+	/**
+	 * Index within the instance data that StateParameterDataHandle is relative to,
+	 * if it lives as state parameter in the parent frame
+	 * Invalid means GlobalParameterDataHandle points to the root global parameter.
+	 */
+	UPROPERTY()
+	FStateTreeIndex16 StateParameterIndexBase = FStateTreeIndex16::Invalid;
+
 	/** Index within the instance data to the first active state's instance data (e.g. tasks) */
 	UPROPERTY()
 	FStateTreeIndex16 ActiveInstanceIndexBase = FStateTreeIndex16::Invalid;
@@ -1058,14 +1099,6 @@ PRAGMA_ENABLE_DEPRECATION_WARNINGS
 	UPROPERTY()
 	FStateTreeIndex16 ExecutionRuntimeIndexBase = FStateTreeIndex16::Invalid;
 
-	/** Handle to the state parameter data, exists in ParentFrame. */
-	UPROPERTY()
-	FStateTreeDataHandle StateParameterDataHandle = FStateTreeDataHandle::Invalid; 
-
-	/** Handle to the global parameter data, exists in ParentFrame. */
-	UPROPERTY()
-	FStateTreeDataHandle GlobalParameterDataHandle = FStateTreeDataHandle::Invalid;
-	
 	/** If true, the global tasks of the State Tree should be handle in this frame. */
 	UPROPERTY()
 	uint8 bIsGlobalFrame : 1 = false;
diff --git a/Engine/Plugins/Runtime/StateTree/Source/StateTreeModule/Public/StateTreeInstanceData.h b/Engine/Plugins/Runtime/StateTree/Source/StateTreeModule/Public/StateTreeInstanceData.h
--- a/Engine/Plugins/Runtime/StateTree/Source/StateTreeModule/Public/StateTreeInstanceData.h
+++ b/Engine/Plugins/Runtime/StateTree/Source/StateTreeModule/Public/StateTreeInstanceData.h
@@ -33,17 +33,37 @@ namespace UE::StateTree::InstanceData
 	[[nodiscard]] UE_API FStateTreeDataView GetDataView(
 		FStateTreeInstanceStorage& InstanceStorage,
 		FStateTreeInstanceStorage* SharedInstanceStorage,
-		const FStateTreeExecutionFrame* ParentFrame,
 		const FStateTreeExecutionFrame& CurrentFrame,
 		const FStateTreeDataHandle& Handle);
 
+	UE_DEPRECATED(5.8, "Use the version without Parent Frame instead")
+	[[nodiscard]] inline FStateTreeDataView GetDataView(
+		FStateTreeInstanceStorage& InstanceStorage,
+		FStateTreeInstanceStorage* SharedInstanceStorage,
+		const FStateTreeExecutionFrame* ParentFrame,
+		const FStateTreeExecutionFrame& CurrentFrame,
+		const FStateTreeDataHandle& Handle)
+	{
+		return GetDataView(InstanceStorage, SharedInstanceStorage, CurrentFrame, Handle);
+	}
+
 	/** @return data view of the specified handle relative to the given frame, or tries to find a matching temporary instance. */
 	[[nodiscard]] UE_API FStateTreeDataView GetDataViewOrTemporary(
 		FStateTreeInstanceStorage& InstanceStorage,
 		FStateTreeInstanceStorage* SharedInstanceStorage,
-		const FStateTreeExecutionFrame* ParentFrame,
 		const FStateTreeExecutionFrame& CurrentFrame,
 		const FStateTreeDataHandle& Handle);
+
+	UE_DEPRECATED(5.8, "Use the version without Parent Frame instead")
+	[[nodiscard]] inline FStateTreeDataView GetDataViewOrTemporary(
+		FStateTreeInstanceStorage& InstanceStorage,
+		FStateTreeInstanceStorage* SharedInstanceStorage,
+		const FStateTreeExecutionFrame* ParentFrame,
+		const FStateTreeExecutionFrame& CurrentFrame,
+		const FStateTreeDataHandle& Handle)
+	{
+		return GetDataViewOrTemporary(InstanceStorage, SharedInstanceStorage, CurrentFrame, Handle);
+	}
 }
 
 /**
@@ -269,12 +289,18 @@ struct FStateTreeInstanceStorage
 	 * @returns mutable struct view to the instance.
 	 */
 	UE_API FStructView AddTemporaryInstance(UObject& InOwner, const FStateTreeExecutionFrame& Frame, const FStateTreeIndex16 OwnerNodeIndex, const FStateTreeDataHandle DataHandle, FConstStructView NewInstanceData);
-	
+
+	UE_DEPRECATED(5.7, "Use the version with FrameID instead.")
+	FStructView GetMutableTemporaryStruct(const FStateTreeExecutionFrame& Frame, const FStateTreeDataHandle DataHandle) { return GetMutableTemporaryStruct(Frame.FrameID, DataHandle); }
+
 	/** @returns mutable view to the specified instance data, or invalid view if not found. */
-	UE_API FStructView GetMutableTemporaryStruct(const FStateTreeExecutionFrame& Frame, const FStateTreeDataHandle DataHandle);
+	UE_API FStructView GetMutableTemporaryStruct(UE::StateTree::FActiveFrameID FrameID, const FStateTreeDataHandle DataHandle);
+
+	UE_DEPRECATED(5.7, "Use the version with FrameID instead.")
+	UObject* GetMutableTemporaryObject(const FStateTreeExecutionFrame& Frame, const FStateTreeDataHandle DataHandle) { return GetMutableTemporaryObject(Frame.FrameID, DataHandle); }
 
 	/** @returns mutable pointer to the specified instance data object, or invalid view if not found. Will check() if called on non-object data. */
-	UE_API UObject* GetMutableTemporaryObject(const FStateTreeExecutionFrame& Frame, const FStateTreeDataHandle DataHandle);
+	UE_API UObject* GetMutableTemporaryObject(UE::StateTree::FActiveFrameID FrameID, const FStateTreeDataHandle DataHandle);
 
 	/** Empties the temporary instances. */
 	UE_API void ResetTemporaryInstances();
@@ -552,13 +578,13 @@ struct FStateTreeInstanceData
 	/** @returns mutable view to the specified instance data, or invalid view if not found. */
 	FStructView GetMutableTemporaryStruct(const FStateTreeExecutionFrame& Frame, const FStateTreeDataHandle DataHandle)
 	{
-		return GetMutableStorage().GetMutableTemporaryStruct(Frame, DataHandle);
+		return GetMutableStorage().GetMutableTemporaryStruct(Frame.FrameID, DataHandle);
 	}
 
 	/** @returns mutable pointer to the specified instance data object, or invalid view if not found. Will check() if called on non-object data. */
 	UObject* GetMutableTemporaryObject(const FStateTreeExecutionFrame& Frame, const FStateTreeDataHandle DataHandle)
 	{
-		return GetMutableStorage().GetMutableTemporaryObject(Frame, DataHandle);
+		return GetMutableStorage().GetMutableTemporaryObject(Frame.FrameID, DataHandle);
 	}
 
 	/** Empties the temporary instances. */
@@ -681,7 +707,7 @@ PRAGMA_ENABLE_DEPRECATION_WARNINGS
 			}
 			else
 			{
-				Struct = Storage.GetMutableTemporaryStruct(*CurrentFrame, DataHandle);
+				Struct = Storage.GetMutableTemporaryStruct(CurrentFrame->FrameID, DataHandle);
 			}
 
 			check(Struct.GetScriptStruct() == TBaseStructure<T>::Get());
diff --git a/Engine/Plugins/Runtime/StateTree/Source/StateTreeModule/Public/StateTreePropertyRef.h b/Engine/Plugins/Runtime/StateTree/Source/StateTreeModule/Public/StateTreePropertyRef.h
--- a/Engine/Plugins/Runtime/StateTree/Source/StateTreeModule/Public/StateTreePropertyRef.h
+++ b/Engine/Plugins/Runtime/StateTree/Source/StateTreeModule/Public/StateTreePropertyRef.h
@@ -28,7 +28,7 @@ namespace UE::StateTree::PropertyRefHelpers
 		const FStateTreePropertyBindings& PropertyBindings = ExecutionFrame.StateTree->GetPropertyBindings();
 		if (const FStateTreePropertyAccess* PropertyAccess = PropertyBindings.GetPropertyAccess(PropertyRef))
 		{
-			const FStateTreeDataView SourceView = InstanceData::GetDataViewOrTemporary(InstanceDataStorage, nullptr, ParentExecutionFrame, ExecutionFrame, PropertyAccess->SourceDataHandle);
+			const FStateTreeDataView SourceView = InstanceData::GetDataViewOrTemporary(InstanceDataStorage, nullptr, ExecutionFrame, PropertyAccess->SourceDataHandle);
 			
 			// The only possibility when PropertyRef references another PropertyRef is when source one is a global or subtree parameter, i.e lives in parent execution frame.
 			// If that's the case, referenced PropertyRef is obtained and we recursively take the address where it points to.
@@ -88,7 +88,7 @@ namespace UE::StateTree::PropertyRefHelpers
 		if (const FStateTreePropertyAccess* PropertyAccess = PropertyBindings.GetPropertyAccess(PropertyRef))
 		{
 			// Passing empty ContextAndExternalDataViews, as PropertyRef is not allowed to point to context or external data.
-			const FStateTreeDataView SourceView = InstanceData::GetDataViewOrTemporary(InstanceDataStorage, nullptr, ParentExecutionFrame, ExecutionFrame, PropertyAccess->SourceDataHandle);
+			const FStateTreeDataView SourceView = InstanceData::GetDataViewOrTemporary(InstanceDataStorage, nullptr, ExecutionFrame, PropertyAccess->SourceDataHandle);
 			return TTuple<T*...>(PropertyBindings.GetMutablePropertyPtr<T>(SourceView, *PropertyAccess)...);
 		}
 
diff --git a/Engine/Plugins/Runtime/StateTree/Source/StateTreeModule/Public/StateTreeTypes.h b/Engine/Plugins/Runtime/StateTree/Source/StateTreeModule/Public/StateTreeTypes.h
--- a/Engine/Plugins/Runtime/StateTree/Source/StateTreeModule/Public/StateTreeTypes.h
+++ b/Engine/Plugins/Runtime/StateTree/Source/StateTreeModule/Public/StateTreeTypes.h
@@ -380,6 +380,11 @@ struct FStateTreeDataHandle
 		return Index >= 0 && Index < (int32)InvalidIndex;
 	}
 
+	bool HasValidIndex() const
+	{
+		return IsValidIndex(Index);
+	}
+
 	friend inline uint32 GetTypeHash(const FStateTreeDataHandle& Handle)
 	{
 		uint32 Hash = GetTypeHash(Handle.Index);
diff --git a/Engine/Plugins/Runtime/StateTree/Source/StateTreeTestSuite/Private/StateTreeBindingTest.cpp b/Engine/Plugins/Runtime/StateTree/Source/StateTreeTestSuite/Private/StateTreeBindingTest.cpp
--- a/Engine/Plugins/Runtime/StateTree/Source/StateTreeTestSuite/Private/StateTreeBindingTest.cpp
+++ b/Engine/Plugins/Runtime/StateTree/Source/StateTreeTestSuite/Private/StateTreeBindingTest.cpp
@@ -808,8 +808,6 @@ struct FStateTreeTest_CopyObjects : FStateTreeTestBase
 };
 IMPLEMENT_STATE_TREE_INSTANT_TEST(FStateTreeTest_CopyObjects, "System.StateTree.Binding.CopyObjects");
 
-//@todo: Add test coverage for propertyrefs across frames
-
 struct FStateTreeTest_References : FStateTreeTestBase
 {
 	virtual bool InstantTest() override
@@ -1463,8 +1461,6 @@ struct FStateTreeTest_OutputBinding : FStateTreeTestBase
 		//	Root: OutputTaskB
 
 		FStateTreeCompilerLog Log;
-		FStateTreePropertyBindings Bindings;
-		FStateTreePropertyBindingCompiler BindingCompiler;
 
 		int32 Value = 0;
 		UStateTree& StateTree = NewStateTree();
@@ -1635,6 +1631,320 @@ struct FStateTreeTest_OutputBinding : FStateTreeTestBase
 	}
 };
 IMPLEMENT_STATE_TREE_INSTANT_TEST(FStateTreeTest_OutputBinding, "System.StateTree.Binding.OutputBinding");
+
+struct FStateTreeTest_SubtreeParameter : FStateTreeTestBase
+{
+	virtual bool InstantTest() override
+	{
+		// Root
+		//   LinkedState -> SubtreeA
+		// SubtreeA (EnterConditionA, TaskA)
+
+		auto AddIntPropertyToPropertyBag = [](FInstancedPropertyBag& PropertyBag, FName VarName, int32 DefaultValue)
+			{
+				PropertyBag.AddProperty(VarName, EPropertyBagPropertyType::Int32);
+				PropertyBag.SetValueInt32(VarName, DefaultValue);
+				return PropertyBag.FindPropertyDescByName(VarName)->ID;
+			};
+
+		UStateTree& StateTree = NewStateTree();
+		{
+			UStateTreeEditorData& EditorData = *Cast<UStateTreeEditorData>(StateTree.EditorData);
+			{
+				auto& RootState = EditorData.AddSubTree(TEXT("Root"));
+
+				FGuid SubtreeAIntID;
+				auto& SubtreeA = EditorData.AddSubTree(TEXT("SubtreeA"));
+				{
+					SubtreeA.Type = EStateTreeStateType::Subtree;
+					SubtreeAIntID = AddIntPropertyToPropertyBag(SubtreeA.Parameters.Parameters, TEXT("SubtreeAInt"), 10);
+					TStateTreeEditorNode<FTestCondition_PrintValue>& ConditionANode = SubtreeA.AddEnterCondition<FTestCondition_PrintValue>(TEXT("ConditionA"));
+					EditorData.AddPropertyBinding(
+						FPropertyBindingPath(SubtreeA.Parameters.ID, TEXT("SubtreeAInt")),
+						FPropertyBindingPath(ConditionANode.ID, TEXT("Value"))
+					);
+					TStateTreeEditorNode<FTestTask_PrintValue>& TaskANode = SubtreeA.AddTask<FTestTask_PrintValue>(TEXT("TaskA"));
+					EditorData.AddPropertyBinding(
+						FPropertyBindingPath(SubtreeA.Parameters.ID, TEXT("SubtreeAInt")),
+						FPropertyBindingPath(TaskANode.ID, TEXT("Value"))
+					);
+				}
+
+				auto& LinkedState = RootState.AddChildState(TEXT("LinkedState"), EStateTreeStateType::Linked);
+				LinkedState.SetLinkedState(SubtreeA.GetLinkToState());
+				constexpr bool bIsOverridden = true;
+				LinkedState.SetParametersPropertyOverridden(SubtreeAIntID, bIsOverridden);
+				LinkedState.Parameters.Parameters.SetValueInt32(TEXT("SubtreeAInt"), 30);
+			}
+
+			{
+				FStateTreeCompilerLog Log;
+				FStateTreeCompiler Compiler(Log);
+				const bool bResult = Compiler.Compile(StateTree);
+				AITEST_TRUE(TEXT("StateTree should get compiled"), bResult);
+			}
+
+			{
+				FStateTreeInstanceData InstanceData;
+				FTestStateTreeExecutionContext Exec(StateTree, StateTree, InstanceData);
+
+				const bool bInitSucceeded = Exec.IsValid();
+				AITEST_TRUE(TEXT("StateTree should init"), bInitSucceeded);
+
+				EStateTreeRunStatus Status = EStateTreeRunStatus::Unset;
+
+				Status = Exec.Start();
+
+				AITEST_EQUAL(TEXT("Start should complete with Running"), Status, EStateTreeRunStatus::Running);
+				AITEST_TRUE(TEXT("In correct states"), Exec.ExpectInActiveStates("Root", "LinkedState", "SubtreeA"));
+
+				AITEST_TRUE(TEXT("Conditions should correctly update bindings from temporary data"),
+					Exec.Expect("ConditionA", TEXT("TestCondition30")));
+
+				AITEST_TRUE(TEXT("Tasks should correctly update bindings from instance data"),
+					Exec.Expect("TaskA", TEXT("EnterState30")));
+
+				Exec.LogClear();
+
+				Exec.Stop();
+			}
+		}
+		return true;
+	}
+};
+IMPLEMENT_STATE_TREE_INSTANT_TEST(FStateTreeTest_SubtreeParameter, "System.StateTree.Binding.SubtreeParameter");
+
+struct FStateTreeTest_GlobalInstanceData : FStateTreeTestBase
+{
+	virtual bool InstantTest() override
+	{
+		// Tree1(Global Task)
+		//	Root(EnterCondition, Task)
+		//    LinkedStateA -> SubtreeA
+		//  SubtreeA(ConditionA, TaskA)
+		//    LinkedStateB -> SubtreeB
+		//  SubtreeB(ConditionB, TaskB)
+
+		UStateTree& StateTree = NewStateTree();
+		{
+			UStateTreeEditorData& EditorData = *Cast<UStateTreeEditorData>(StateTree.EditorData);
+			{
+				TStateTreeEditorNode<FTestTask_PrintValue>& GlobalTask = EditorData.AddGlobalTask<FTestTask_PrintValue>(TEXT("GlobalTask"));
+				GlobalTask.GetInstanceData().Value = 10;
+
+				auto AddStateConditionAndTask = [&EditorData, &GlobalTask](UStateTreeState& InState, const FName InConditionName, const FName InTaskName)
+				{
+					TStateTreeEditorNode<FTestCondition_PrintValue>& ConditionNode = InState.AddEnterCondition<FTestCondition_PrintValue>(InConditionName);
+					EditorData.AddPropertyBinding(
+						FPropertyBindingPath(GlobalTask.ID, TEXT("Value")),
+						FPropertyBindingPath(ConditionNode.ID, TEXT("Value"))
+					);
+					TStateTreeEditorNode<FTestTask_PrintValue>& TaskNode = InState.AddTask<FTestTask_PrintValue>(InTaskName);
+					EditorData.AddPropertyBinding(
+						FPropertyBindingPath(GlobalTask.ID, TEXT("Value")),
+						FPropertyBindingPath(TaskNode.ID, TEXT("Value"))
+					);
+				};
+
+				auto& RootState = EditorData.AddSubTree(TEXT("Root"));
+				{
+					AddStateConditionAndTask(RootState, TEXT("RootCondition"), TEXT("RootTask"));
+				}
+				auto& SubtreeA = EditorData.AddSubTree(TEXT("SubtreeA"));
+				{
+					SubtreeA.Type = EStateTreeStateType::Subtree;
+					AddStateConditionAndTask(SubtreeA, TEXT("ConditionA"), TEXT("TaskA"));
+				}
+				auto& SubtreeB = EditorData.AddSubTree(TEXT("SubtreeB"));
+				{
+					SubtreeB.Type = EStateTreeStateType::Subtree;
+					AddStateConditionAndTask(SubtreeB, TEXT("ConditionB"), TEXT("TaskB"));
+				}
+
+				auto& LinkedStateA = RootState.AddChildState(TEXT("LinkedStateA"), EStateTreeStateType::Linked);
+				LinkedStateA.SetLinkedState(SubtreeA.GetLinkToState());
+
+				auto& LinkedStateB = SubtreeA.AddChildState(TEXT("LinkedStateB"), EStateTreeStateType::Linked);
+				LinkedStateB.SetLinkedState(SubtreeB.GetLinkToState());
+			}
+
+			{
+				FStateTreeCompilerLog Log;
+				FStateTreeCompiler Compiler(Log);
+				const bool bResult = Compiler.Compile(StateTree);
+				AITEST_TRUE(TEXT("StateTree should get compiled"), bResult);
+			}
+
+			{
+				FStateTreeInstanceData InstanceData;
+				FTestStateTreeExecutionContext Exec(StateTree, StateTree, InstanceData);
+
+				const bool bInitSucceeded = Exec.IsValid();
+				AITEST_TRUE(TEXT("StateTree should init"), bInitSucceeded);
+
+				EStateTreeRunStatus Status = EStateTreeRunStatus::Unset;
+
+				Status = Exec.Start();
+
+				AITEST_EQUAL(TEXT("Start should complete with Running"), Status, EStateTreeRunStatus::Running);
+				AITEST_TRUE(TEXT("In correct states"), Exec.ExpectInActiveStates("Root", "LinkedStateA", "SubtreeA", "LinkedStateB", "SubtreeB"));
+
+				AITEST_TRUE(TEXT("Conditions should correctly update bindings from temporary data"),
+					Exec.Expect("RootCondition", TEXT("TestCondition10"))
+					.Then("ConditionA", TEXT("TestCondition10"))
+					.Then("ConditionB", TEXT("TestCondition10")));
+
+				AITEST_TRUE(TEXT("Tasks should correctly update bindings from instance data"),
+					Exec.Expect("RootTask", TEXT("EnterState10"))
+					.Then("TaskA", TEXT("EnterState10"))
+					.Then("TaskB", TEXT("EnterState10")));
+
+				Exec.LogClear();
+
+				Exec.Stop();
+			}
+		}
+		return true;
+	}
+};
+IMPLEMENT_STATE_TREE_INSTANT_TEST(FStateTreeTest_GlobalInstanceData, "System.StateTree.Binding.GlobalInstanceData");
+
+struct FStateTreeTest_GlobalParameter : FStateTreeTestBase
+{
+	virtual bool InstantTest() override
+	{
+		// TreeA(Global Parameter: 10)
+		//	RootA(ConditionA, TaskA)
+		//    LinkedStateB -> SubtreeB
+		//  SubtreeB(ConditionB, TaskB)
+		//    LinkedAssetStateC -> TreeC
+		// TreeC(Global Parameter: 20)
+		//  RootC(ConditionC, TaskC)
+		//   LinkedStateD -> SubtreeD
+		//  SubtreeD(ConditionD, TaskD)
+
+		auto AddIntPropertyToPropertyBag = [](FInstancedPropertyBag& PropertyBag, FName VarName, int32 DefaultValue)
+			{
+				PropertyBag.AddProperty(VarName, EPropertyBagPropertyType::Int32);
+				PropertyBag.SetValueInt32(VarName, DefaultValue);
+				return PropertyBag.FindPropertyDescByName(VarName)->ID;
+			};
+
+		auto AddStateConditionAndTask = [](UStateTreeEditorData& InEditorData, UStateTreeState& InState, const FName InConditionName, const FName InTaskName, const FName InRootParameterName)
+			{
+				TStateTreeEditorNode<FTestCondition_PrintValue>& ConditionNode = InState.AddEnterCondition<FTestCondition_PrintValue>(InConditionName);
+				InEditorData.AddPropertyBinding(
+					FPropertyBindingPath(InEditorData.GetRootParametersGuid(), InRootParameterName),
+					FPropertyBindingPath(ConditionNode.ID, TEXT("Value"))
+				);
+				TStateTreeEditorNode<FTestTask_PrintValue>& TaskNode = InState.AddTask<FTestTask_PrintValue>(InTaskName);
+				InEditorData.AddPropertyBinding(
+					FPropertyBindingPath(InEditorData.GetRootParametersGuid(), InRootParameterName),
+					FPropertyBindingPath(TaskNode.ID, TEXT("Value"))
+				);
+			};
+
+		UStateTree& StateTreeC = NewStateTree();
+		{
+			UStateTreeEditorData& EditorDataC = *Cast<UStateTreeEditorData>(StateTreeC.EditorData);
+			{
+				AddIntPropertyToPropertyBag(const_cast<FInstancedPropertyBag&>(EditorDataC.GetRootParametersPropertyBag()), TEXT("TreeCInt32"), 20);
+
+				auto& RootCState = EditorDataC.AddSubTree(TEXT("RootC"));
+				{
+					AddStateConditionAndTask(EditorDataC, RootCState, TEXT("ConditionC"), TEXT("TaskC"), TEXT("TreeCInt32"));
+				}
+
+				auto& SubtreeDState = EditorDataC.AddSubTree(TEXT("SubtreeD"));
+				{
+					SubtreeDState.Type = EStateTreeStateType::Subtree;
+					AddStateConditionAndTask(EditorDataC, SubtreeDState, TEXT("ConditionD"), TEXT("TaskD"), TEXT("TreeCInt32"));
+				}
+
+				auto& LinkedStateD = RootCState.AddChildState(TEXT("LinkedStateD"), EStateTreeStateType::Linked);
+				LinkedStateD.SetLinkedState(SubtreeDState.GetLinkToState());
+			}
+		}
+
+		{
+			FStateTreeCompilerLog Log;
+			FStateTreeCompiler Compiler(Log);
+			const bool bResult = Compiler.Compile(StateTreeC);
+			AITEST_TRUE(TEXT("StateTreeC should get compiled"), bResult);
+		}
+
+		UStateTree& StateTreeA = NewStateTree();
+		{
+			UStateTreeEditorData& EditorDataA = *Cast<UStateTreeEditorData>(StateTreeA.EditorData);
+			{
+				AddIntPropertyToPropertyBag(const_cast<FInstancedPropertyBag&>(EditorDataA.GetRootParametersPropertyBag()), TEXT("TreeAInt32"), 10);
+
+				auto& RootAState = EditorDataA.AddSubTree(TEXT("RootA"));
+				{
+					AddStateConditionAndTask(EditorDataA, RootAState, TEXT("ConditionA"), TEXT("TaskA"), TEXT("TreeAInt32"));
+				}
+
+				auto& SubtreeB = EditorDataA.AddSubTree(TEXT("SubtreeB"));
+				{
+					SubtreeB.Type = EStateTreeStateType::Subtree;
+					AddStateConditionAndTask(EditorDataA, SubtreeB, TEXT("ConditionB"), TEXT("TaskB"), TEXT("TreeAInt32"));
+				}
+
+				auto& LinkedAssetStateC = SubtreeB.AddChildState(TEXT("LinkedAssetStateC"), EStateTreeStateType::LinkedAsset);
+				{
+					LinkedAssetStateC.SetLinkedStateAsset(&StateTreeC);
+				}
+
+				auto& LinkedStateB = RootAState.AddChildState(TEXT("LinkedStateB"), EStateTreeStateType::Linked);
+				{
+					LinkedStateB.SetLinkedState(SubtreeB.GetLinkToState());
+				}
+			}
+		}
+		
+		{
+			FStateTreeCompilerLog Log;
+			FStateTreeCompiler Compiler(Log);
+			const bool bResult = Compiler.Compile(StateTreeA);
+			AITEST_TRUE(TEXT("StateTreeA should get compiled"), bResult);
+		}
+
+		{
+			FStateTreeInstanceData InstanceData;
+			FTestStateTreeExecutionContext Exec(StateTreeA, StateTreeA, InstanceData);
+
+			const bool bInitSucceeded = Exec.IsValid();
+			AITEST_TRUE(TEXT("StateTree should init"), bInitSucceeded);
+
+			EStateTreeRunStatus Status = EStateTreeRunStatus::Unset;
+
+			Status = Exec.Start();
+
+			AITEST_EQUAL(TEXT("Start should complete with Running"), Status, EStateTreeRunStatus::Running);
+			AITEST_TRUE(TEXT("In correct states"), Exec.ExpectInActiveStates("RootA", "LinkedStateB", "SubtreeB", "LinkedAssetStateC", "RootC", "LinkedStateD", "SubtreeD"));
+
+			AITEST_TRUE(TEXT("Conditions should correctly update bindings from temporary data"),
+				Exec.Expect("ConditionA", TEXT("TestCondition10"))
+				.Then("ConditionB", TEXT("TestCondition10"))
+				.Then("ConditionC", TEXT("TestCondition20"))
+				.Then("ConditionD", TEXT("TestCondition20")));
+
+			AITEST_TRUE(TEXT("Tasks should correctly update bindings from instance data"),
+				Exec.Expect("TaskA", TEXT("EnterState10"))
+				.Then("TaskB", TEXT("EnterState10"))
+				.Then("TaskC", TEXT("EnterState20"))
+				.Then("TaskD", TEXT("EnterState20")));
+
+			Exec.LogClear();
+
+			Exec.Stop();
+		}
+
+		return true;
+	}
+};
+IMPLEMENT_STATE_TREE_INSTANT_TEST(FStateTreeTest_GlobalParameter, "System.StateTree.Binding.GlobalParameter");
+
 } // namespace UE::StateTree::Tests
 
 UE_ENABLE_OPTIMIZATION_SHIP
diff --git a/Engine/Plugins/Runtime/StateTree/Source/StateTreeTestSuite/Private/StateTreeTestTypes.h b/Engine/Plugins/Runtime/StateTree/Source/StateTreeTestSuite/Private/StateTreeTestTypes.h
--- a/Engine/Plugins/Runtime/StateTree/Source/StateTreeTestSuite/Private/StateTreeTestTypes.h
+++ b/Engine/Plugins/Runtime/StateTree/Source/StateTreeTestSuite/Private/StateTreeTestTypes.h
@@ -405,6 +405,49 @@ struct FTestTask_PrintValue : public FStateTreeTaskBase
 	TFunction<void(FStateTreeExecutionContext&, const FTestTask_PrintValue*)> CustomTickFunc;
 };
 
+USTRUCT(meta = (Hidden))
+struct FTestCondition_PrintValue : public FStateTreeConditionCommonBase
+{
+	GENERATED_BODY()
+
+	using FInstanceDataType = FTestTask_PrintValueInstanceData;
+	using FExecutionRuntimeDataType = FTestTask_PrintValueExecutionRuntimeData;
+
+	FTestCondition_PrintValue() = default;
+	FTestCondition_PrintValue(const FName InName)
+	{
+		Name = InName;
+	}
+	virtual ~FTestCondition_PrintValue() override {}
+
+	virtual const UStruct* GetInstanceDataType() const override { return FInstanceDataType::StaticStruct(); }
+	virtual const UStruct* GetExecutionRuntimeDataType() const override { return FExecutionRuntimeDataType::StaticStruct(); }
+
+	virtual bool Link(FStateTreeLinker& Linker) override
+	{
+		Linker.LinkExternalData(LogHandle);
+		return true;
+	}
+
+	virtual bool TestCondition(FStateTreeExecutionContext& Context) const override
+	{
+		FStateTreeTestLog& TestLog = Context.GetExternalData(LogHandle);
+		const FInstanceDataType& InstanceData = Context.GetInstanceData(*this);
+		TestLog.Log(Name, FString::Printf(TEXT("TestCondition%d"), InstanceData.Value));
+		TestLog.Log(Name, FString::Printf(TEXT("TestCondition:%s"), *InstanceData.GetArrayDescription()));
+
+		if (CustomTestConditionFunc)
+		{
+			return CustomTestConditionFunc(Context, this);
+		}
+
+		return true;
+	}
+
+	TFunction<bool(FStateTreeExecutionContext&, const FTestCondition_PrintValue*)> CustomTestConditionFunc;
+	TStateTreeExternalDataHandle<FStateTreeTestLog> LogHandle;
+};
+
 USTRUCT()
 struct FTestTask_PrintAndResetValue : public FTestTask_PrintValue
 {
