[State Tree] fixed incorrect comment

diff --git a/Engine/Plugins/Runtime/StateTree/Source/StateTreeModule/Public/StateTreeExecutionTypes.h b/Engine/Plugins/Runtime/StateTree/Source/StateTreeModule/Public/StateTreeExecutionTypes.h
--- a/Engine/Plugins/Runtime/StateTree/Source/StateTreeModule/Public/StateTreeExecutionTypes.h
+++ b/Engine/Plugins/Runtime/StateTree/Source/StateTreeModule/Public/StateTreeExecutionTypes.h
@@ -1085,8 +1085,8 @@ PRAGMA_ENABLE_DEPRECATION_WARNINGS
 
 	/**
 	 * Index within the instance data that StateParameterDataHandle is relative to,
-	 * if it lives as state parameter in the parent frame
-	 * Invalid means GlobalParameterDataHandle points to the root global parameter.
+	 * It could be the current frame or parent frame's active instance index base, depending on if the current frame is a linked subtree.
+	 * Always valid.
 	 */
 	UPROPERTY()
 	FStateTreeIndex16 StateParameterIndexBase = FStateTreeIndex16::Invalid;
