[StateTree] Add option to hide the debugger migration warning.

#rnx

diff --git a/Engine/Plugins/Runtime/StateTree/Source/StateTreeEditorModule/Private/Debugger/SStateTreeDebuggerView.cpp b/Engine/Plugins/Runtime/StateTree/Source/StateTreeEditorModule/Private/Debugger/SStateTreeDebuggerView.cpp
--- a/Engine/Plugins/Runtime/StateTree/Source/StateTreeEditorModule/Private/Debugger/SStateTreeDebuggerView.cpp
+++ b/Engine/Plugins/Runtime/StateTree/Source/StateTreeEditorModule/Private/Debugger/SStateTreeDebuggerView.cpp
@@ -20,6 +20,7 @@
 #include "StateTreeEditorData.h"
 #include "StateTreeEditorSettings.h"
 #include "StateTreeEditorStyle.h"
+#include "StateTreeEditorUserSettings.h"
 #include "StateTreeModule.h"
 #include "StateTreeViewModel.h"
 #include "SWarningOrErrorBox.h"
@@ -544,8 +545,7 @@ void SStateTreeDebuggerView::ConstructView(const TSharedRef<FUICommandList>& InC
 
 		LeftToolbar.SetLabelVisibility(EVisibility::Visible);
 		LeftToolbar.AddToolBarButton(DebuggerCommands.OpenRewindDebugger);
-		LeftToolbar.SetLabelVisibility(EVisibility::Hidden);
-
+		LeftToolbar.SetLabelVisibility(EVisibility::Collapsed);
 		LeftToolbar.AddToolBarButton(DebuggerCommands.ResumeDebuggerAnalysis);
 	}
 	LeftToolbar.EndSection();
@@ -634,6 +634,41 @@ void SStateTreeDebuggerView::ConstructView(const TSharedRef<FUICommandList>& InC
 					})
 		];
 
+	// Deprecation warning message. Only visible when bDisplayMigrationToRewindDebuggerWarning.
+	TSharedRef<SWidget> MigrationToRewindDebugger = SNullWidget::NullWidget;
+	bDisplayMigrationToRewindDebuggerWarning = GetDefault<UStateTreeEditorUserSettings>()->GetDisplayMigrationToRewindDebuggerWarning();
+	if (bDisplayMigrationToRewindDebuggerWarning)
+	{
+		const FButtonStyle* const CloseButtonStyle = &FAppStyle::Get().GetWidgetStyle<FDockTabStyle>("Docking.Tab").CloseButtonStyle;
+		MigrationToRewindDebugger = SNew(SWarningOrErrorBox)
+			.MessageStyle(EMessageStyle::Warning)
+			.Message(LOCTEXT("MigrationToRewindDebugger",
+				"Note that recording and visualization of the recorded data has been migrated to RewindDebugger.\n"
+				"It is still possible to use the previous visualization by using 'Enable Legacy Debugger Window'"
+				" under the StateTree Editor plugin section of the Editor Preferences."))
+			.Visibility(MakeAttributeSPLambda(AsShared(), [Self = this]() -> EVisibility
+				{
+					return Self->bDisplayMigrationToRewindDebuggerWarning ? EVisibility::Visible : EVisibility::Collapsed;
+				}))
+			.ContentVAlign(EVerticalAlignment::VAlign_Top)
+			[
+				SNew(SButton)
+					.ButtonStyle(CloseButtonStyle)
+					.ContentPadding(FMargin(0.0f, 1.5f, 0.0f, 0.0f))
+					.ToolTipText(LOCTEXT("MigrationToRewindDebuggerCloseTooltip", "Close the warning box."))
+					.OnClicked(FOnClicked::CreateSPLambda(AsShared(), [Self = this]() -> FReply
+						{
+							Self->bDisplayMigrationToRewindDebuggerWarning = false;
+							GetMutableDefault<UStateTreeEditorUserSettings>()->SetDisplayMigrationToRewindDebuggerWarning(false);
+							return FReply::Handled();
+						}))
+					[
+						SNew(SSpacer)
+						.Size(CloseButtonStyle->Normal.ImageSize)
+					]
+			];
+	}
+
 	ChildSlot
 	[
 		SNew(SBorder)
@@ -646,28 +681,23 @@ void SStateTreeDebuggerView::ConstructView(const TSharedRef<FUICommandList>& InC
 			.MinHeight(48.f)
 			[
 				SNew(SHorizontalBox)
-					+ SHorizontalBox::Slot()
-					.AutoWidth()
-					.VAlign(VAlign_Center)
-					[
-						LeftToolbar.MakeWidget()
-					]
-					+ SHorizontalBox::Slot()
-					.AutoWidth()
-					.VAlign(VAlign_Center)
-					[
-						InstanceSelectionBox
-					]
+				+ SHorizontalBox::Slot()
+				.AutoWidth()
+				.VAlign(EVerticalAlignment::VAlign_Center)
+				[
+					LeftToolbar.MakeWidget()
+				]
+				+ SHorizontalBox::Slot()
+				.AutoWidth()
+				.VAlign(EVerticalAlignment::VAlign_Center)
+				[
+					InstanceSelectionBox
+				]
 			]
 			+ SVerticalBox::Slot()
 			.AutoHeight()
 			[
-				SNew(SWarningOrErrorBox)
-				.MessageStyle(EMessageStyle::Warning)
-				.Message(LOCTEXT("MigrationToRewindDebugger",
-					"Note that recording and visualization of the recorded data has been migrated to RewindDebugger.\n"
-					"It is still possible to use the previous visualization by using 'Enable Legacy Debugger Window'"
-					" under the StateTree Editor plugin section of the Editor Preferences."))
+				MigrationToRewindDebugger
 			]
 		]
 	];
diff --git a/Engine/Plugins/Runtime/StateTree/Source/StateTreeEditorModule/Private/Debugger/SStateTreeDebuggerView.h b/Engine/Plugins/Runtime/StateTree/Source/StateTreeEditorModule/Private/Debugger/SStateTreeDebuggerView.h
--- a/Engine/Plugins/Runtime/StateTree/Source/StateTreeEditorModule/Private/Debugger/SStateTreeDebuggerView.h
+++ b/Engine/Plugins/Runtime/StateTree/Source/StateTreeEditorModule/Private/Debugger/SStateTreeDebuggerView.h
@@ -187,6 +187,9 @@ private:
 	 */
 	double ExtrapolatedRecordedWorldTime = 0.0;
 
+	/** Whether the "Migration To Rewind Debugger" is visible. */
+	bool bDisplayMigrationToRewindDebuggerWarning = true;
+
 private:
 
 	FDelegateHandle BeginPIEHandle;
diff --git a/Engine/Plugins/Runtime/StateTree/Source/StateTreeEditorModule/Private/StateTreeEditorUserSettings.cpp b/Engine/Plugins/Runtime/StateTree/Source/StateTreeEditorModule/Private/StateTreeEditorUserSettings.cpp
--- a/Engine/Plugins/Runtime/StateTree/Source/StateTreeEditorModule/Private/StateTreeEditorUserSettings.cpp
+++ b/Engine/Plugins/Runtime/StateTree/Source/StateTreeEditorModule/Private/StateTreeEditorUserSettings.cpp
@@ -13,6 +13,14 @@ void UStateTreeEditorUserSettings::SetStatesViewDisplayNodeType(EStateTreeEditor
 	}
 }
 
+void UStateTreeEditorUserSettings::SetDisplayMigrationToRewindDebuggerWarning(bool bNewValue)
+{
+	if (bDisplayMigrationToRewindDebuggerWarning != bNewValue)
+	{
+		bDisplayMigrationToRewindDebuggerWarning = bNewValue;
+		SaveConfig();
+	}
+}
 
 void UStateTreeEditorUserSettings::PostEditChangeChainProperty(struct FPropertyChangedChainEvent& PropertyChangedEvent)
 {
diff --git a/Engine/Plugins/Runtime/StateTree/Source/StateTreeEditorModule/Private/StateTreeEditorUserSettings.h b/Engine/Plugins/Runtime/StateTree/Source/StateTreeEditorModule/Private/StateTreeEditorUserSettings.h
--- a/Engine/Plugins/Runtime/StateTree/Source/StateTreeEditorModule/Private/StateTreeEditorUserSettings.h
+++ b/Engine/Plugins/Runtime/StateTree/Source/StateTreeEditorModule/Private/StateTreeEditorUserSettings.h
@@ -41,6 +41,13 @@ public:
 		return StatesViewNodeRowHeight;
 	}
 
+	bool GetDisplayMigrationToRewindDebuggerWarning() const
+	{
+		return bDisplayMigrationToRewindDebuggerWarning;
+	}
+
+	void SetDisplayMigrationToRewindDebuggerWarning(bool bNewValue);
+
 	/** Broadcast when a setting changes. */
 	FSimpleMulticastDelegate OnSettingsChanged;
 
@@ -59,4 +66,8 @@ private:
 	/** Height of a node in the States View. */
 	UPROPERTY(config, EditAnywhere, Category = "State View", meta = (DisplayName = "Node Height", ClampMin = 8.0f))
 	float StatesViewNodeRowHeight = 16.0f;
+
+	/** Display the debugger migration warning. */
+	UPROPERTY(config, EditAnywhere, Category = "Debugger")
+	bool bDisplayMigrationToRewindDebuggerWarning = true;
 };
