[StateTree] Fix fallout from CL47100678. When moving global task, you do not want to duplicate the UObject because it is already copied and used in EnterState.

diff --git a/Engine/Plugins/Runtime/StateTree/Source/StateTreeModule/Private/StateTreeExecutionContext.cpp b/Engine/Plugins/Runtime/StateTree/Source/StateTreeModule/Private/StateTreeExecutionContext.cpp
--- a/Engine/Plugins/Runtime/StateTree/Source/StateTreeModule/Private/StateTreeExecutionContext.cpp
+++ b/Engine/Plugins/Runtime/StateTree/Source/StateTreeModule/Private/StateTreeExecutionContext.cpp
@@ -2588,7 +2588,7 @@ FStateTreeDataView FStateTreeExecutionContext::MoveTemporaryToInstance(const FSt
 			FStateTreeInstanceObjectWrapper* Wrapper = Struct.GetMutablePtr<FStateTreeInstanceObjectWrapper>();
 			if (Wrapper->InstanceObject)
 			{
-				const bool bDuplicate = true;
+				constexpr bool bDuplicate = false;
 				Wrapper->InstanceObject = UE::StateTree::InstanceData::Private::CopyNodeInstance(Wrapper->InstanceObject, &Owner, bDuplicate);
 			}
 			Struct = FStateTreeDataView(Wrapper->InstanceObject);
