[StateTree] Only append temporaries when they will not be copied later in Enter. This fix issue introduced with CL47100678.

diff --git a/Engine/Plugins/Runtime/StateTree/Source/StateTreeModule/Private/StateTreeExecutionContext.cpp b/Engine/Plugins/Runtime/StateTree/Source/StateTreeModule/Private/StateTreeExecutionContext.cpp
--- a/Engine/Plugins/Runtime/StateTree/Source/StateTreeModule/Private/StateTreeExecutionContext.cpp
+++ b/Engine/Plugins/Runtime/StateTree/Source/StateTreeModule/Private/StateTreeExecutionContext.cpp
@@ -2289,7 +2289,7 @@ void FStateTreeExecutionContext::UpdateInstanceData(const TSharedPtr<FSelectStat
 	// Temporaries that are moved to the new instance data buffer.
 	// Node instance data are moved in EnterState when the ActiveNodeIndex is incremented.
 	// The array is clear at the end of EnterState.
-	TArray<FInstancedStruct*, FNonconcurrentLinearArrayAllocator> TempInstanceStructs;
+	TArray<TOptional<FInstancedStruct*>, FNonconcurrentLinearArrayAllocator> TempInstanceStructs;
 	TempInstanceStructs.Reserve(EstimatedNumStructs);
 
 	constexpr int32 ExpectedAmountOfFrames = 8;
@@ -2305,6 +2305,9 @@ void FStateTreeExecutionContext::UpdateInstanceData(const TSharedPtr<FSelectStat
 			return TempData ? &TempData->Instance : nullptr;
 		};
 
+	constexpr TOptional<FInstancedStruct*> NoExistingTemporaryDuplicateWrapper = TOptional<FInstancedStruct*>(nullptr);
+	constexpr TOptional<FInstancedStruct*> TemporaryToBeMovedLater = TOptional<FInstancedStruct*>();
+
 	// Find next instance data sources and find common/existing section of instance data at start.
 	int32 CurrentGlobalInstanceIndexBase = 0;
 	int32 NumCommonInstanceData = 0;
@@ -2366,7 +2369,9 @@ void FStateTreeExecutionContext::UpdateInstanceData(const TSharedPtr<FSelectStat
 					{
 						const FStateTreeEvaluatorBase& Eval = CurrentStateTree->Nodes[EvalIndex].Get<const FStateTreeEvaluatorBase>();
 						const FConstStructView EvalInstanceData = CurrentStateTree->DefaultInstanceData.GetStruct(Eval.InstanceTemplateIndex.Get());
-						InstanceStructs[BaseIndex + Eval.InstanceDataHandle.GetIndex()] = EvalInstanceData;
+						const int32 EvalInstanceDataIndex = BaseIndex + Eval.InstanceDataHandle.GetIndex();
+						InstanceStructs[EvalInstanceDataIndex] = EvalInstanceData;
+						TempInstanceStructs[EvalInstanceDataIndex] = TemporaryToBeMovedLater;
 					}
 
 					// Global tasks
@@ -2374,7 +2379,9 @@ void FStateTreeExecutionContext::UpdateInstanceData(const TSharedPtr<FSelectStat
 					{
 						const FStateTreeTaskBase& Task = CurrentStateTree->Nodes[TaskIndex].Get<const FStateTreeTaskBase>();
 						const FConstStructView TaskInstanceData = CurrentStateTree->DefaultInstanceData.GetStruct(Task.InstanceTemplateIndex.Get());
-						InstanceStructs[BaseIndex + Task.InstanceDataHandle.GetIndex()] = TaskInstanceData;
+						const int32 TaskInstanceDataIndex = BaseIndex + Task.InstanceDataHandle.GetIndex();
+						InstanceStructs[TaskInstanceDataIndex] = TaskInstanceData;
+						TempInstanceStructs[TaskInstanceDataIndex] = TemporaryToBeMovedLater;
 					}
 
 					if (bIsFrameCommon)
@@ -2406,6 +2413,8 @@ void FStateTreeExecutionContext::UpdateInstanceData(const TSharedPtr<FSelectStat
 					InstanceStructs.AddDefaulted(CurrentState->InstanceDataNum);
 					TempInstanceStructs.AddZeroed(CurrentState->InstanceDataNum);
 
+					const int32 CurrentStateParameterInstanceDataIndex = CurrentFrameBaseIndex + CurrentState->ParameterDataHandle.GetIndex();
+
 					bool bCanHaveTempData = false;
 					if (CurrentState->Type == EStateTreeStateType::Subtree)
 					{
@@ -2415,7 +2424,7 @@ void FStateTreeExecutionContext::UpdateInstanceData(const TSharedPtr<FSelectStat
 						if (!NextStateParameterDataHandle.IsValid())
 						{
 							// Parameters are not set by a linked state, create instance data.
-							InstanceStructs[CurrentFrameBaseIndex + CurrentState->ParameterDataHandle.GetIndex()] = ParamsInstanceData;
+							InstanceStructs[CurrentStateParameterInstanceDataIndex] = ParamsInstanceData;
 							bCanHaveTempData = true;
 						}
 						else
@@ -2446,7 +2455,7 @@ void FStateTreeExecutionContext::UpdateInstanceData(const TSharedPtr<FSelectStat
 							{
 								// If we have temp data for the parameters, then setup the instance data with just a type, so that we can steal the temp data below (TempInstanceStructs).
 								// We expect overridden linked assets to hit this code path. 
-								InstanceStructs[CurrentFrameBaseIndex + CurrentState->ParameterDataHandle.GetIndex()] = FConstStructView(TempParamsInstanceData->GetScriptStruct());
+								InstanceStructs[CurrentStateParameterInstanceDataIndex] = FConstStructView(TempParamsInstanceData->GetScriptStruct());
 								Params = TempParamsInstanceData->GetPtr<const FCompactStateTreeParameters>();
 								bCanHaveTempData = true;
 							}
@@ -2474,7 +2483,7 @@ void FStateTreeExecutionContext::UpdateInstanceData(const TSharedPtr<FSelectStat
 								{
 									ParamsInstanceData = CurrentStateTree->DefaultInstanceData.GetStruct(CurrentState->ParameterTemplateIndex.Get());
 								}
-								InstanceStructs[CurrentFrameBaseIndex + CurrentState->ParameterDataHandle.GetIndex()] = ParamsInstanceData;
+								InstanceStructs[CurrentStateParameterInstanceDataIndex] = ParamsInstanceData;
 								Params = ParamsInstanceData.GetPtr<const FCompactStateTreeParameters>();
 								bCanHaveTempData = true;
 							}
@@ -2491,21 +2500,32 @@ void FStateTreeExecutionContext::UpdateInstanceData(const TSharedPtr<FSelectStat
 						}
 					}
 
-					if (!bIsStateCommon && bCanHaveTempData)
+					if (!bIsStateCommon)
 					{
-						TempInstanceStructs[CurrentFrameBaseIndex + CurrentState->ParameterDataHandle.GetIndex()] = FindInstanceTempData(CurrentFrameID, CurrentState->ParameterDataHandle);
+						if (bCanHaveTempData)
+						{
+							TempInstanceStructs[CurrentStateParameterInstanceDataIndex] = FindInstanceTempData(CurrentFrameID, CurrentState->ParameterDataHandle);
+						}
+						else
+						{
+							TempInstanceStructs[CurrentStateParameterInstanceDataIndex] = NoExistingTemporaryDuplicateWrapper;
+						}
 					}
 
 					if (CurrentState->EventDataIndex.IsValid())
 					{
-						InstanceStructs[CurrentFrameBaseIndex + CurrentState->EventDataIndex.Get()] = FConstStructView(FStateTreeSharedEvent::StaticStruct());
+						const int32 EventDataInstanceDataIndex = CurrentFrameBaseIndex + CurrentState->EventDataIndex.Get();
+						InstanceStructs[EventDataInstanceDataIndex] = FConstStructView(FStateTreeSharedEvent::StaticStruct());
+						TempInstanceStructs[EventDataInstanceDataIndex] = NoExistingTemporaryDuplicateWrapper;
 					}
 
 					for (int32 TaskIndex = CurrentState->TasksBegin; TaskIndex < (CurrentState->TasksBegin + CurrentState->TasksNum); TaskIndex++)
 					{
 						const FStateTreeTaskBase& Task = CurrentStateTree->Nodes[TaskIndex].Get<const FStateTreeTaskBase>();
 						const FConstStructView TaskInstanceData = CurrentStateTree->DefaultInstanceData.GetStruct(Task.InstanceTemplateIndex.Get());
-						InstanceStructs[CurrentFrameBaseIndex + Task.InstanceDataHandle.GetIndex()] = TaskInstanceData;
+						const int32 TaskInstanceDataIndex = CurrentFrameBaseIndex + Task.InstanceDataHandle.GetIndex();
+						InstanceStructs[TaskInstanceDataIndex] = TaskInstanceData;
+						TempInstanceStructs[TaskInstanceDataIndex] = NoExistingTemporaryDuplicateWrapper;
 					}
 
 					if (bIsStateCommon)
@@ -2519,7 +2539,7 @@ void FStateTreeExecutionContext::UpdateInstanceData(const TSharedPtr<FSelectStat
 
 	// Common section should match.
 #if WITH_STATETREE_DEBUG
-	for (int32 Index = 0; Index < NumCommonInstanceData; Index++)
+	for (int32 Index = 0; Index < NumCommonInstanceData; ++Index)
 	{
 		check(Index < InstanceData.Num());
 
@@ -2547,12 +2567,12 @@ void FStateTreeExecutionContext::UpdateInstanceData(const TSharedPtr<FSelectStat
 		MakeArrayView(TempInstanceStructs.GetData() + NumCommonInstanceData, TempInstanceStructs.Num() - NumCommonInstanceData));
 }
 
-FStateTreeDataView FStateTreeExecutionContext::MoveTemporaryToInstance(const FStateTreeExecutionFrame& CurrentFrame, const FStateTreeDataHandle Handle)
+FStateTreeDataView FStateTreeExecutionContext::MoveTemporaryToInstance(const FStateTreeExecutionFrame& CurrentFrame, const FStateTreeDataHandle DataHandle)
 {
-	check(Handle.GetSource() == EStateTreeDataSourceType::ActiveInstanceData
-		|| Handle.GetSource() == EStateTreeDataSourceType::ActiveInstanceDataObject
-		|| Handle.GetSource() == EStateTreeDataSourceType::GlobalInstanceData
-		|| Handle.GetSource() == EStateTreeDataSourceType::GlobalInstanceDataObject);
+	check(DataHandle.GetSource() == EStateTreeDataSourceType::ActiveInstanceData
+		|| DataHandle.GetSource() == EStateTreeDataSourceType::ActiveInstanceDataObject
+		|| DataHandle.GetSource() == EStateTreeDataSourceType::GlobalInstanceData
+		|| DataHandle.GetSource() == EStateTreeDataSourceType::GlobalInstanceDataObject);
 
 	using namespace UE::StateTree;
 	using namespace UE::StateTree::ExecutionContext;
@@ -2568,34 +2588,33 @@ FStateTreeDataView FStateTreeExecutionContext::MoveTemporaryToInstance(const FSt
 			return TempData ? &TempData->Instance : nullptr;
 		};
 
-	FStateTreeDataView Struct = UE::StateTree::InstanceData::GetDataView(Storage, nullptr, CurrentFrame, Handle.ToStructSource());
-	FInstancedStruct* Source = FindInstanceTempData(CurrentFrame.FrameID, Handle);
+	// Use ToStructSource to get the FStateTreeInstanceObjectWrapper
+	const FStateTreeDataHandle StructDataHandle = DataHandle.ToStructSource();
+	FStateTreeDataView DataView = UE::StateTree::InstanceData::GetDataView(Storage, nullptr, CurrentFrame, StructDataHandle);
+	FInstancedStruct* TemporarySource = FindInstanceTempData(CurrentFrame.FrameID, DataHandle);
 
-	if (Struct.IsValid())
+	// If we found a temporary data, move it from the temporary to the instance data.
+	if (DataView.IsValid())
 	{
 		// ToStructSource should have converted it to FStateTreeInstanceObjectWrapper
-		check(Struct.GetStruct() && Cast<const UClass>(Struct.GetStruct()) == nullptr);
-
-		if (Source && Source->IsValid())
-		{
-			check(Struct.GetStruct() == Source->GetScriptStruct());
+		check(DataView.GetStruct() && Cast<const UClass>(DataView.GetStruct()) == nullptr);
 
-			FMemory::Memswap(Struct.GetMutableMemory(), Source->GetMutableMemory(), Struct.GetStruct()->GetStructureSize());
-			Source->Reset();
-		}
-		else if (Struct.GetStruct()->IsChildOf(FStateTreeInstanceObjectWrapper::StaticStruct()))
+		if (TemporarySource)
 		{
-			FStateTreeInstanceObjectWrapper* Wrapper = Struct.GetMutablePtr<FStateTreeInstanceObjectWrapper>();
-			if (Wrapper->InstanceObject)
+			if (ensure(TemporarySource->IsValid()))
 			{
-				constexpr bool bDuplicate = false;
-				Wrapper->InstanceObject = UE::StateTree::InstanceData::Private::CopyNodeInstance(Wrapper->InstanceObject, &Owner, bDuplicate);
+				check(DataView.GetStruct() == TemporarySource->GetScriptStruct());
+				FMemory::Memswap(DataView.GetMutableMemory(), TemporarySource->GetMutableMemory(), DataView.GetStruct()->GetStructureSize());
 			}
-			Struct = FStateTreeDataView(Wrapper->InstanceObject);
+		}
+		else if (StructDataHandle != DataHandle)
+		{
+			// if there is no temporary, get the correct data from the input handle
+			DataView = UE::StateTree::InstanceData::GetDataView(Storage, nullptr, CurrentFrame, DataHandle);
 		}
 	}
 
-	return Struct;
+	return DataView;
 }
 
 FStateTreeDataView FStateTreeExecutionContext::GetDataView(const FStateTreeExecutionFrame* ParentFrame, const FStateTreeExecutionFrame& CurrentFrame, const FStateTreeDataHandle Handle)
diff --git a/Engine/Plugins/Runtime/StateTree/Source/StateTreeModule/Private/StateTreeInstanceData.cpp b/Engine/Plugins/Runtime/StateTree/Source/StateTreeModule/Private/StateTreeInstanceData.cpp
--- a/Engine/Plugins/Runtime/StateTree/Source/StateTreeModule/Private/StateTreeInstanceData.cpp
+++ b/Engine/Plugins/Runtime/StateTree/Source/StateTreeModule/Private/StateTreeInstanceData.cpp
@@ -1047,7 +1047,7 @@ void FStateTreeInstanceData::Append(UObject& InOwner, TConstArrayView<FConstStru
 void FStateTreeInstanceData::Append(UObject& InOwner, TConstArrayView<FConstStructView> InStructs, TConstArrayView<FInstancedStruct*> InInstancesToMove, FAddArgs Args)
 {
 	check(InStructs.Num() == InInstancesToMove.Num());
-	
+
 	FStateTreeInstanceStorage& Storage = GetMutableStorage();
 
 	const int32 StartIndex = Storage.InstanceStructs.Num();
@@ -1083,6 +1083,48 @@ void FStateTreeInstanceData::Append(UObject& InOwner, TConstArrayView<FConstStru
 	}
 }
 
+void FStateTreeInstanceData::Append(UObject& InOwner, TConstArrayView<FConstStructView> InStructs, TConstArrayView<TOptional<FInstancedStruct*>> InInstancesToMove, FAddArgs Args)
+{
+	check(InStructs.Num() == InInstancesToMove.Num());
+
+	FStateTreeInstanceStorage& Storage = GetMutableStorage();
+
+	const int32 StartIndex = Storage.InstanceStructs.Num();
+	//The Append will not duplicate the UObject.
+	Storage.InstanceStructs.Append(InStructs);
+
+	for (int32 Index = StartIndex; Index < Storage.InstanceStructs.Num(); Index++)
+	{
+		FStructView Struct = Storage.InstanceStructs[Index];
+		const TOptional<FInstancedStruct*>& OptionalSource = InInstancesToMove[Index - StartIndex];
+
+		// The source is used to move temporary instance data into instance data. Not all entries may have it.
+		// The instance struct can be empty, in which case the temporary instance is ignored.
+		// If the source is specified, move it to the instance data.
+		// We assume that if the source is object wrapper, it is already the instance we want.
+		if (Struct.IsValid() && OptionalSource.IsSet())
+		{
+			FInstancedStruct* Source = OptionalSource.GetValue();
+			if (Source && Source->IsValid())
+			{
+				check(Struct.GetScriptStruct() == Source->GetScriptStruct());
+
+				FMemory::Memswap(Struct.GetMemory(), Source->GetMutableMemory(), Struct.GetScriptStruct()->GetStructureSize());
+				Source->Reset();
+			}
+			else if (FStateTreeInstanceObjectWrapper* Wrapper = Struct.GetPtr<FStateTreeInstanceObjectWrapper>())
+			{
+				// The source is explicitly set to nullptr. Duplicate the UObject.
+				if (Wrapper->InstanceObject)
+				{
+					const bool bDuplicate = Args.bDuplicateWrappedObject || &InOwner != Wrapper->InstanceObject->GetOuter();
+					Wrapper->InstanceObject = UE::StateTree::InstanceData::Private::CopyNodeInstance(Wrapper->InstanceObject, &InOwner, bDuplicate);
+				}
+			}
+		}
+	}
+}
+
 void FStateTreeInstanceData::ShrinkTo(const int32 NumStructs)
 {
 	FStateTreeInstanceStorage& Storage = GetMutableStorage();
diff --git a/Engine/Plugins/Runtime/StateTree/Source/StateTreeModule/Public/StateTreeInstanceData.h b/Engine/Plugins/Runtime/StateTree/Source/StateTreeModule/Public/StateTreeInstanceData.h
--- a/Engine/Plugins/Runtime/StateTree/Source/StateTreeModule/Public/StateTreeInstanceData.h
+++ b/Engine/Plugins/Runtime/StateTree/Source/StateTreeModule/Public/StateTreeInstanceData.h
@@ -459,8 +459,16 @@ struct FStateTreeInstanceData
 	UE_API void Append(UObject& InOwner, TConstArrayView<FInstancedStruct> InStructs, FAddArgs Args = FAddArgs::Default);
 	UE_API void Append(UObject& InOwner, TConstArrayView<FConstStructView> InStructs, FAddArgs Args = FAddArgs::Default);
 
-	/** Appends new items to the instance, and moves existing data into the allocated instances. */
+	/**
+	 * Appends new items to the instance, and moves existing data into the allocated instances.
+	 * If the struct is of type FStateTreeInstanceObjectWrapper then this will duplicate the UObject instance.
+	 */
 	UE_API void Append(UObject& InOwner, TConstArrayView<FConstStructView> InStructs, TConstArrayView<FInstancedStruct*> InInstancesToMove, FAddArgs Args = FAddArgs::Default);
+	/**
+	 * Appends new items to the instance, and moves existing data into the allocated instances, if provided.
+	 * If the struct is of type FStateTreeInstanceObjectWrapper then this will duplicate the UObject instance only if the InstanceToMove is set.
+	 */
+	UE_API void Append(UObject& InOwner, TConstArrayView<FConstStructView> InStructs, TConstArrayView<TOptional<FInstancedStruct*>> InInstancesToMove, FAddArgs Args = FAddArgs::Default);
 
 	/** Shrinks the array sizes to specified lengths. Sizes must be small or equal than current size. */
 	UE_API void ShrinkTo(const int32 Num);
