[StateTree] Making sure the TargetIndirections is not empty before using it.

#rnx

diff --git a/Engine/Plugins/Runtime/StateTree/Source/StateTreeEditorModule/Private/Customizations/StateTreeEditorNodeDetails.cpp b/Engine/Plugins/Runtime/StateTree/Source/StateTreeEditorModule/Private/Customizations/StateTreeEditorNodeDetails.cpp
--- a/Engine/Plugins/Runtime/StateTree/Source/StateTreeEditorModule/Private/Customizations/StateTreeEditorNodeDetails.cpp
+++ b/Engine/Plugins/Runtime/StateTree/Source/StateTreeEditorModule/Private/Customizations/StateTreeEditorNodeDetails.cpp
@@ -177,7 +177,7 @@ namespace UE::StateTreeEditor::Internal
 					if (ensure(EditorData->GetBindingDataViewByID(ID, TargetDataView)))
 					{
 						TArray<FPropertyBindingPathIndirection> TargetIndirections;
-						if (ensure(Path.ResolveIndirectionsWithValue(TargetDataView, TargetIndirections)))
+						if (ensure(Path.ResolveIndirectionsWithValue(TargetDataView, TargetIndirections)) && TargetIndirections.Num() > 0)
 						{
 							const void* PropertyRef = TargetIndirections.Last().GetPropertyAddress();
 							PinType = UE::StateTree::PropertyRefHelpers::GetPropertyRefInternalTypeAsPin(*Property, PropertyRef);
