[StateTree] Add a queue and asynchronous state tree compilation. The compilation is always on the game thread. The state tree asset is compiled "just in time" the first time it is used by the game. The features are disabled by default.

With cvar StateTree.Compiler.EnableQueuedCompilationOnAssetLoad. The public step will be compiled synchronously when the asset loads, and the internal steps will be queued for later. Without the cvar enabled, the previous behavior is maintained. The compilation of the public and internal steps are compiled synchronously.

With cvar StateTree.Compiler.EnableQueuedCompilationWhenAssetDirty, the asset is queued for compilation as soon as it is modified in the editor.

With cvar StateTree.Compiler.EnableCompileSynchronouslyInPIESession, all queued assets are compiled synchronously before a PIE session starts. While in PIE, assets will not be queued. They are compiled synchronously.

With car StateTree.Compiler.NumberOfQueuedCompilationPerBatch, you control how many state tree assets compile per ticking loop. To deactivate, set to 0. Every tick, the engine will compile NumberOfQueuedCompilationPerBatch queued state tree assets, including the dependencies.

diff --git a/Engine/Plugins/AI/MassAI/Source/MassAIBehavior/Public/MassStateTreeSchema.h b/Engine/Plugins/AI/MassAI/Source/MassAIBehavior/Public/MassStateTreeSchema.h
--- a/Engine/Plugins/AI/MassAI/Source/MassAIBehavior/Public/MassStateTreeSchema.h
+++ b/Engine/Plugins/AI/MassAI/Source/MassAIBehavior/Public/MassStateTreeSchema.h
@@ -20,6 +20,14 @@ public:
 	/** Fetches a read-only view of the Mass-relevant requirements of the associated StateTee */
 	const TArray<FMassStateTreeDependency>& GetDependencies() const;
 
+#if WITH_EDITOR
+	virtual bool AllowQueuedCompilation() const override
+	{
+		// If queued, they might be compiled in a mass worker thread via CompileIfNeededSynchronously.
+		return false;
+	}
+#endif
+
 protected:
 	UE_API virtual bool IsStructAllowed(const UScriptStruct* InScriptStruct) const override;
 	UE_API virtual bool IsExternalItemAllowed(const UStruct& InStruct) const override;
diff --git a/Engine/Plugins/Experimental/UAF/UAFStateTree/Source/UAFStateTreeUncookedOnly/Private/AnimNextStateTreeEditorData.cpp b/Engine/Plugins/Experimental/UAF/UAFStateTree/Source/UAFStateTreeUncookedOnly/Private/AnimNextStateTreeEditorData.cpp
--- a/Engine/Plugins/Experimental/UAF/UAFStateTree/Source/UAFStateTreeUncookedOnly/Private/AnimNextStateTreeEditorData.cpp
+++ b/Engine/Plugins/Experimental/UAF/UAFStateTree/Source/UAFStateTreeUncookedOnly/Private/AnimNextStateTreeEditorData.cpp
@@ -5,6 +5,7 @@
 #include "UncookedOnlyUtils.h"
 #include "AnimNextStateTree_EditorData.h"
 #include "StateTreeDelegates.h"
+#include "StateTreeEditingSubsystem.h"
 
 void UAnimNextStateTreeTreeEditorData::PostLoad()
 {
@@ -59,8 +60,9 @@ void UAnimNextStateTreeTreeEditorData::HandleStateTreeAssetChanges(UAnimNextRigV
 	case EAnimNextEditorDataNotifType::VariableDefaultValueChanged:
 	case EAnimNextEditorDataNotifType::VariableBindingChanged:
 		{			
-			const UStateTree* StateTree = GetTypedOuter<UStateTree>();
-			UE::StateTree::Delegates::OnParametersChanged.Broadcast(*StateTree);	
+			UStateTree* StateTree = GetTypedOuter<UStateTree>();
+			UE::StateTree::Delegates::OnParametersChanged.Broadcast(*StateTree);
+			UStateTreeEditingSubsystem::MarkAsPubliclyModified(StateTree);
 		}		
 		break;
 	default:
diff --git a/Engine/Plugins/Runtime/StateTree/Source/StateTreeEditorModule/Private/Commandlets/StateTreeCompileAllCommandlet.cpp b/Engine/Plugins/Runtime/StateTree/Source/StateTreeEditorModule/Private/Commandlets/StateTreeCompileAllCommandlet.cpp
--- a/Engine/Plugins/Runtime/StateTree/Source/StateTreeEditorModule/Private/Commandlets/StateTreeCompileAllCommandlet.cpp
+++ b/Engine/Plugins/Runtime/StateTree/Source/StateTreeEditorModule/Private/Commandlets/StateTreeCompileAllCommandlet.cpp
@@ -75,9 +75,6 @@ bool UStateTreeCompileAllCommandlet::CompileAndSaveStateTree(TNonNullPtr<UStateT
 	const FString PackageFileName = SourceControlHelpers::PackageFilename(Package);
 
 	// Compile the StateTree asset.
-	UStateTreeEditingSubsystem::ValidateStateTree(StateTree);
-	const uint32 EditorDataHash = UStateTreeEditingSubsystem::CalculateStateTreeHash(StateTree);
-
 	FStateTreeCompilerLog Log;
 	const bool bSuccess = UStateTreeEditingSubsystem::CompileStateTree(StateTree, Log);
 
diff --git a/Engine/Plugins/Runtime/StateTree/Source/StateTreeEditorModule/Private/StateTreeCompiler.cpp b/Engine/Plugins/Runtime/StateTree/Source/StateTreeEditorModule/Private/StateTreeCompiler.cpp
--- a/Engine/Plugins/Runtime/StateTree/Source/StateTreeEditorModule/Private/StateTreeCompiler.cpp
+++ b/Engine/Plugins/Runtime/StateTree/Source/StateTreeEditorModule/Private/StateTreeCompiler.cpp
@@ -411,7 +411,7 @@ namespace UE::StateTree::Compiler
 
 }; // UE::StateTree::Compiler
 
-void FStateTreeCompiler::AppendToClassSchema(FAppendToClassSchemaContext& Context)
+void FStateTreeCompiler::AppendToStateTreeClassSchema(FAppendToClassSchemaContext& Context)
 {
 	Context.Update(&UE::StateTree::Compiler::ClassSchemaVersion, sizeof(UE::StateTree::Compiler::ClassSchemaVersion));
 
@@ -428,12 +428,15 @@ void FStateTreeCompiler::AppendToClassSchema(FAppendToClassSchemaContext& Contex
 	UpdateContext(TEXT("StateTree.Compiler.EnableUtilityConsiderationWithEvaluationScopeInstanceData"));
 }
 
-bool FStateTreeCompiler::Compile(UStateTree& InStateTree)
+
+bool FStateTreeCompiler::FailCompilation(UStateTree::ECompileStatus Status) const
 {
-	return Compile(&InStateTree);
+	StateTree->ResetCompiled();
+	StateTree->CompileStatus = Status;
+	return false;
 }
 
-bool FStateTreeCompiler::Compile(TNotNull<UStateTree*> InStateTree)
+bool FStateTreeCompiler::PreCompile(TNotNull<UStateTree*> InStateTree)
 {
 	if (bCompiled)
 	{
@@ -444,16 +447,9 @@ bool FStateTreeCompiler::Compile(TNotNull<UStateTree*> InStateTree)
 
 	StateTree = InStateTree;
 	EditorData = Cast<UStateTreeEditorData>(StateTree->EditorData);
-
-	auto FailCompilation = [StateTree=StateTree]()
-		{
-			StateTree->ResetCompiled();
-			return false;
-		};
-
 	if (!EditorData)
 	{
-		return FailCompilation();
+		return false;
 	}
 	
 	// Cleanup existing state
@@ -462,28 +458,52 @@ bool FStateTreeCompiler::Compile(TNotNull<UStateTree*> InStateTree)
 	if (!EditorData->Schema)
 	{
 		Log.Reportf(EMessageSeverity::Error, TEXT("Missing Schema. Please set valid schema in the State Tree Asset settings."));
-		return FailCompilation();
+		return false;
 	}
 
 	if (!BindingsCompiler.Init(StateTree->PropertyBindings, Log))
 	{
-		return FailCompilation();
+		return false;
 	}
 
 	EditorData->GetAllStructValues(IDToStructValue);
 
-	// Copy schema the EditorData
-	StateTree->Schema = DuplicateObject(EditorData->Schema, StateTree);
+	return true;
+}
+
+bool FStateTreeCompiler::CompilePublic(TNotNull<UStateTree*> InStateTree)
+{
+	TRACE_CPUPROFILER_EVENT_SCOPE(FStateTreeCompiler::CompilePublic)
+	if (!PreCompile(InStateTree))
+	{
+		return FailCompilation(UStateTree::ECompileStatus::PublicDirty);
+	}
 
 	if (!CreateParameters())
 	{
-		return FailCompilation();
+		return FailCompilation(UStateTree::ECompileStatus::PublicDirty);
 	}
 
+	// Public succeeded, the next step is dirty
+	StateTree->CompileStatus = UStateTree::ECompileStatus::InternalDirty;
+
+	return true;
+}
+
+bool FStateTreeCompiler::CompileInternalImpl()
+{
+	TRACE_CPUPROFILER_EVENT_SCOPE(FStateTreeCompiler::CompileInternal)
+
+	check(StateTree);
+	check(EditorData);
+
+	// Copy schema the EditorData
+	StateTree->Schema = DuplicateObject(EditorData->Schema, StateTree);
+
 	int32 ContextDataIndex = 0;
 
 	// Mark all named external values as binding source
-	if (StateTree->Schema)
+	if (ensure(StateTree->Schema))
 	{
 		StateTree->ContextDataDescs = StateTree->Schema->GetContextDataDescs();
 		for (FStateTreeExternalDataDesc& Desc : StateTree->ContextDataDescs)
@@ -502,7 +522,7 @@ bool FStateTreeCompiler::Compile(TNotNull<UStateTree*> InStateTree)
 				if (const auto Validation = UE::StateTree::Compiler::IsValidIndex16(ContextDataIndex); Validation.DidFail())
 				{
 					Validation.Log(Log, TEXT("ExternalStructIndex"), ExtDataDesc);
-					return FailCompilation();
+					return FailCompilation(UStateTree::ECompileStatus::InternalDirty);
 				}
 				Desc.Handle.DataHandle = ExtDataDesc.DataHandle;
 			}
@@ -512,13 +532,13 @@ bool FStateTreeCompiler::Compile(TNotNull<UStateTree*> InStateTree)
 	if (const UE::StateTree::Compiler::FValidationResult Validation = UE::StateTree::Compiler::IsValidIndex16(ContextDataIndex); Validation.DidFail())
 	{
 		Validation.Log(Log, TEXT("NumContextData"));
-		return FailCompilation();
+		return FailCompilation(UStateTree::ECompileStatus::InternalDirty);
 	}
 	StateTree->NumContextData = static_cast<uint16>(ContextDataIndex);
 
 	if (!CreateStates())
 	{
-		return FailCompilation();
+		return FailCompilation(UStateTree::ECompileStatus::InternalDirty);
 	}
 
 	// Eval and Global task methods use InstanceStructs.Num() as ID generator.
@@ -526,35 +546,35 @@ bool FStateTreeCompiler::Compile(TNotNull<UStateTree*> InStateTree)
 	
 	if (!CreateEvaluators())
 	{
-		return FailCompilation();
+		return FailCompilation(UStateTree::ECompileStatus::InternalDirty);
 	}
 
 	if (!CreateGlobalTasks())
 	{
-		return FailCompilation();
+		return FailCompilation(UStateTree::ECompileStatus::InternalDirty);
 	}
 
 	const int32 NumGlobalInstanceData = InstanceStructs.Num();
 	if (const auto Validation = UE::StateTree::Compiler::IsValidIndex16(NumGlobalInstanceData); Validation.DidFail())
 	{
 		Validation.Log(Log, TEXT("NumGlobalInstanceData"));
-		return FailCompilation();
+		return FailCompilation(UStateTree::ECompileStatus::InternalDirty);
 	}
 	StateTree->NumGlobalInstanceData = uint16(NumGlobalInstanceData);
 
 	if (!CreateStateTasksAndParameters())
 	{
-		return FailCompilation();
+		return FailCompilation(UStateTree::ECompileStatus::InternalDirty);
 	}
 
 	if (!CreateStateTransitions())
 	{
-		return FailCompilation();
+		return FailCompilation(UStateTree::ECompileStatus::InternalDirty);
 	}
 
 	if (!CreateStateConsiderations())
 	{
-		return FailCompilation();
+		return FailCompilation(UStateTree::ECompileStatus::InternalDirty);
 	}
 
 	StateTree->Nodes = Nodes;
@@ -568,10 +588,13 @@ bool FStateTreeCompiler::Compile(TNotNull<UStateTree*> InStateTree)
 
 	BindingsCompiler.Finalize();
 
+	// Internal succeeded, the next step is dirty
+	StateTree->CompileStatus = UStateTree::ECompileStatus::LinkDirty;
+
 	if (!StateTree->Link())
 	{
 		Log.Reportf(EMessageSeverity::Error, TEXT("Unexpected failure to link the StateTree asset. See log for more info."));
-		return FailCompilation();
+		return FailCompilation(UStateTree::ECompileStatus::LinkDirty);
 	}
 
 	// Store mapping between node unique ID and their compiled index. Used for debugging purposes.
@@ -594,11 +617,14 @@ bool FStateTreeCompiler::Compile(TNotNull<UStateTree*> InStateTree)
 
 	if (!NotifyInternalPost())
 	{
-		return FailCompilation();
+		return FailCompilation(UStateTree::ECompileStatus::InternalDirty);
 	}
 
-	UE::StateTree::Compiler::FCheckOutersArchive CheckOuters(*StateTree, *EditorData, Log);
-	StateTree->Serialize(CheckOuters);
+	{
+		// Test if there are invalid reference
+		UE::StateTree::Compiler::FCheckOutersArchive CheckOuters(*StateTree, *EditorData, Log);
+		StateTree->Serialize(CheckOuters);
+	}
 
 	if (UE::StateTree::Compiler::CVarLogCompiledStateTree->GetBool())
 	{
@@ -608,6 +634,26 @@ bool FStateTreeCompiler::Compile(TNotNull<UStateTree*> InStateTree)
 	return true;
 }
 
+bool FStateTreeCompiler::CompileInternal(TNotNull<UStateTree*> InStateTree)
+{
+	// For now, don't take any chances and also compile the public.
+	if (!CompilePublic(InStateTree))
+	{
+		return FailCompilation(UStateTree::ECompileStatus::PublicDirty);
+	}
+	return CompileInternalImpl();
+}
+
+bool FStateTreeCompiler::Compile(TNotNull<UStateTree*> InStateTree)
+{
+	return CompileInternal(InStateTree);
+}
+
+bool FStateTreeCompiler::Compile(UStateTree& InStateTree)
+{
+	return Compile(&InStateTree);
+}
+
 FStateTreeStateHandle FStateTreeCompiler::GetStateHandle(const FGuid& StateID) const
 {
 	const int32* Idx = IDToState.Find(StateID);
diff --git a/Engine/Plugins/Runtime/StateTree/Source/StateTreeEditorModule/Private/StateTreeCompilerManager.cpp b/Engine/Plugins/Runtime/StateTree/Source/StateTreeEditorModule/Private/StateTreeCompilerManager.cpp
--- a/Engine/Plugins/Runtime/StateTree/Source/StateTreeEditorModule/Private/StateTreeCompilerManager.cpp
+++ b/Engine/Plugins/Runtime/StateTree/Source/StateTreeEditorModule/Private/StateTreeCompilerManager.cpp
@@ -4,14 +4,18 @@
 #include "StateTree.h"
 #include "StateTreeCompilerLog.h"
 #include "StateTreeDelegates.h"
+#include "StateTreeDelegatesInternal.h"
 #include "StateTreeEditingSubsystem.h"
 #include "StateTreeEditorData.h"
 #include "StateTreeEditorModule.h"
 #include "StateTreeEditorPropertyBindings.h"
+#include "StateTreeSchema.h"
 
+#include "Containers/Ticker.h"
 #include "Editor.h"
 #include "StructUtilsDelegates.h"
 #include "StructUtils/UserDefinedStruct.h"
+#include "Templates/GuardValueAccessors.h"
 #include "UObject/ObjectKey.h"
 #include "UObject/UObjectGlobals.h"
 #include "UObject/UObjectIterator.h"
@@ -25,12 +29,36 @@ FAutoConsoleVariable CVarLogStateTreeDependencies(
 );
 
 bool bUseDependenciesToTriggerCompilation = true;
-FAutoConsoleVariableRef CVarLogStateTreeUseDependenciesToTriggerCompilation(
+FAutoConsoleVariableRef CVarUseDependenciesToTriggerCompilation(
 	TEXT("StateTree.Compiler.UseDependenciesToTriggerCompilation"),
 	bUseDependenciesToTriggerCompilation,
 	TEXT("Use the build dependencies to detect when a state tree needs to be linked or compiled.")
 );
 
+bool bEnableCompileSynchronouslyInPIESession = true;
+FAutoConsoleVariableRef CVarEnableCompileSynchronouslyInPIESession(
+	TEXT("StateTree.Compiler.EnableCompileSynchronouslyInPIESession"),
+	bEnableCompileSynchronouslyInPIESession,
+	TEXT("When PIE starts, flush the compilation queue.\n")
+	TEXT("Any new asset will be compile synchronously event is async was requested.")
+);
+
+int32 NumberOfQueuedCompilationPerBatch = 1;
+FAutoConsoleVariableRef CVarNumberOfQueuedCompilationPerBatch(
+	TEXT("StateTree.Compiler.NumberOfQueuedCompilationPerBatch"),
+	NumberOfQueuedCompilationPerBatch,
+	TEXT("Amount of queued state tree asset to compile per engine tick.")
+);
+
+FAutoConsoleCommand FlushCompilationQueue(
+	TEXT("StateTree.Compiler.FlushCompilationQueue"),
+	TEXT("Compile all pending state tree asset that requested a compilation."),
+	FConsoleCommandDelegate::CreateLambda([]()
+		{
+			FCompilerManager::FlushCompilationQueue();
+		})
+);
+
 struct FStateTreeDependencies
 {
 	enum EDependencyType
@@ -50,12 +78,11 @@ struct FStateTreeDependencies
 ENUM_CLASS_FLAGS(FStateTreeDependencies::EDependencyType);
 
 /** Find the references that are needed by the asset. */
-class FArchiveReferencingProperties : public FArchiveUObject
+class FArchiveReferencingPropertiesBase : public FArchiveUObject
 {
 public:
-	FArchiveReferencingProperties(TNotNull<const UObject*> InReferencingObject)
-		: ReferencingObjectPackage(InReferencingObject->GetPackage())
-		, StateTreeModulePackage(UStateTree::StaticClass()->GetOutermost())
+	FArchiveReferencingPropertiesBase()
+		: StateTreeModulePackage(UStateTree::StaticClass()->GetOutermost())
 		, StateTreeEditorModulePackage(UStateTreeEditorData::StaticClass()->GetOutermost())
 		, CoreUObjectModulePackage(UObject::StaticClass()->GetOutermost())
 	{
@@ -76,6 +103,20 @@ public:
 			&& !Struct->IsInPackage(CoreUObjectModulePackage);
 	}
 
+private:
+	TNotNull<const UPackage*> StateTreeModulePackage;
+	TNotNull<const UPackage*> StateTreeEditorModulePackage;
+	TNotNull<const UPackage*> CoreUObjectModulePackage;
+};
+
+/** Find the references that are needed by the asset. */
+class FArchiveReferencingProperties : public FArchiveReferencingPropertiesBase
+{
+public:
+	FArchiveReferencingProperties(TNotNull<const UObject*> InReferencingObject)
+		: ReferencingObjectPackage(InReferencingObject->GetPackage())
+	{ }
+
 	virtual FArchive& operator<<(UObject*& InSerializedObject) override
 	{
 		if (InSerializedObject)
@@ -118,9 +159,48 @@ private:
 	TSet<UObject*> SerializedObjects;
 
 	TNotNull<const UPackage*> ReferencingObjectPackage;
-	TNotNull<const UPackage*> StateTreeModulePackage;
-	TNotNull<const UPackage*> StateTreeEditorModulePackage;
-	TNotNull<const UPackage*> CoreUObjectModulePackage;
+};
+
+/** Find the state tree asset references that are needed by the asset. */
+class FArchiveReferencingStateTree : public FArchiveReferencingPropertiesBase
+{
+public:
+	FArchiveReferencingStateTree(TNotNull<const UStateTree*> InReferencingObject)
+		: ReferencingObjectPackage(InReferencingObject->GetPackage())
+	{
+	}
+
+	virtual FArchive& operator<<(UObject*& InSerializedObject) override
+	{
+		if (UStateTree* StateTree = Cast<UStateTree>(InSerializedObject))
+		{
+			Dependencies.Add(StateTree);
+		}
+		else if (InSerializedObject)
+		{
+			// Traversing the asset inner dependencies (instanced type).
+			if (InSerializedObject->IsInPackage(ReferencingObjectPackage))
+			{
+				bool bAlreadyExists;
+				SerializedObjects.Add(InSerializedObject, &bAlreadyExists);
+
+				if (!bAlreadyExists)
+				{
+					InSerializedObject->Serialize(*this);
+				}
+			}
+		}
+
+		return *this;
+	}
+
+	TSet<UStateTree*> Dependencies;
+
+private:
+	/** Tracks the objects which have been serialized by this archive, to prevent recursion */
+	TSet<UObject*> SerializedObjects;
+
+	TNotNull<const UPackage*> ReferencingObjectPackage;
 };
 
 class FCompilerManagerImpl
@@ -131,69 +211,350 @@ public:
 	FCompilerManagerImpl(const FCompilerManagerImpl&) = delete;
 	FCompilerManagerImpl& operator=(const FCompilerManagerImpl&) = delete;
 
+	void QueueForCompilation(TNotNull<UStateTree*> InStateTree);
+	void FlushCompilationQueue();
+
+	bool CompilePublicSynchronously(TNotNull<UStateTree*> InStateTree);
 	bool CompileInternalSynchronously(TNotNull<UStateTree*> InStateTree, FStateTreeCompilerLog& InOutLog);
+	TOptional<bool> CompileIfNeededSynchronously(TNotNull<UStateTree*> StateTree);
 
 private:
-	bool HandleCompileStateTree(UStateTree& StateTree);
+	void AddToQueue_AnyThread(TNotNull<UStateTree*> InStateTree);
 	void UpdateBindingsInstanceStructsIfNeeded(TSet<const UStruct*>& Structs, TNotNull<UStateTreeEditorData*> EditorData);
+
 	void GatherDependencies(TNotNull<UStateTree*> StateTree);
+	TSet<UStateTree*> GatherStateTreeDependencies_AnyThread(TNotNull<UStateTree*> StateTree) const;
 	void LogDependencies(TNotNull<UStateTree*> StateTree) const;
 
+	void SendToBatchCompilationQueue_AnyThread();
+	bool BatchCompilationQueue();
+	int32 CompileWithQueue_AnyThread(TNotNull<UStateTree*> StateTree, FStateTreeCompilerLog& Log);
+	int32 CompileWithQueueImpl_AnyThread(TNotNull<UStateTree*> StateTree, FStateTreeCompilerLog& Log, TSet<TNotNull<const UStateTree*>>& ProcessedAssets);
+	void CompileStateTree_AnyThread(TNotNull<UStateTree*> StateTree, FStateTreeCompilerLog& Log);
+	void ConditionalResetLastCompilationHash(TNotNull<UStateTree*> StateTree) const;
+	bool AllowQueuedCompilation(TNotNull<UStateTree*> StateTree) const;
+
+	void HandleStateTreeLoaded(TNotNull<UStateTree*> StateTree);
+	void HandleStateTreeCooking(TNotNull<UStateTree*> StateTree);
 	using FReplacementObjectMap = TMap<UObject*, UObject*>;
 	void HandleObjectsReinstanced(const FReplacementObjectMap& ObjectMap);
 	void HandlePreBeginPIE(bool bIsSimulating);
+	void HandleEndPIE(bool bIsSimulating);
 	void HandleUserDefinedStructReinstanced(const UUserDefinedStruct& UserDefinedStruct);
 
 private:
 	FDelegateHandle ObjectsReinstancedHandle;
 	FDelegateHandle UserDefinedStructReinstancedHandle;
 	FDelegateHandle PreBeginPIEHandle;
+	FDelegateHandle EndPIEHandle;
 
+	/** Referencee to dependencies. Generated with GatherDependencies and used when a dependencies is "recreated" (blueprint compilation). */
 	TMap<TObjectKey<UStateTree>, TSharedPtr<FStateTreeDependencies>> StateTreeToDependencies;
 	TMap<FObjectKey, TArray<TObjectKey<UStateTree>>> DependenciesToStateTree;
+
+	using FDirtyAssetArray = TArray<TObjectKey<UStateTree>, TInlineAllocator<4>>;
+	FCriticalSection QueueLock;
+	FDirtyAssetArray DirtyAssets_AnyThread;
+	bool bPushedBatchOnGameThread_AnyThread = false;
+	std::atomic<bool> bCompileWithQueueInProgress_AnyThread = false;
+
+	bool bCanQueueCompilation = true;
 };
 static TUniquePtr<FCompilerManagerImpl> CompilerManagerImpl;
 
 FCompilerManagerImpl::FCompilerManagerImpl()
 {
-	UE::StateTree::Delegates::OnRequestCompile.BindRaw(this, &FCompilerManagerImpl::HandleCompileStateTree);
+	UE::StateTree::Delegates::Private::OnPreCookStateTreeAsset.BindRaw(this, &FCompilerManagerImpl::HandleStateTreeCooking);
 	ObjectsReinstancedHandle = FCoreUObjectDelegates::OnObjectsReinstanced.AddRaw(this, &FCompilerManagerImpl::HandleObjectsReinstanced);
 	UserDefinedStructReinstancedHandle = UE::StructUtils::Delegates::OnUserDefinedStructReinstanced.AddRaw(this, &FCompilerManagerImpl::HandleUserDefinedStructReinstanced);
 	PreBeginPIEHandle = FEditorDelegates::PreBeginPIE.AddRaw(this, &FCompilerManagerImpl::HandlePreBeginPIE);
+	EndPIEHandle = FEditorDelegates::EndPIE.AddRaw(this, &FCompilerManagerImpl::HandleEndPIE);
+
+	// When running in standalone
+	bCanQueueCompilation = !IsRunningGame();
 }
 
 FCompilerManagerImpl::~FCompilerManagerImpl()
 {
+	FEditorDelegates::EndPIE.Remove(EndPIEHandle);
 	FEditorDelegates::PreBeginPIE.Remove(PreBeginPIEHandle);
 	UE::StructUtils::Delegates::OnUserDefinedStructReinstanced.Remove(UserDefinedStructReinstancedHandle);
 	FCoreUObjectDelegates::OnObjectsReinstanced.Remove(ObjectsReinstancedHandle);
-	UE::StateTree::Delegates::OnRequestCompile.Unbind();
+	UE::StateTree::Delegates::Private::OnPreCookStateTreeAsset.Unbind();
+}
+
+bool FCompilerManagerImpl::CompilePublicSynchronously(TNotNull<UStateTree*> StateTree)
+{
+	check(IsInGameThread() || IsInAsyncLoadingThread());
+
+	UStateTreeEditingSubsystem::ValidateStateTree(StateTree);
+
+	StateTree->CompileStatus = UStateTree::ECompileStatus::PublicDirty;
+
+	FStateTreeCompilerLog Log;
+	FStateTreeCompiler Compiler(Log);
+	const bool bResult = Compiler.CompilePublic(StateTree);
+	if (bResult)
+	{
+		check(StateTree->CompileStatus == UStateTree::ECompileStatus::InternalDirty);
+
+		// Queue the internal steps
+		UE::TScopeLock _(QueueLock);
+		AddToQueue_AnyThread(StateTree);
+	}
+	return bResult;
 }
 
 bool FCompilerManagerImpl::CompileInternalSynchronously(TNotNull<UStateTree*> StateTree, FStateTreeCompilerLog& Log)
 {
+	check(IsInGameThread() || IsInAsyncLoadingThread());
+
 	UStateTreeEditingSubsystem::ValidateStateTree(StateTree);
+
+	{
+		UE::TScopeLock _(QueueLock);
+		const int32 DirtyAssetsIndex = DirtyAssets_AnyThread.IndexOfByKey(TObjectKey<UStateTree>(StateTree));
+		if (DirtyAssetsIndex != INDEX_NONE)
+		{
+			DirtyAssets_AnyThread.RemoveAtSwap(DirtyAssetsIndex);
+		}
+		CompileWithQueue_AnyThread(StateTree, Log);
+	}
+
+	return StateTree->CompileStatus == UStateTree::ECompileStatus::Executable;
+}
+
+TOptional<bool> FCompilerManagerImpl::CompileIfNeededSynchronously(TNotNull<UStateTree*> StateTree)
+{
+	check(IsInGameThread() || IsInAsyncLoadingThread());
+
+	if (StateTree->CompileStatus == UStateTree::ECompileStatus::PublicDirty
+		|| StateTree->CompileStatus == UStateTree::ECompileStatus::InternalDirty)
+	{
+		// The last compilation failed and we are trying to compile again with the same data (nothing changed). Do not compile.
+		//The user should trigger the compilation manually from the editor.
+		const uint32 CurrentHash = UStateTreeEditingSubsystem::CalculateStateTreeHash(StateTree);
+		if (CurrentHash == StateTree->LastCompilationEditorDataHash)
+		{
+			return {};
+		}
+
+		FStateTreeCompilerLog Log;
+		const bool bResult = CompileInternalSynchronously(StateTree, Log);
+		return TOptional<bool>(bResult);
+	}
+	return {};
+}
+
+void FCompilerManagerImpl::QueueForCompilation(TNotNull<UStateTree*> StateTree)
+{
+	check(IsInGameThread() || IsInAsyncLoadingThread());
+
+	// Asset that runs on any thread (ex: mass), can't be queued.
+	//They might be compiled via CompileIfNeededSynchronously on a worker thread.
+	if (bCanQueueCompilation && AllowQueuedCompilation(StateTree))
+	{
+		CompilePublicSynchronously(StateTree);
+	}
+	else
+	{
+		FStateTreeCompilerLog Log;
+		CompileInternalSynchronously(StateTree, Log);
+	}
+}
+
+void FCompilerManagerImpl::FlushCompilationQueue()
+{
+	check(IsInGameThread());
+
+	UE::TScopeLock _(QueueLock);
+	while (!DirtyAssets_AnyThread.IsEmpty())
+	{
+		TObjectKey<UStateTree> StateTreeKey = DirtyAssets_AnyThread.Pop();
+		UStateTree* StateTree = StateTreeKey.ResolveObjectPtr();
+		if (StateTree == nullptr)
+		{
+			continue;
+		}
+
+		FStateTreeCompilerLog Log;
+		CompileWithQueue_AnyThread(StateTree, Log);
+	}
+	bPushedBatchOnGameThread_AnyThread = !DirtyAssets_AnyThread.IsEmpty();
+}
+
+void FCompilerManagerImpl::ConditionalResetLastCompilationHash(TNotNull<UStateTree*> StateTree) const
+{
+	if (StateTree->CompileStatus == UStateTree::ECompileStatus::PublicDirty
+		|| StateTree->CompileStatus == UStateTree::ECompileStatus::InternalDirty)
+	{
+		// The last compilation failed.
+		//In CompileIfNeededSynchronously, it will trigger a compilation only if the hash has also changed.
+		//Reset the LastCompilationEditorDataHash, to trigger a new compilation.
+		StateTree->LastCompilationEditorDataHash = 0;
+	}
+}
+
+bool FCompilerManagerImpl::AllowQueuedCompilation(TNotNull<UStateTree*> StateTree) const
+{
+	if (UStateTreeEditorData* EditorData = Cast<UStateTreeEditorData>(StateTree->EditorData))
+	{
+		if (ensure(EditorData->Schema))
+		{
+			return EditorData->Schema->AllowQueuedCompilation();
+		}
+	}
+	return true;
+}
+
+void FCompilerManagerImpl::AddToQueue_AnyThread(TNotNull<UStateTree*> StateTree)
+{
+	TObjectKey<UStateTree> StateTreeKey = StateTree;
+	DirtyAssets_AnyThread.AddUnique(StateTreeKey);
+
+	SendToBatchCompilationQueue_AnyThread();
+}
+
+void FCompilerManagerImpl::SendToBatchCompilationQueue_AnyThread()
+{
+	if (!bPushedBatchOnGameThread_AnyThread && NumberOfQueuedCompilationPerBatch > 0)
+	{
+		bPushedBatchOnGameThread_AnyThread = true;
+		FTSTicker::GetCoreTicker().AddTicker(UE_SOURCE_LOCATION, UE_SMALL_NUMBER,
+			[](float)
+			{
+				if (CompilerManagerImpl.IsValid())
+				{
+					const bool bContinue = CompilerManagerImpl->BatchCompilationQueue();
+					return bContinue;
+				}
+				return false;
+			}
+		);
+	}
+}
+
+bool FCompilerManagerImpl::BatchCompilationQueue()
+{
+	// If locked for compilation... wait for next frame to compile the dirty assets.
+	if (bCompileWithQueueInProgress_AnyThread)
+	{
+		constexpr bool bContinue = true;
+		return bContinue;
+	}
+
+	UE::TScopeLock _(QueueLock);
+
+	while (!DirtyAssets_AnyThread.IsEmpty())
+	{
+		TObjectKey<UStateTree> StateTreeKey = DirtyAssets_AnyThread.Pop();
+		UStateTree* StateTree = StateTreeKey.ResolveObjectPtr();
+		if (StateTree == nullptr)
+		{
+			continue;
+		}
+
+		FStateTreeCompilerLog Log;
+		const int32 CompileCounter = CompileWithQueue_AnyThread(StateTree, Log);
+		if (CompileCounter >= NumberOfQueuedCompilationPerBatch)
+		{
+			break;
+		}
+	}
+	const bool bContinue = !DirtyAssets_AnyThread.IsEmpty();
+	bPushedBatchOnGameThread_AnyThread = bContinue;
+	return bContinue;
+}
+
+int32 FCompilerManagerImpl::CompileWithQueue_AnyThread(TNotNull<UStateTree*> StateTree, FStateTreeCompilerLog& Log)
+{
+	bCompileWithQueueInProgress_AnyThread = true;
+
+	TSet<TNotNull<const UStateTree*>> ProcessedAssets;
+	ProcessedAssets.Add(StateTree);
+	int32 CompiledStateTreeCount = CompileWithQueueImpl_AnyThread(StateTree, Log, ProcessedAssets);
+
+	bCompileWithQueueInProgress_AnyThread = false;
+
+	return CompiledStateTreeCount;
+}
+
+int32 FCompilerManagerImpl::CompileWithQueueImpl_AnyThread(TNotNull<UStateTree*> StateTree, FStateTreeCompilerLog& Log, TSet<TNotNull<const UStateTree*>>& ProcessedAssets)
+{
+	// State tree asset may depends on each other via linked asset or special task like FStateTreeRunParallelStateTreeTask.
+	//The assets should be compiled in order. The referencer asset reference the global parameters of the referencee asset.
+	//Depending if cvar bEnableQueuedCompilationWhenAssetDirty is enabled, the "public exported" can always be up-to-date.
+	//We don't take any chances. We sort them base on their dependencies.
+
+	int32 CompiledStateTreeCount = 0;
+	if (DirtyAssets_AnyThread.Num() == 0)
+	{
+		CompileStateTree_AnyThread(StateTree, Log);
+		++CompiledStateTreeCount;
+	}
+	else
+	{
+		const TSet<UStateTree*> Dependencies = GatherStateTreeDependencies_AnyThread(StateTree);
+		for (UStateTree* Dependency : Dependencies)
+		{
+			if (ProcessedAssets.Contains(Dependency))
+			{
+				UE_LOG(LogStateTree, Warning, TEXT("Circular dependency detected while processing %s and %s."), *StateTree->GetFullName(), *Dependency->GetFullName());
+			}
+			else
+			{
+				const int32 DirtyAssetsIndex = DirtyAssets_AnyThread.IndexOfByKey(TObjectKey<UStateTree>(Dependency));
+				if (DirtyAssetsIndex != INDEX_NONE)
+				{
+					FStateTreeCompilerLog NewLog;
+					ProcessedAssets.Add(Dependency);
+					CompiledStateTreeCount += CompileWithQueueImpl_AnyThread(Dependency, NewLog, ProcessedAssets);
+					DirtyAssets_AnyThread.RemoveAtSwap(DirtyAssetsIndex);
+				}
+			}
+		}
+
+		CompileStateTree_AnyThread(StateTree, Log);
+		++CompiledStateTreeCount;
+	}
+	return CompiledStateTreeCount;
+}
+
+void FCompilerManagerImpl::CompileStateTree_AnyThread(TNotNull<UStateTree*> StateTree, FStateTreeCompilerLog& Log)
+{
 	FStateTreeCompiler Compiler(Log);
 	const bool bCompilationResult = Compiler.Compile(StateTree);
+
+	const uint32 NewHash = UStateTreeEditingSubsystem::CalculateStateTreeHash(StateTree);
+	StateTree->LastCompilationEditorDataHash = NewHash;
+	checkf(StateTree->LastCompiledEditorDataHash == 0, TEXT("The compilation has to reset the statetree data. Including the hash value."));
+
 	if (bCompilationResult)
 	{
-		const uint32 EditorDataHash = UStateTreeEditingSubsystem::CalculateStateTreeHash(StateTree);
-
-		// Success
-		StateTree->LastCompiledEditorDataHash = EditorDataHash;
+		StateTree->LastCompiledEditorDataHash = NewHash;
 		UE_LOG(LogStateTreeEditor, Log, TEXT("Compile StateTree '%s' succeeded."), *StateTree->GetFullName());
 	}
 	else
 	{
-		// Make sure not to leave stale data on failed compile.
-		StateTree->ResetCompiled();
-		StateTree->LastCompiledEditorDataHash = 0;
-
 		UE_LOG(LogStateTreeEditor, Error, TEXT("Failed to compile '%s', errors follow."), *StateTree->GetFullName());
-		Log.DumpToLog(LogStateTreeEditor);
 	}
 
-	UE::StateTree::Delegates::OnPostCompile.Broadcast(*StateTree);
+	Log.DumpToLog(LogStateTreeEditor);
+
+	if (IsInGameThread())
+	{
+		UE::StateTree::Delegates::OnPostCompile.Broadcast(*StateTree);
+	}
+	else
+	{
+		ExecuteOnGameThread(UE_SOURCE_LOCATION, [WeakStateTree = TWeakObjectPtr<UStateTree>(StateTree)]()
+			{
+				if (UStateTree* StateTree = WeakStateTree.Get())
+				{
+					UE::StateTree::Delegates::OnPostCompile.Broadcast(*StateTree);
+				}
+			});
+	}
 
 	GatherDependencies(StateTree);
 
@@ -201,8 +562,6 @@ bool FCompilerManagerImpl::CompileInternalSynchronously(TNotNull<UStateTree*> St
 	{
 		LogDependencies(StateTree);
 	}
-
-	return bCompilationResult;
 }
 
 void FCompilerManagerImpl::UpdateBindingsInstanceStructsIfNeeded(TSet<const UStruct*>& Structs, TNotNull<UStateTreeEditorData*> EditorData)
@@ -229,19 +588,37 @@ void FCompilerManagerImpl::UpdateBindingsInstanceStructsIfNeeded(TSet<const UStr
 	}
 }
 
-bool FCompilerManagerImpl::HandleCompileStateTree(UStateTree& StateTree)
+void FCompilerManagerImpl::HandleStateTreeCooking(TNotNull<UStateTree*> StateTree)
 {
-	FStateTreeCompilerLog Log;
-	return CompileInternalSynchronously(&StateTree, Log);
+	switch (StateTree->CompileStatus)
+	{
+		case UStateTree::ECompileStatus::PublicDirty:
+		case UStateTree::ECompileStatus::InternalDirty:
+		{
+			FStateTreeCompilerLog Log;
+			CompileInternalSynchronously(StateTree, Log);
+			break;
+		}
+		case UStateTree::ECompileStatus::LinkDirty:
+			std::ignore = StateTree->Link();
+			break;
+	}
+
+	if (StateTree->CompileStatus != UStateTree::ECompileStatus::Executable)
+	{
+		UE_LOG(LogStateTree, Error, TEXT("StateTree asset '%s' failed compilation."), *StateTree->GetPathName());
+	}
 }
 
 void FCompilerManagerImpl::HandlePreBeginPIE(const bool bIsSimulating)
 {
-	for (TObjectIterator<UStateTree> It; It; ++It)
-	{
-		check(!It->HasAnyFlags(RF_ClassDefaultObject));
-		It->CompileIfChanged();
-	}
+	bCanQueueCompilation = !Private::bEnableCompileSynchronouslyInPIESession;
+	FlushCompilationQueue();
+}
+
+void FCompilerManagerImpl::HandleEndPIE(const bool bIsSimulating)
+{
+	bCanQueueCompilation = Private::bEnableCompileSynchronouslyInPIESession;
 }
 
 void FCompilerManagerImpl::HandleObjectsReinstanced(const FReplacementObjectMap& ObjectMap)
@@ -251,7 +628,7 @@ void FCompilerManagerImpl::HandleObjectsReinstanced(const FReplacementObjectMap&
 		return;
 	}
 
-	TArray<const UObject*> ObjectsToBeReplaced;
+	TArray<const UObject*, TInlineAllocator<8>> ObjectsToBeReplaced;
 	ObjectsToBeReplaced.Reserve(ObjectMap.Num());
 	for (TMap<UObject*, UObject*>::TConstIterator It(ObjectMap); It; ++It)
 	{
@@ -276,12 +653,17 @@ void FCompilerManagerImpl::HandleObjectsReinstanced(const FReplacementObjectMap&
 		}
 	}
 
-	if (UE::StateTree::Compiler::Private::bUseDependenciesToTriggerCompilation)
+	if (bUseDependenciesToTriggerCompilation)
 	{
-		TArray<TNotNull<UStateTree*>> StateTreeToLink;
+		TArray<TNotNull<UStateTree*>, TInlineAllocator<8>> StateTreeToLink;
+		TArray<FObjectKey, TInlineAllocator<8>> StructsToBeReplacedAsKey;
+		StructsToBeReplacedAsKey.Reserve(StructsToBeReplaced.Num());
+
 		for (const UStruct* StructToBeReplaced : StructsToBeReplaced)
 		{
 			const FObjectKey StructToReplacedKey = StructToBeReplaced;
+			StructsToBeReplacedAsKey.Add(StructToReplacedKey);
+
 			TArray<TObjectKey<UStateTree>>* Dependencies = DependenciesToStateTree.Find(StructToReplacedKey);
 			if (Dependencies)
 			{
@@ -297,8 +679,29 @@ void FCompilerManagerImpl::HandleObjectsReinstanced(const FReplacementObjectMap&
 
 		for (UStateTree* StateTree : StateTreeToLink)
 		{
+			ConditionalResetLastCompilationHash(StateTree);
 			if (!StateTree->Link())
 			{
+				//If the task modification make it not linkable, then we need to recompile the asset.
+				bool bModifyPublic = true;
+				const TObjectKey<UStateTree> StateTreeKey = StateTree;
+				if (TSharedPtr<FStateTreeDependencies>* FoundDependencies = StateTreeToDependencies.Find(StateTreeKey))
+				{
+					bModifyPublic = false;
+					for (const FStateTreeDependencies::FItem& Dependency : (*FoundDependencies)->Dependencies)
+					{
+						if (StructsToBeReplacedAsKey.Contains(Dependency.Key))
+						{
+							if (EnumHasAnyFlags(Dependency.Type, FStateTreeDependencies::EDependencyType::DT_Public))
+							{
+								bModifyPublic = true;
+								break;
+							}
+						}
+					}
+				}
+
+				StateTree->MarkAsModified(bModifyPublic);
 				UE_LOG(LogStateTree, Error, TEXT("%s failed to link after Object reinstantiation. Take a look at the asset for any errors. Asset will not be usable at runtime."), *StateTree->GetPathName());
 			}
 		}
@@ -334,7 +737,7 @@ void FCompilerManagerImpl::HandleObjectsReinstanced(const FReplacementObjectMap&
 			}
 
 			// If the asset is not linked yet (or has failed), no need to link.
-			if (!bShouldRelink && !StateTree->bIsLinked)
+			if (!bShouldRelink && StateTree->CompileStatus != UStateTree::ECompileStatus::Executable)
 			{
 				continue;
 			}
@@ -360,6 +763,7 @@ void FCompilerManagerImpl::HandleObjectsReinstanced(const FReplacementObjectMap&
 
 			if (bShouldRelink)
 			{
+				ConditionalResetLastCompilationHash(StateTree);
 				if (!StateTree->Link())
 				{
 					UE_LOG(LogStateTree, Error, TEXT("%s failed to link after Object reinstantiation. Take a look at the asset for any errors. Asset will not be usable at runtime."), *StateTree->GetPathName());
@@ -389,6 +793,7 @@ void FCompilerManagerImpl::HandleUserDefinedStructReinstanced(const UUserDefined
 
 		for (UStateTree* StateTree : StateTreeToLink)
 		{
+			ConditionalResetLastCompilationHash(StateTree);
 			if (!StateTree->Link())
 			{
 				UE_LOG(LogStateTree, Error, TEXT("%s failed to link after Object reinstantiation. Take a look at the asset for any errors. Asset will not be usable at runtime."), *StateTree->GetPathName());
@@ -411,6 +816,7 @@ void FCompilerManagerImpl::HandleUserDefinedStructReinstanced(const UUserDefined
 			UStateTree* StateTree = *It;
 			if (StateTree->PropertyBindings.ContainsAnyStruct(Structs))
 			{
+				ConditionalResetLastCompilationHash(StateTree);
 				if (!StateTree->Link())
 				{
 					UE_LOG(LogStateTree, Error, TEXT("%s failed to link after Struct reinstantiation. Take a look at the asset for any errors. Asset will not be usable at runtime."), *StateTree->GetPathName());
@@ -463,8 +869,7 @@ void FCompilerManagerImpl::GatherDependencies(TNotNull<UStateTree*> StateTree)
 		};
 
 	// Gather new inner dependencies
-	UStateTreeEditorData* EditorData = CastChecked<UStateTreeEditorData>(StateTree->EditorData);
-	if (EditorData)
+	if (UStateTreeEditorData* EditorData = Cast<UStateTreeEditorData>(StateTree->EditorData))
 	{
 		// Internal
 		{
@@ -477,6 +882,8 @@ void FCompilerManagerImpl::GatherDependencies(TNotNull<UStateTree*> StateTree)
 		}
 		// Public
 		{
+			//The root parameters depend on the editor data type. Serialize it independently in case it is not discoverable by the editor data.
+			//It has public dependency because other state tree asset or state tree reference can depends on it.
 			FArchiveReferencingProperties DependencyArchive(StateTree);
 			const_cast<FInstancedPropertyBag&>(EditorData->GetRootParametersPropertyBag()).Serialize(DependencyArchive);
 			for (const UStruct* Dependency : DependencyArchive.Dependencies)
@@ -526,6 +933,19 @@ void FCompilerManagerImpl::GatherDependencies(TNotNull<UStateTree*> StateTree)
 	}
 }
 
+TSet<UStateTree*> FCompilerManagerImpl::GatherStateTreeDependencies_AnyThread(TNotNull<UStateTree*> StateTree) const
+{
+	FArchiveReferencingStateTree DependencyArchive(StateTree);
+	UStateTreeEditorData* EditorData = Cast<UStateTreeEditorData>(StateTree->EditorData);
+	if (EditorData)
+	{
+		EditorData->Serialize(DependencyArchive);
+		//The root parameters depend on the editor data type. Serialize it independently in case it is not discoverable by the editor data.
+		const_cast<FInstancedPropertyBag&>(EditorData->GetRootParametersPropertyBag()).Serialize(DependencyArchive);
+	}
+	return DependencyArchive.Dependencies;
+}
+
 void FCompilerManagerImpl::LogDependencies(TNotNull<UStateTree*> StateTree) const
 {
 	FStringBuilderBase LogString;
@@ -606,6 +1026,24 @@ void FCompilerManager::Shutdown()
 	Private::CompilerManagerImpl.Reset();
 }
 
+void FCompilerManager::QueueForCompilation(TNotNull<UStateTree*> StateTree)
+{
+	TRACE_CPUPROFILER_EVENT_SCOPE(FCompilerManager::QueueForCompilation)
+	if (ensureMsgf(Private::CompilerManagerImpl.IsValid(), TEXT("Can't queue the asset when the module is not available.")))
+	{
+		Private::CompilerManagerImpl->QueueForCompilation(StateTree);
+	}
+}
+
+void FCompilerManager::FlushCompilationQueue()
+{
+	TRACE_CPUPROFILER_EVENT_SCOPE(FCompilerManager::FlushCompilationQueue)
+	if (Private::CompilerManagerImpl.IsValid())
+	{
+		Private::CompilerManagerImpl->FlushCompilationQueue();
+	}
+}
+
 bool FCompilerManager::CompileSynchronously(TNotNull<UStateTree*> StateTree)
 {
 	FStateTreeCompilerLog Log;
@@ -614,6 +1052,7 @@ bool FCompilerManager::CompileSynchronously(TNotNull<UStateTree*> StateTree)
 
 bool FCompilerManager::CompileSynchronously(TNotNull<UStateTree*> StateTree, FStateTreeCompilerLog& Log)
 {
+	TRACE_CPUPROFILER_EVENT_SCOPE(FCompilerManager::CompileSynchronously)
 	if (ensureMsgf(Private::CompilerManagerImpl.IsValid(), TEXT("Can't compile the asset when the module is not available.")))
 	{
 		return Private::CompilerManagerImpl->CompileInternalSynchronously(StateTree, Log);
@@ -621,5 +1060,16 @@ bool FCompilerManager::CompileSynchronously(TNotNull<UStateTree*> StateTree, FSt
 	return false;
 }
 
+
+TOptional<bool> FCompilerManager::CompileIfNeededSynchronously(TNotNull<UStateTree*> StateTree)
+{
+	TRACE_CPUPROFILER_EVENT_SCOPE(FCompilerManager::CompileIfNeededSynchronously)
+	if (ensureMsgf(Private::CompilerManagerImpl.IsValid(), TEXT("Can't compile the asset when the module is not available.")))
+	{
+		return Private::CompilerManagerImpl->CompileIfNeededSynchronously(StateTree);
+	}
+	return {};
+}
+
 } // UE::StateTree::Compiler
 
diff --git a/Engine/Plugins/Runtime/StateTree/Source/StateTreeEditorModule/Private/StateTreeEditingSubsystem.cpp b/Engine/Plugins/Runtime/StateTree/Source/StateTreeEditorModule/Private/StateTreeEditingSubsystem.cpp
--- a/Engine/Plugins/Runtime/StateTree/Source/StateTreeEditorModule/Private/StateTreeEditingSubsystem.cpp
+++ b/Engine/Plugins/Runtime/StateTree/Source/StateTreeEditorModule/Private/StateTreeEditingSubsystem.cpp
@@ -18,6 +18,15 @@
 
 #include UE_INLINE_GENERATED_CPP_BY_NAME(StateTreeEditingSubsystem)
 
+namespace UE::StateTree::Editor
+{
+	bool bEnableQueuedCompilationWhenAssetDirty = false;
+	FAutoConsoleVariableRef CVarbEnableQueuedCompilationWhenAssetDirty(
+		TEXT("StateTree.Compiler.EnableQueuedCompilationWhenAssetDirty"),
+		bEnableQueuedCompilationWhenAssetDirty,
+		TEXT("When the asset is modified, automatically queue the asset compilation.")
+	);
+}
 
 UStateTreeEditingSubsystem::UStateTreeEditingSubsystem()
 {
@@ -37,6 +46,32 @@ bool UStateTreeEditingSubsystem::CompileStateTree(TNotNull<UStateTree*> InStateT
 	return UE::StateTree::Compiler::FCompilerManager::CompileSynchronously(InStateTree, InOutLog);
 }
 
+bool UStateTreeEditingSubsystem::NeedsRecompile(TNotNull<const UStateTree*> InStateTree)
+{
+	const uint32 CurrentHash = CalculateStateTreeHash(InStateTree);
+	return InStateTree->LastCompiledEditorDataHash != CurrentHash;
+}
+
+void UStateTreeEditingSubsystem::MarkAsPubliclyModified(TNotNull<UStateTree*> InStateTree)
+{
+	constexpr bool bPubliclyModified = true;
+	InStateTree->MarkAsModified(bPubliclyModified);
+	if (UE::StateTree::Editor::bEnableQueuedCompilationWhenAssetDirty)
+	{
+		UE::StateTree::Compiler::FCompilerManager::QueueForCompilation(InStateTree);
+	}
+}
+
+void UStateTreeEditingSubsystem::MarkAsModified(TNotNull<UStateTree*> InStateTree)
+{
+	constexpr bool bPubliclyModified = false;
+	InStateTree->MarkAsModified(bPubliclyModified);
+	if (UE::StateTree::Editor::bEnableQueuedCompilationWhenAssetDirty)
+	{
+		UE::StateTree::Compiler::FCompilerManager::QueueForCompilation(InStateTree);
+	}
+}
+
 TSharedPtr<FStateTreeViewModel> UStateTreeEditingSubsystem::FindViewModel(TNotNull<const UStateTree*> InStateTree) const
 {
 	const FObjectKey StateTreeKey = InStateTree;
@@ -88,6 +123,8 @@ TSharedRef<SWidget> UStateTreeEditingSubsystem::GetStateTreeView(TSharedRef<FSta
 
 void UStateTreeEditingSubsystem::ValidateStateTree(TNotNull<UStateTree*> InStateTree)
 {
+	TRACE_CPUPROFILER_EVENT_SCOPE(UStateTreeEditingSubsystem::ValidateStateTree)
+
 	auto FixChangedStateLinkName = [](FStateTreeStateLink& StateLink, const TMap<FGuid, FName>& IDToName) -> bool
 		{
 			if (StateLink.ID.IsValid())
@@ -284,7 +321,6 @@ void UStateTreeEditingSubsystem::ValidateStateTree(TNotNull<UStateTree*> InState
 		{
 			FStateTreeEditorModule::GetModule().OnPreValidateStateTree().Broadcast(InStateTree);
 
-			// Schema validation
 			if (EditorData->EditorSchema)
 			{
 				EditorData->EditorSchema->PreValidate(InStateTree);
@@ -314,7 +350,6 @@ void UStateTreeEditingSubsystem::ValidateStateTree(TNotNull<UStateTree*> InState
 		{
 			FStateTreeEditorModule::GetModule().OnValidateStateTree().Broadcast(InStateTree);
 
-			// Schema validation
 			if (EditorData->EditorSchema)
 			{
 				EditorData->EditorSchema->Validate(InStateTree);
diff --git a/Engine/Plugins/Runtime/StateTree/Source/StateTreeEditorModule/Private/StateTreeEditorData.cpp b/Engine/Plugins/Runtime/StateTree/Source/StateTreeEditorModule/Private/StateTreeEditorData.cpp
--- a/Engine/Plugins/Runtime/StateTree/Source/StateTreeEditorModule/Private/StateTreeEditorData.cpp
+++ b/Engine/Plugins/Runtime/StateTree/Source/StateTreeEditorModule/Private/StateTreeEditorData.cpp
@@ -5,6 +5,7 @@
 #include "StateTreeConditionBase.h"
 #include "StateTreeConsiderationBase.h"
 #include "StateTreeDelegates.h"
+#include "StateTreeEditingSubsystem.h"
 #include "StateTreeEditorDataExtension.h"
 #include "StateTreeEditorModule.h"
 #include "StateTreeEditorSchema.h"
@@ -203,7 +204,7 @@ void UStateTreeEditorData::PostEditChangeChainProperty(FPropertyChangedChainEven
 
 	if (MemberProperty && Property)
 	{
-		const UStateTree* StateTree = GetTypedOuter<UStateTree>();
+		UStateTree* StateTree = GetTypedOuter<UStateTree>();
 		checkf(StateTree, TEXT("UStateTreeEditorData should only be allocated within a UStateTree"));
 		
 		const FName MemberName = MemberProperty->GetFName();
@@ -215,10 +216,36 @@ void UStateTreeEditorData::PostEditChangeChainProperty(FPropertyChangedChainEven
 				Extensions.Reset();
 			}
 			UE::StateTree::Delegates::OnSchemaChanged.Broadcast(*StateTree);
+			UStateTreeEditingSubsystem::MarkAsModified(StateTree);
 		}
 		else if (MemberName == GET_MEMBER_NAME_CHECKED(UStateTreeEditorData, RootParameterPropertyBag))
 		{
 			UE::StateTree::Delegates::OnParametersChanged.Broadcast(*StateTree);
+			UStateTreeEditingSubsystem::MarkAsPubliclyModified(StateTree);
+		}
+		// Notify that the color data has changed and fix existing data
+		else if (MemberName == GET_MEMBER_NAME_CHECKED(UStateTreeEditorData, Colors))
+		{
+			if (Colors.IsEmpty())
+			{
+				// Add default color
+				Colors.Add(UE::StateTree::Editor::CreateDefaultColor());
+			}
+			VisitHierarchy([Self = this](UStateTreeState& State, UStateTreeState* ParentState)
+				{
+					if (!Self->FindColor(State.ColorRef))
+					{
+						State.Modify();
+						State.ColorRef = FStateTreeEditorColorRef();
+					}
+					return EStateTreeVisitor::Continue;
+				});
+
+			UE::StateTree::Delegates::OnVisualThemeChanged.Broadcast(*StateTree);
+		}
+		else
+		{
+			UStateTreeEditingSubsystem::MarkAsModified(StateTree);
 		}
 
 		// Ensure unique ID on duplicated items.
@@ -263,27 +290,6 @@ void UStateTreeEditorData::PostEditChangeChainProperty(FPropertyChangedChainEven
 		{
 			UE::StateTree::Delegates::OnGlobalDataChanged.Broadcast(*StateTree);
 		}
-
-		// Notify that the color data has changed and fix existing data
-		if (MemberName == GET_MEMBER_NAME_CHECKED(UStateTreeEditorData, Colors))
-		{
-			if (Colors.IsEmpty())
-			{
-				// Add default color
-				Colors.Add(UE::StateTree::Editor::CreateDefaultColor());
-			}
-			VisitHierarchy([Self=this](UStateTreeState& State, UStateTreeState* ParentState)
-			{
-				if (!Self->FindColor(State.ColorRef))
-				{
-					State.Modify();
-					State.ColorRef = FStateTreeEditorColorRef();
-				}
-				return EStateTreeVisitor::Continue;
-			});
-
-			UE::StateTree::Delegates::OnVisualThemeChanged.Broadcast(*StateTree);
-		}
 	}
 
 	UE::StateTree::PropertyHelpers::DispatchPostEditToNodes(*this, PropertyChangedEvent, *this);
@@ -742,6 +748,13 @@ void UStateTreeEditorData::ReparentStates()
 	});
 }
 
+void UStateTreeEditorData::CreateRootProperties(TArrayView<UE::PropertyBinding::FPropertyCreationDescriptor> InOutCreationDescs)
+{
+	UE::PropertyBinding::CreateUniquelyNamedPropertiesInPropertyBag(InOutCreationDescs, RootParameterPropertyBag);
+	UStateTree* StateTree = GetTypedOuter<UStateTree>();
+	UStateTreeEditingSubsystem::MarkAsPubliclyModified(StateTree);
+}
+
 void UStateTreeEditorData::FixObjectInstances(TSet<UObject*>& SeenObjects, UObject& Outer, FStateTreeEditorNode& Node)
 {
 	auto NewInstances = [this, &Node, &Outer](const UStruct* Struct, const UStruct* ExecutionRuntimeStruct)
diff --git a/Engine/Plugins/Runtime/StateTree/Source/StateTreeEditorModule/Private/StateTreeEditorModule.cpp b/Engine/Plugins/Runtime/StateTree/Source/StateTreeEditorModule/Private/StateTreeEditorModule.cpp
--- a/Engine/Plugins/Runtime/StateTree/Source/StateTreeEditorModule/Private/StateTreeEditorModule.cpp
+++ b/Engine/Plugins/Runtime/StateTree/Source/StateTreeEditorModule/Private/StateTreeEditorModule.cpp
@@ -39,6 +39,7 @@
 #include "StateTreeEditorStyle.h"
 #include "StateTreeNodeClassCache.h"
 #include "StateTreePropertyFunctionBase.h"
+#include "Templates/GuardValueAccessors.h"
 #include "UObject/AssetRegistryTagsContext.h"
 
 #define LOCTEXT_NAMESPACE "StateTreeEditor"
@@ -49,12 +50,20 @@ IMPLEMENT_MODULE(FStateTreeEditorModule, StateTreeEditorModule)
 
 namespace UE::StateTree::Editor
 {
-	// @todo Could we make this a IModularFeature?
-	static bool CompileStateTree(UStateTree& StateTree)
-	{
-		FStateTreeCompilerLog Log;
-		return UStateTreeEditingSubsystem::CompileStateTree(&StateTree, Log);
-	}
+	/**
+	 * If disabled, the dependencies won't be taken into consideration.
+	 * The PostLoad calls ConditionalPostLoad on the editor data. That calls ConditionalPostLoad on all editor nodes.
+	 * The editor node call ConditionalPostLoad on the other state tree assets.
+	 * That ensures that dependent state tree assets are compiled before being used.
+	 * When compiling in the editor (once the asset are loaded), you might have a dependency bug.
+	 */
+	bool bEnableQueuedCompilationOnAssetLoad = false;
+	FAutoConsoleVariableRef CVarEnableQueuedCompilationOnAssetLoad(
+		TEXT("StateTree.Compiler.EnableQueuedCompilationOnAssetLoad"),
+		bEnableQueuedCompilationOnAssetLoad,
+		TEXT("On state tree asset load, queue the asset compilation. The asset will be compiled when needed.\n")
+		TEXT("If false, then then asset is compiled when loaded.")
+	);
 
 	static TSharedRef<FStateTreeNodeClassCache> InitNodeClassCache()
 	{
@@ -86,18 +95,14 @@ FStateTreeEditorModule* FStateTreeEditorModule::GetModulePtr()
 
 void FStateTreeEditorModule::StartupModule()
 {
-	UE::StateTree::Delegates::OnRequestEditorHash.BindLambda([](const UStateTree& InStateTree) -> uint32
-		{
-			return UStateTreeEditingSubsystem::CalculateStateTreeHash(&InStateTree);
-		});
-	UE::StateTree::Delegates::Private::OnRequestAssetRegistryTags.BindRaw(this, &FStateTreeEditorModule::HandleRequestAssetRegistryTags);
-	UE::StateTree::Delegates::Private::OnAppendToClassSchema.BindLambda([](FAppendToClassSchemaContext& Context)
-		{
-			FStateTreeCompiler::AppendToClassSchema(Context);
-		});
-
 	UE::StateTree::Compiler::FCompilerManager::Startup();
 
+	UE::StateTree::Delegates::OnRequestEditorHash.BindRaw(this, &FStateTreeEditorModule::HandleRequestEditorHash);
+	UE::StateTree::Delegates::Private::OnCompileIfChanged.BindRaw(this, &FStateTreeEditorModule::HandleCompileIfChanged);
+	UE::StateTree::Delegates::Private::OnRequestAssetRegistryTags.BindRaw(this, &FStateTreeEditorModule::HandleRequestAssetRegistryTags);
+	UE::StateTree::Delegates::Private::OnStateTreeAssetLoaded.BindRaw(this, &FStateTreeEditorModule::HandleStateTreeAssetLoaded);
+	UE::StateTree::Delegates::Private::OnAppendToClassSchema.BindRaw(this, &FStateTreeEditorModule::HandleAppendToStateTreeClassSchema);
+
 	OnPostEngineInitHandle = FCoreDelegates::OnPostEngineInit.AddLambda([this]()
 		{
 			NodeClassCache = UE::StateTree::Editor::InitNodeClassCache();
@@ -179,7 +184,9 @@ void FStateTreeEditorModule::ShutdownModule()
 	UE::StateTree::Compiler::FCompilerManager::Shutdown();
 
 	UE::StateTree::Delegates::Private::OnAppendToClassSchema.Unbind();
+	UE::StateTree::Delegates::Private::OnStateTreeAssetLoaded.Unbind();
 	UE::StateTree::Delegates::Private::OnRequestAssetRegistryTags.Unbind();
+	UE::StateTree::Delegates::Private::OnCompileIfChanged.Unbind();
 	UE::StateTree::Delegates::OnRequestEditorHash.Unbind();
 }
 
@@ -187,6 +194,10 @@ TSharedRef<IStateTreeEditor> FStateTreeEditorModule::CreateStateTreeEditor(const
 {
 	TSharedRef<FStateTreeEditor> NewEditor(new FStateTreeEditor());
 	NewEditor->InitEditor(Mode, InitToolkitHost, StateTree);
+	if (StateTree->LastCompiledEditorDataHash == 0)
+	{
+		UE::StateTree::Compiler::FCompilerManager::CompileSynchronously(StateTree);
+	}
 	return NewEditor;
 }
 
@@ -202,6 +213,16 @@ TSharedPtr<FStateTreeNodeClassCache> FStateTreeEditorModule::GetNodeClassCache()
 	return NodeClassCache;
 }
 
+uint32 FStateTreeEditorModule::HandleRequestEditorHash(const UStateTree& InStateTree)
+{
+	return UStateTreeEditingSubsystem::CalculateStateTreeHash(&InStateTree);
+}
+
+void FStateTreeEditorModule::HandleCompileIfChanged(const TNotNull<UStateTree*> InStateTree)
+{
+	UE::StateTree::Compiler::FCompilerManager::CompileIfNeededSynchronously(InStateTree);
+}
+
 void FStateTreeEditorModule::HandleRequestAssetRegistryTags(TNotNull<const UStateTree*> InStateTree, FAssetRegistryTagsContext& Context)
 {
 	// Notify the module, schema, and extension. Use that order to go from the less specific to the more specific.
@@ -227,6 +248,25 @@ void FStateTreeEditorModule::HandleRequestAssetRegistryTags(TNotNull<const UStat
 	}
 }
 
+void FStateTreeEditorModule::HandleStateTreeAssetLoaded(TNotNull<UStateTree*> InStateTree)
+{
+	TGuardValueAccessors<bool> IsEditorLoadingPackageGuard(UE::GetIsEditorLoadingPackage, UE::SetIsEditorLoadingPackage, true);
+	if (UE::StateTree::Editor::bEnableQueuedCompilationOnAssetLoad)
+	{
+		UE::StateTree::Compiler::FCompilerManager::QueueForCompilation(InStateTree);
+	}
+	else
+	{
+		FStateTreeCompilerLog Log;
+		UE::StateTree::Compiler::FCompilerManager::CompileSynchronously(InStateTree, Log);
+	}
+}
+
+void FStateTreeEditorModule::HandleAppendToStateTreeClassSchema(FAppendToClassSchemaContext& Context)
+{
+	FStateTreeCompiler::AppendToStateTreeClassSchema(Context);
+}
+
 bool FStateTreeEditorModule::FEditorTypes::HasData() const
 {
 	return EditorData.IsValid() || EditorSchema.IsValid();
diff --git a/Engine/Plugins/Runtime/StateTree/Source/StateTreeEditorModule/Private/StateTreeState.cpp b/Engine/Plugins/Runtime/StateTree/Source/StateTreeEditorModule/Private/StateTreeState.cpp
--- a/Engine/Plugins/Runtime/StateTree/Source/StateTreeEditorModule/Private/StateTreeState.cpp
+++ b/Engine/Plugins/Runtime/StateTree/Source/StateTreeEditorModule/Private/StateTreeState.cpp
@@ -2,12 +2,14 @@
 
 #include "StateTreeState.h"
 #include "StateTree.h"
+#include "StateTreeEditingSubsystem.h"
 #include "StateTreeEditorData.h"
 #include "StateTreeConditionBase.h"
 #include "StateTreeConsiderationBase.h"
 #include "StateTreeTaskBase.h"
 #include "StateTreeDelegates.h"
 #include "StateTreePropertyHelpers.h"
+
 #include "Customizations/StateTreeEditorNodeUtils.h"
 
 #include UE_INLINE_GENERATED_CPP_BY_NAME(StateTreeState)
@@ -134,12 +136,13 @@ void UStateTreeState::PostEditChangeChainProperty(FPropertyChangedChainEvent& Pr
 	static const FStateTreeEditPropertyPath StateTransitionsConditionsPath(UStateTreeState::StaticClass(), TEXT("Transitions.Conditions"));
 	static const FStateTreeEditPropertyPath StateTransitionsIDPath(UStateTreeState::StaticClass(), TEXT("Transitions.ID"));
 
+	UStateTree* StateTree = GetTypedOuter<UStateTree>();
+	ensure(StateTree);
 
 	// Broadcast name changes so that the UI can update.
 	if (ChangePropertyPath.IsPathExact(StateNamePath))
 	{
-		const UStateTree* StateTree = GetTypedOuter<UStateTree>();
-		if (ensure(StateTree))
+		if (StateTree)
 		{
 			UE::StateTree::Delegates::OnIdentifierChanged.Broadcast(*StateTree);
 		}
@@ -148,8 +151,7 @@ void UStateTreeState::PostEditChangeChainProperty(FPropertyChangedChainEvent& Pr
 	if (ChangePropertyPath.IsPathExact(SelectionBehaviorPath))
 	{
 		// Broadcast selection type changes so that the UI can update.
-		const UStateTree* StateTree = GetTypedOuter<UStateTree>();
-		if (ensure(StateTree))
+		if (StateTree)
 		{
 			UE::StateTree::Delegates::OnIdentifierChanged.Broadcast(*StateTree);
 		}
@@ -216,8 +218,7 @@ void UStateTreeState::PostEditChangeChainProperty(FPropertyChangedChainEvent& Pr
 	// Broadcast subtree parameter layout edits so that the linked states can adapt, and bindings can update.
 	if (ChangePropertyPath.IsPathExact(StateParametersPath))
 	{
-		const UStateTree* StateTree = GetTypedOuter<UStateTree>();
-		if (ensure(StateTree))
+		if (StateTree)
 		{
 			UE::StateTree::Delegates::OnStateParametersChanged.Broadcast(*StateTree, ID);
 		}
@@ -255,6 +256,11 @@ void UStateTreeState::PostEditChangeChainProperty(FPropertyChangedChainEvent& Pr
 		}
 	}
 
+	if (StateTree)
+	{
+		UStateTreeEditingSubsystem::MarkAsModified(StateTree);
+	}
+
 	if (UStateTreeEditorData* TreeData = GetTypedOuter<UStateTreeEditorData>())
 	{
 		UE::StateTree::PropertyHelpers::DispatchPostEditToNodes(*this, PropertyChangedEvent, *TreeData);
diff --git a/Engine/Plugins/Runtime/StateTree/Source/StateTreeEditorModule/Public/StateTreeCompiler.h b/Engine/Plugins/Runtime/StateTree/Source/StateTreeEditorModule/Public/StateTreeCompiler.h
--- a/Engine/Plugins/Runtime/StateTree/Source/StateTreeEditorModule/Public/StateTreeCompiler.h
+++ b/Engine/Plugins/Runtime/StateTree/Source/StateTreeEditorModule/Public/StateTreeCompiler.h
@@ -2,6 +2,7 @@
 
 #pragma once
 
+#include "StateTree.h"
 #include "StateTreePropertyBindingCompiler.h"
 
 #define UE_API STATETREEEDITORMODULE_API
@@ -14,7 +15,6 @@ enum class EStateTreePropertyUsage : uint8;
 struct FStateTreeDataView;
 struct FStateTreeStateHandle;
 
-class UStateTree;
 class UStateTreeState;
 class UStateTreeEditorData;
 class UStateTreeExtension;
@@ -25,8 +25,6 @@ struct FStateTreeNodeBase;
 
 namespace UE::StateTree::Compiler
 {
-	// Compiler steps: PrePublic, PostPublic, PreInternal, PostInternal
-
 	/**
 	 * Compiler context for when the state tree fully compiled and succeeded.
 	 * This is the last compilation step.
@@ -56,12 +54,29 @@ public:
 	{
 	}
 
+	/** Compile the public step of a state tree asset. */
+	UE_API bool CompilePublic(TNotNull<UStateTree*> StateTree);
+
+	/** Compile the internal step of a state tree asset. */
+	UE_API bool CompileInternal(TNotNull<UStateTree*> StateTree);
+
+	/** Compile the public and internal steps of a state tree asset. */
 	UE_API bool Compile(TNotNull<UStateTree*> StateTree);
+
+	/** Compile the public and internal steps of a state tree asset. */
 	UE_API bool Compile(UStateTree& InStateTree);
 
-	UE_API static void AppendToClassSchema(FAppendToClassSchemaContext& Context);
+	/**
+	 * Append config values that can change how state tree asset compiles when cooked.
+	 * When the context is different from the previous result, then a cook is needed.
+	 */
+	UE_API static void AppendToStateTreeClassSchema(FAppendToClassSchemaContext& Context);
 
 private:
+	bool PreCompile(TNotNull<UStateTree*> StateTree);
+	bool CompileInternalImpl();
+	bool FailCompilation(UStateTree::ECompileStatus Status) const;
+
 	/** Resolves the state a transition points to, and the optional fallback for failing to enter the state. SourceState is nullptr for global tasks. */
 	bool ResolveTransitionStateAndFallback(const UStateTreeState* SourceState, const FStateTreeStateLink& Link, FStateTreeStateHandle& OutTransitionHandle, EStateTreeSelectionFallback& OutFallback) const;
 	FStateTreeStateHandle GetStateHandle(const FGuid& StateID) const;
@@ -109,11 +124,11 @@ private:
 	bool CompileAndValidateNode(const UStateTreeState* SourceState, const FStateTreeBindableStructDesc& InstanceDesc, FStructView NodeView, const FStateTreeDataView InstanceData);
 
 	void InstantiateStructSubobjects(FStructView Struct);
-        
+
 	bool NotifyInternalPost();
 	
 	void CreateBindingSourceStructsForNode(const FStateTreeEditorNode& EditorNode, const FStateTreeBindableStructDesc& InstanceDesc);
-      
+
 private:
 	FStateTreeCompilerLog& Log;
 	TObjectPtr<UStateTree> StateTree = nullptr;
@@ -135,6 +150,9 @@ private:
 
 	FStateTreePropertyBindingCompiler BindingsCompiler;
 
+	/** The state tree hash. */
+	TOptional<uint32> EditorDataHash;
+
 	/** The Compile function executed. */
 	bool bCompiled = false;
 };
diff --git a/Engine/Plugins/Runtime/StateTree/Source/StateTreeEditorModule/Public/StateTreeCompilerManager.h b/Engine/Plugins/Runtime/StateTree/Source/StateTreeEditorModule/Public/StateTreeCompilerManager.h
--- a/Engine/Plugins/Runtime/StateTree/Source/StateTreeEditorModule/Public/StateTreeCompilerManager.h
+++ b/Engine/Plugins/Runtime/StateTree/Source/StateTreeEditorModule/Public/StateTreeCompilerManager.h
@@ -18,8 +18,38 @@ public:
 	static UE_API void Startup();
 	static UE_API void Shutdown();
 
+	/**
+	 * Queue the compilation of a dirty state tree asset for later.
+	 * The public step will be compiled synchronously if it is dirty.
+	 * The internal step will be compiled when requested by the game or asynchronously.
+	 * Queuing state tree asset can be disabled with cvar. When disabled, this function behaves like CompileSynchronously.
+	 * @note Can only compile on the game thread or async loading thread.
+	 */
+	static UE_API void QueueForCompilation(TNotNull<UStateTree*> StateTree);
+
+	/**
+	 * Compile the internal steps of all queued state tree assets.
+	 * Sort them to compile them with regard to their dependencies.
+	 * @note Can only flush on the game thread.
+	 */
+	static UE_API void FlushCompilationQueue();
+
+	/**
+	 * Compile the public and internal steps of a state tree asset.
+	 * It will compile all the pending state tree asset dependencies if they are queued.
+	 * @return whether the compilation succeeded or failed.
+	 * @note Can only compile on the game thread or async loading thread.
+	 */
 	static UE_API bool CompileSynchronously(TNotNull<UStateTree*> StateTree);
 	static UE_API bool CompileSynchronously(TNotNull<UStateTree*> StateTree, FStateTreeCompilerLog& Log);
+
+	/**
+	 * Compile the public and internal steps of a state tree asset only if the editor data changed.
+	 * It will compile all the pending state tree asset dependencies if they are queued.
+	 * @return whether the compilation succeeded or failed. The optional is not set if a compilation was not necessary.
+	 * @note Can only compile on the game thread or async loading thread.
+	 */
+	static UE_API TOptional<bool> CompileIfNeededSynchronously(TNotNull<UStateTree*> StateTree);
 	
 private:
 	FCompilerManager() = delete;
diff --git a/Engine/Plugins/Runtime/StateTree/Source/StateTreeEditorModule/Public/StateTreeEditingSubsystem.h b/Engine/Plugins/Runtime/StateTree/Source/StateTreeEditorModule/Public/StateTreeEditingSubsystem.h
--- a/Engine/Plugins/Runtime/StateTree/Source/StateTreeEditorModule/Public/StateTreeEditingSubsystem.h
+++ b/Engine/Plugins/Runtime/StateTree/Source/StateTreeEditorModule/Public/StateTreeEditingSubsystem.h
@@ -26,12 +26,32 @@ public:
 	
 	UE_API TSharedPtr<FStateTreeViewModel> FindViewModel(TNotNull<const UStateTree*> InStateTree) const;
 	UE_API TSharedRef<FStateTreeViewModel> FindOrAddViewModel(TNotNull<UStateTree*> InStateTree);
-	
-	static UE_API bool CompileStateTree(TNotNull<UStateTree*> InStateTree,  FStateTreeCompilerLog& InOutLog);
-	
-	/** Create a StateTreeView widget for the viewmodel. */
+
+	/** Create a state tree view widget for the viewmodel. */
 	static UE_API TSharedRef<SWidget> GetStateTreeView(TSharedRef<FStateTreeViewModel> InViewModel, const TSharedRef<FUICommandList>& TreeViewCommandList);
 	
+	/**
+	 * Compile the state tree asset.
+	 * It will compile the public and internal steps.
+	 */
+	static UE_API bool CompileStateTree(TNotNull<UStateTree*> InStateTree, FStateTreeCompilerLog& InOutLog);
+
+	/** Whether the state tree asset is dirty and requires a compilation. */
+	static UE_API bool NeedsRecompile(TNotNull<const UStateTree*> InStateTree);
+
+	/**
+	 * Mark the public and the internal steps of the state tree asset as dirty.
+	 * It will require the state tree to be compiled before it can be used.
+	 * It will affect the other assets that reference the state tree asset.
+	 */
+	static UE_API void MarkAsPubliclyModified(TNotNull<UStateTree*> InStateTree);
+
+	/**
+	 * Mark the internal step of the state tree asset as dirty.
+	 * It will require the state tree to be compiled before it can be used.
+	 */
+	static UE_API void MarkAsModified(TNotNull<UStateTree*> InStateTree);
+	
 	/**
 	 * Validates and applies the schema restrictions on the StateTree.
 	 * Also serves as the "safety net" of fixing up editor data following an editor operation.
diff --git a/Engine/Plugins/Runtime/StateTree/Source/StateTreeEditorModule/Public/StateTreeEditorData.h b/Engine/Plugins/Runtime/StateTree/Source/StateTreeEditorModule/Public/StateTreeEditorData.h
--- a/Engine/Plugins/Runtime/StateTree/Source/StateTreeEditorModule/Public/StateTreeEditorData.h
+++ b/Engine/Plugins/Runtime/StateTree/Source/StateTreeEditorModule/Public/StateTreeEditorData.h
@@ -359,7 +359,7 @@ PRAGMA_ENABLE_DEPRECATION_WARNINGS
 		return Colors.Find(FStateTreeEditorColor(ColorRef));
 	}
 
-	virtual void CreateRootProperties(TArrayView<UE::PropertyBinding::FPropertyCreationDescriptor> InOutCreationDescs) { UE::PropertyBinding::CreateUniquelyNamedPropertiesInPropertyBag(InOutCreationDescs, RootParameterPropertyBag); }
+	UE_API virtual void CreateRootProperties(TArrayView<UE::PropertyBinding::FPropertyCreationDescriptor> InOutCreationDescs);
 
 private:
 	UE_API void FixObjectInstances(TSet<UObject*>& SeenObjects, UObject& Outer, FStateTreeEditorNode& Node);
diff --git a/Engine/Plugins/Runtime/StateTree/Source/StateTreeEditorModule/Public/StateTreeEditorModule.h b/Engine/Plugins/Runtime/StateTree/Source/StateTreeEditorModule/Public/StateTreeEditorModule.h
--- a/Engine/Plugins/Runtime/StateTree/Source/StateTreeEditorModule/Public/StateTreeEditorModule.h
+++ b/Engine/Plugins/Runtime/StateTree/Source/StateTreeEditorModule/Public/StateTreeEditorModule.h
@@ -132,7 +132,11 @@ public:
 	UE_API TNonNullSubclassOf<UStateTreeEditorSchema> GetEditorSchemaClass(TNonNullSubclassOf<const UStateTreeSchema> Schema) const;
 
 private:
+	uint32 HandleRequestEditorHash(const UStateTree& InStateTree);
+	void HandleCompileIfChanged(TNotNull<UStateTree*> InStateTree);
 	void HandleRequestAssetRegistryTags(TNotNull<const UStateTree*> InStateTree, FAssetRegistryTagsContext& Context);
+	void HandleStateTreeAssetLoaded(TNotNull<UStateTree*> InStateTree);
+	void HandleAppendToStateTreeClassSchema(FAppendToClassSchemaContext& Context);
 
 protected:
 	TSharedPtr<FExtensibilityManager> MenuExtensibilityManager;
diff --git a/Engine/Plugins/Runtime/StateTree/Source/StateTreeModule/Internal/StateTreeDelegatesInternal.h b/Engine/Plugins/Runtime/StateTree/Source/StateTreeModule/Internal/StateTreeDelegatesInternal.h
--- a/Engine/Plugins/Runtime/StateTree/Source/StateTreeModule/Internal/StateTreeDelegatesInternal.h
+++ b/Engine/Plugins/Runtime/StateTree/Source/StateTreeModule/Internal/StateTreeDelegatesInternal.h
@@ -6,25 +6,35 @@
 
 #define UE_API STATETREEMODULE_API
 
-#if WITH_EDITOR
 class UStateTree;
+#if WITH_EDITOR
 struct FAppendToClassSchemaContext;
 class FAssetRegistryTagsContext;
 #endif
 
 namespace UE::StateTree::Delegates::Private
 {
+	/** A StateTree asset loads. */
+	DECLARE_DELEGATE_OneParam(FOnStateTreeAssetLoaded, TNotNull<UStateTree*> /*StateTree*/);
+	extern UE_API FOnStateTreeAssetLoaded OnStateTreeAssetLoaded;
 
 #if WITH_EDITOR
-/** A callback for UStateTree::GetAssetRegistryTags. */
-DECLARE_DELEGATE_TwoParams(FOnStateTreeRequestAssetRegistryTags, TNotNull<const UStateTree*> /*StateTree*/, FAssetRegistryTagsContext& /*Context*/);
-extern UE_API FOnStateTreeRequestAssetRegistryTags OnRequestAssetRegistryTags;
+	/** A callback for UStateTree::CompileIfChanged. */
+	DECLARE_DELEGATE_OneParam(FOnCompileIfChanged, TNotNull<UStateTree*> /*StateTree*/);
+	extern UE_API FOnCompileIfChanged OnCompileIfChanged;
 
-/** A callback for UStateTree::AppendToClassSchema. */
-DECLARE_DELEGATE_OneParam(FOnAppendToClassSchema, FAppendToClassSchemaContext& /*Context*/);
-extern UE_API FOnAppendToClassSchema OnAppendToClassSchema;
-#endif
+	/** A callback for UStateTree::GetAssetRegistryTags. */
+	DECLARE_DELEGATE_TwoParams(FOnStateTreeRequestAssetRegistryTags, TNotNull<const UStateTree*> /*StateTree*/, FAssetRegistryTagsContext& /*Context*/);
+	extern UE_API FOnStateTreeRequestAssetRegistryTags OnRequestAssetRegistryTags;
 
-}; // UE::StateTree::Delegates
+	/** A callback for UStateTree::AppendToClassSchema. */
+	DECLARE_DELEGATE_OneParam(FOnAppendToClassSchema, FAppendToClassSchemaContext& /*Context*/);
+	extern UE_API FOnAppendToClassSchema OnAppendToClassSchema;
+
+	/** The StateTree asset is going to get cooked. */
+	DECLARE_DELEGATE_OneParam(FOnPreCookStateTreeAsset, TNotNull<UStateTree*> /*StateTree*/);
+	extern UE_API FOnPreCookStateTreeAsset OnPreCookStateTreeAsset;
+#endif
+}; // namespace
 
 #undef UE_API
diff --git a/Engine/Plugins/Runtime/StateTree/Source/StateTreeModule/Private/StateTree.cpp b/Engine/Plugins/Runtime/StateTree/Source/StateTreeModule/Private/StateTree.cpp
--- a/Engine/Plugins/Runtime/StateTree/Source/StateTreeModule/Private/StateTree.cpp
+++ b/Engine/Plugins/Runtime/StateTree/Source/StateTreeModule/Private/StateTree.cpp
@@ -23,9 +23,6 @@
 #include <atomic>
 
 #if WITH_EDITOR
-#include "Editor.h"
-#include "StructUtils/UserDefinedStruct.h"
-#include "Templates/GuardValueAccessors.h"
 #include "UObject/LinkerLoad.h"
 #endif
 
@@ -41,8 +38,13 @@ PRAGMA_ENABLE_DEPRECATION_WARNINGS
 
 bool UStateTree::IsReadyToRun() const
 {
+#if WITH_EDITOR
+	// The asset can be in the compiler queue. Compile the asset sync.
+	const_cast<UStateTree*>(this)->CompileIfChanged();
+#endif
+
 	// Valid tree must have at least one state and valid instance data.
-	return States.Num() > 0 && bIsLinked && PropertyBindings.IsValid();
+	return States.Num() > 0 && CompileStatus == ECompileStatus::Executable && PropertyBindings.IsValid();
 }
 
 FConstStructView UStateTree::GetNode(const int32 NodeIndex) const
@@ -273,6 +275,9 @@ void UStateTree::ResetCompiled()
 	
 	ResetLinked();
 
+	CompileStatus = ECompileStatus::PublicDirty;
+	LastCompiledEditorDataHash = 0;
+
 	// Remove objects created from last compilation.
 	{
 		TArray<UObject*, TInlineAllocator<32>> Children;
@@ -292,6 +297,11 @@ void UStateTree::ResetCompiled()
 	}
 }
 
+void UStateTree::MarkAsModified(const bool bPubliclyModified)
+{
+	CompileStatus = bPubliclyModified ? ECompileStatus::PublicDirty : ECompileStatus::InternalDirty;
+}
+
 void UStateTree::PostDuplicate(EDuplicateMode::Type DuplicateMode)
 {
 	Super::PostDuplicate(DuplicateMode);
@@ -304,6 +314,12 @@ void UStateTree::GetAssetRegistryTags(FAssetRegistryTagsContext Context) const
 	Super::GetAssetRegistryTags(Context);
 }
 
+void UStateTree::BeginCacheForCookedPlatformData(const ITargetPlatform* TargetPlatform)
+{
+	UE::StateTree::Delegates::Private::OnPreCookStateTreeAsset.ExecuteIfBound(this);
+	Super::BeginCacheForCookedPlatformData(TargetPlatform);
+}
+
 void UStateTree::ThreadedPostLoadAssetRegistryTagsOverride(FPostLoadAssetRegistryTagsContext& Context) const
 {
 	Super::ThreadedPostLoadAssetRegistryTagsOverride(Context);
@@ -375,16 +391,16 @@ void UStateTree::PostLoad()
 		if (FStateTreeNodeBase* Node = NodeView.GetPtr<FStateTreeNodeBase>())
 		{
 			auto PostLoadInstance = [&Node]<typename T>(T & Container, FStateTreeIndex16 Index)
-			{
-				if (Container.IsObject(Index.Get()))
-				{
-					Node->PostLoad(Container.GetMutableObject(Index.Get()));
-				}
-				else
 				{
-					Node->PostLoad(Container.GetMutableStruct(Index.Get()));
-				}
-			};
+					if (Container.IsObject(Index.Get()))
+					{
+						Node->PostLoad(Container.GetMutableObject(Index.Get()));
+					}
+					else
+					{
+						Node->PostLoad(Container.GetMutableStruct(Index.Get()));
+					}
+				};
 			if (Node->InstanceTemplateIndex.IsValid())
 			{
 				const bool bIsUsingSharedInstanceData = Node->InstanceDataHandle.GetSource() == EStateTreeDataSourceType::SharedInstanceData
@@ -407,26 +423,6 @@ void UStateTree::PostLoad()
 		}
 	}
 
-#if WITH_EDITOR
-	if (EditorData)
-	{
-		// Make sure all the fix up logic in the editor data has had chance to happen.
-		EditorData->ConditionalPostLoad();
-
-		TGuardValueAccessors<bool> IsEditorLoadingPackageGuard(UE::GetIsEditorLoadingPackage, UE::SetIsEditorLoadingPackage, true);
-		Compile();
-	}
-#else
-PRAGMA_DISABLE_DEPRECATION_WARNINGS
-	const int32 CurrentVersion = GetLinkerCustomVersion(FStateTreeCustomVersion::GUID);
-	if (CurrentVersion < FStateTreeCustomVersion::LatestVersion)
-	{		
-		UE_LOG(LogStateTree, Error, TEXT("%s: compiled data is in older format. Please recompile the StateTree asset."), *GetPathName());
-		return;
-	}
-PRAGMA_ENABLE_DEPRECATION_WARNINGS
-#endif
-
 	if (Schema)
 	{
 		Schema->ConditionalPostLoad();
@@ -440,9 +436,27 @@ PRAGMA_ENABLE_DEPRECATION_WARNINGS
 		}
 	}
 
-	if (!Link())
+	// Request a link
+	CompileStatus = ECompileStatus::LinkDirty;
+
+#if WITH_EDITOR
+	if (EditorData)
 	{
-		UE_LOG(LogStateTree, Log, TEXT("%s failed to link. Asset will not be usable at runtime."), *GetPathName());
+		// Make sure all the fix up logic in the editor data has had chance to happen.
+		EditorData->ConditionalPostLoad();
+		ResetCompiled();
+	}
+#endif
+
+	// Request a compilation (can be async or sync).
+	UE::StateTree::Delegates::Private::OnStateTreeAssetLoaded.ExecuteIfBound(this);
+
+	if (CompileStatus == ECompileStatus::LinkDirty && LastCompiledEditorDataHash != 0)
+	{
+		if (!Link())
+		{
+			UE_CLOG(CompileStatus == ECompileStatus::LinkDirty, LogStateTree, Warning, TEXT("%s failed to link. Asset will not be usable at runtime."), *GetPathName());
+		}
 	}
 }
 
@@ -491,7 +505,11 @@ PRAGMA_ENABLE_DEPRECATION_WARNINGS
 
 void UStateTree::ResetLinked()
 {
-	bIsLinked = false;
+	if (CompileStatus == ECompileStatus::Executable)
+	{
+		CompileStatus = ECompileStatus::LinkDirty;
+	}
+
 	ExternalDataDescs.Reset();
 
 #if WITH_EDITOR
@@ -649,6 +667,11 @@ bool UStateTree::ValidateInstanceData()
 
 bool UStateTree::Link()
 {
+	if (CompileStatus == ECompileStatus::PublicDirty || CompileStatus == ECompileStatus::InternalDirty)
+	{
+		return false;
+	}
+
 	// Initialize the instance data default value.
 	// This data will be used to allocate runtime instance on all StateTree users.
 	ResetLinked();
@@ -760,7 +783,7 @@ bool UStateTree::Link()
 	}
 
 	// Link succeeded, setup tree to be ready to run
-	bIsLinked = true;
+	CompileStatus = ECompileStatus::Executable;
 	
 	return true;
 }
@@ -1208,7 +1231,7 @@ TArray<FStateTreeMemoryUsage> UStateTree::CalculateEstimatedMemoryUsage() const
 	TArray<FStateTreeMemoryUsage> MemoryUsages;
 	TArray<TPair<int32, int32>> StateLinks;
 
-	if (!bIsLinked
+	if (CompileStatus != ECompileStatus::Executable
 		|| States.IsEmpty()
 		|| !Nodes.IsValid())
 	{
@@ -1422,33 +1445,11 @@ TArray<FStateTreeMemoryUsage> UStateTree::CalculateEstimatedMemoryUsage() const
 
 void UStateTree::CompileIfChanged()
 {
-	if (UE::StateTree::Delegates::OnRequestCompile.IsBound() && UE::StateTree::Delegates::OnRequestEditorHash.IsBound())
+	// Same check are in the compiler manager. But do it here also as a quick optimization.
+	if (CompileStatus == UStateTree::ECompileStatus::PublicDirty
+		|| CompileStatus == UStateTree::ECompileStatus::InternalDirty)
 	{
-		const uint32 CurrentHash = UE::StateTree::Delegates::OnRequestEditorHash.Execute(*this);
-		if (LastCompiledEditorDataHash != CurrentHash)
-		{
-			UE_LOG(LogStateTree, Log, TEXT("%s: Editor data has changed. Recompiling state tree."), *GetPathName());
-			UE::StateTree::Delegates::OnRequestCompile.Execute(*this);
-		}
-	}
-	else
-	{
-		ResetCompiled();
-		UE_LOG(LogStateTree, Warning, TEXT("%s: could not compile. Please resave the StateTree asset."), *GetPathName());
-	}
-}
-
-void UStateTree::Compile()
-{
-	if (UE::StateTree::Delegates::OnRequestCompile.IsBound())
-	{
-		UE_LOG(LogStateTree, Log, TEXT("%s: Editor data has changed. Recompiling state tree."), *GetPathName());
-		UE::StateTree::Delegates::OnRequestCompile.Execute(*this);
-	}
-	else
-	{
-		ResetCompiled();
-		UE_LOG(LogStateTree, Warning, TEXT("%s: could not compile. Please resave the StateTree asset."), *GetPathName());
+		UE::StateTree::Delegates::Private::OnCompileIfChanged.ExecuteIfBound(this);
 	}
 }
 #endif //WITH_EDITOR
diff --git a/Engine/Plugins/Runtime/StateTree/Source/StateTreeModule/Private/StateTreeDelegates.cpp b/Engine/Plugins/Runtime/StateTree/Source/StateTreeModule/Private/StateTreeDelegates.cpp
--- a/Engine/Plugins/Runtime/StateTree/Source/StateTreeModule/Private/StateTreeDelegates.cpp
+++ b/Engine/Plugins/Runtime/StateTree/Source/StateTreeModule/Private/StateTreeDelegates.cpp
@@ -15,7 +15,10 @@ namespace UE::StateTree::Delegates
 	FOnStateParametersChanged OnStateParametersChanged;
 	FOnBreakpointsChanged OnBreakpointsChanged;
 	FOnPostCompile OnPostCompile;
+PRAGMA_DISABLE_DEPRECATION_WARNINGS
+	//Deprecated
 	FOnRequestCompile OnRequestCompile;
+PRAGMA_ENABLE_DEPRECATION_WARNINGS
 	FOnRequestEditorHash OnRequestEditorHash;
 #endif // WITH_EDITOR
 
@@ -32,8 +35,11 @@ namespace UE::StateTree::Delegates
 
 namespace UE::StateTree::Delegates::Private
 {
+	FOnStateTreeAssetLoaded OnStateTreeAssetLoaded;
 #if WITH_EDITOR
+	FOnCompileIfChanged OnCompileIfChanged;
 	FOnStateTreeRequestAssetRegistryTags OnRequestAssetRegistryTags;
 	FOnAppendToClassSchema OnAppendToClassSchema;
+	FOnStateTreeAssetLoaded OnPreCookStateTreeAsset;
 #endif
 } //UE::StateTree::Delegates::Private
diff --git a/Engine/Plugins/Runtime/StateTree/Source/StateTreeModule/Public/StateTree.h b/Engine/Plugins/Runtime/StateTree/Source/StateTreeModule/Public/StateTree.h
--- a/Engine/Plugins/Runtime/StateTree/Source/StateTreeModule/Public/StateTree.h
+++ b/Engine/Plugins/Runtime/StateTree/Source/StateTreeModule/Public/StateTree.h
@@ -306,6 +306,13 @@ public:
 	/** Resets the compiled data to empty. */
 	UE_API void ResetCompiled();
 
+	/**
+	 * Mark the StateTree asset dirty.
+	 * It will require the StateTree to be compiled before it can be used.
+	 * @param bPubliclyModified Whether to also mark the public dependency dirty. It will affect the other asset that references the StateTree asset.
+	 */
+	UE_API void MarkAsModified(bool bPubliclyModified);
+
 	/** Calculates runtime memory usage for different sections of the tree. */
 	UE_API TArray<FStateTreeMemoryUsage> CalculateEstimatedMemoryUsage() const;
 
@@ -366,6 +373,7 @@ protected:
 #if WITH_EDITOR
 	UE_API virtual void PostDuplicate(EDuplicateMode::Type DuplicateMode) override;
 	UE_API virtual void GetAssetRegistryTags(FAssetRegistryTagsContext Context) const override;
+	UE_API void BeginCacheForCookedPlatformData(const ITargetPlatform* TargetPlatform) override;
 protected:
 	UE_API virtual void ThreadedPostLoadAssetRegistryTagsOverride(FPostLoadAssetRegistryTagsContext& Context) const override;
 public:
@@ -375,9 +383,6 @@ public:
 	static UE_API void AddReferencedObjects(UObject* InThis, FReferenceCollector& Collector);
 	
 private:
-	/** Compile the state trees. */
-	UE_API void Compile();
-
 	/**
 	 * Reset the data generated by Link(), this in turn will cause IsReadyToRun() to return false.
 	 * Used during linking, or to invalidate the linked data when data version is old (requires recompile). 
@@ -512,6 +517,20 @@ private:
 	UPROPERTY()
 	EStateTreeParameterDataType ParameterDataType = EStateTreeParameterDataType::GlobalParameterData;
 
+	enum class ECompileStatus : uint8
+	{
+		PublicDirty,		// The exported dependencies need to be compiled. 
+		InternalDirty,		// The public compiled and the internal dependencies need to be compiled.
+		LinkDirty,			// Compiled and not yet link.
+		Executable,			// Compiled and ready to use.
+	};
+	ECompileStatus CompileStatus = ECompileStatus::LinkDirty;
+
+#if WITH_EDITORONLY_DATA
+	/** The hash during the last compilation. Valid whether the compilation succeeded or failed. */
+	uint32 LastCompilationEditorDataHash = 0;
+#endif
+
 	/** True if any global task is a transition task. */
 	UPROPERTY()
 	uint8 bHasGlobalTransitionTasks : 1 = false;
@@ -541,9 +560,6 @@ private:
 	/** True when the scheduled tick is allowed by the schema. */
 	uint8 bScheduledTickAllowed : 1 = false;
 
-	/** True if the StateTree was linked successfully. */
-	uint8 bIsLinked : 1 = false;
-
 #if WITH_EDITORONLY_DATA
 	/** List of Struct that are out of date and waiting to be replaced with the new instance. */
 	TSet<FObjectKey> OutOfDateStructs;
diff --git a/Engine/Plugins/Runtime/StateTree/Source/StateTreeModule/Public/StateTreeDelegates.h b/Engine/Plugins/Runtime/StateTree/Source/StateTreeModule/Public/StateTreeDelegates.h
--- a/Engine/Plugins/Runtime/StateTree/Source/StateTreeModule/Public/StateTreeDelegates.h
+++ b/Engine/Plugins/Runtime/StateTree/Source/StateTreeModule/Public/StateTreeDelegates.h
@@ -67,9 +67,13 @@ extern UE_API FOnBreakpointsChanged OnBreakpointsChanged;
 DECLARE_MULTICAST_DELEGATE_OneParam(FOnPostCompile, const UStateTree& /*StateTree*/);
 extern UE_API FOnPostCompile OnPostCompile;
 
+PRAGMA_DISABLE_DEPRECATION_WARNINGS
+UE_DEPRECATED(5.8, "Use the UStateTreeEditingSubsystem to trigger a compilation.")
 /** Request StateTree compilation. */
 DECLARE_DELEGATE_RetVal_OneParam(bool, FOnRequestCompile, UStateTree& /*StateTreeToCompile*/);
+UE_DEPRECATED(5.8, "Use the UStateTreeEditingSubsystem to trigger a compilation.")
 extern UE_API FOnRequestCompile OnRequestCompile;
+PRAGMA_ENABLE_DEPRECATION_WARNINGS
 
 /** Request the editor hash of the given state tree. */
 DECLARE_DELEGATE_RetVal_OneParam(uint32, FOnRequestEditorHash, const UStateTree& /*StateTreeToHash*/);
diff --git a/Engine/Plugins/Runtime/StateTree/Source/StateTreeModule/Public/StateTreeSchema.h b/Engine/Plugins/Runtime/StateTree/Source/StateTreeModule/Public/StateTreeSchema.h
--- a/Engine/Plugins/Runtime/StateTree/Source/StateTreeModule/Public/StateTreeSchema.h
+++ b/Engine/Plugins/Runtime/StateTree/Source/StateTreeModule/Public/StateTreeSchema.h
@@ -129,6 +129,12 @@ public:
 	{
 		return true;
 	}
+
+	/** @return Whether the compilation of the state tree is allowed to be queued. */
+	virtual bool AllowQueuedCompilation() const
+	{
+		return true;
+	}
 #endif // WITH_EDITOR
 };
 
diff --git a/Engine/Plugins/VirtualProduction/Avalanche/Source/AvalancheTransitionEditor/Private/Compiler/AvaTransitionCompiler.cpp b/Engine/Plugins/VirtualProduction/Avalanche/Source/AvalancheTransitionEditor/Private/Compiler/AvaTransitionCompiler.cpp
--- a/Engine/Plugins/VirtualProduction/Avalanche/Source/AvalancheTransitionEditor/Private/Compiler/AvaTransitionCompiler.cpp
+++ b/Engine/Plugins/VirtualProduction/Avalanche/Source/AvalancheTransitionEditor/Private/Compiler/AvaTransitionCompiler.cpp
@@ -103,7 +103,6 @@ bool FAvaTransitionCompiler::Compile(EAvaTransitionEditorMode InCompileMode)
 	{
 		// Make sure not to leave stale data on failed compile.
 		TransitionTree->ResetCompiled();
-		TransitionTree->LastCompiledEditorDataHash = 0;
 		OnCompileFail.ExecuteIfBound();
 	}
 
