[StateTree] Enable evaluation scope instance by default for property functions, conditions and utilities.

#rnx

diff --git a/Engine/Plugins/Runtime/StateTree/Source/StateTreeEditorModule/Private/StateTreeCompiler.cpp b/Engine/Plugins/Runtime/StateTree/Source/StateTreeEditorModule/Private/StateTreeCompiler.cpp
--- a/Engine/Plugins/Runtime/StateTree/Source/StateTreeEditorModule/Private/StateTreeCompiler.cpp
+++ b/Engine/Plugins/Runtime/StateTree/Source/StateTreeEditorModule/Private/StateTreeCompiler.cpp
@@ -78,14 +78,14 @@ namespace UE::StateTree::Compiler
 		TEXT("SharedInstance is the previous behavior that is deprecated.")
 	);
 
-	bool bEnableConditionWithEvaluationScopeInstanceData = false;
+	bool bEnableConditionWithEvaluationScopeInstanceData = true;
 	FAutoConsoleVariableRef CVarEnableConditionWithEvaluationScopeInstanceData(
 		TEXT("StateTree.Compiler.EnableConditionWithEvaluationScopeInstanceData"),
 		bEnableConditionWithEvaluationScopeInstanceData,
 		TEXT("Use EvaluationScope data for conditions instead of SharedInstance data.")
 	);
 
-	bool bEnableUtilityConsiderationWithEvaluationScopeInstanceData = false;
+	bool bEnableUtilityConsiderationWithEvaluationScopeInstanceData = true;
 	FAutoConsoleVariableRef CVarEnableUtilityConsiderationWithEvaluationScopeInstanceData(
 		TEXT("StateTree.Compiler.EnableUtilityConsiderationWithEvaluationScopeInstanceData"),
 		bEnableUtilityConsiderationWithEvaluationScopeInstanceData,
