[StateTree] Ensure transient renaming is unique

#tests editor

diff --git a/Engine/Plugins/Runtime/StateTree/Source/StateTreeModule/Private/StateTree.cpp b/Engine/Plugins/Runtime/StateTree/Source/StateTreeModule/Private/StateTree.cpp
--- a/Engine/Plugins/Runtime/StateTree/Source/StateTreeModule/Private/StateTree.cpp
+++ b/Engine/Plugins/Runtime/StateTree/Source/StateTreeModule/Private/StateTree.cpp
@@ -234,9 +234,11 @@ namespace UE::StateTree::Compiler
 		ObjectToRename->SetFlags(RF_Transient);
 		ObjectToRename->ClearFlags(RF_Public | RF_Standalone | RF_ArchetypeObject);
 
+		const FName UniqueName = MakeUniqueObjectName(GetTransientPackage(), ObjectToRename->GetClass(), FName(Name));
+
 		// Rename will remove the renamed object's linker when moving to a new package so invalidate the export beforehand
 		FLinkerLoad::InvalidateExport(ObjectToRename);
-		ObjectToRename->Rename(Name.GetData(), GetTransientPackage(), RenFlags);
+		ObjectToRename->Rename(*UniqueName.ToString(), GetTransientPackage(), RenFlags);
 	}
 }
 
