[StateTree] Remove the temporary instance data once it has been consumed. Previously, all temporary instance data was moved in UpdateInstanceData. Now, they are only moved when the task starts.

diff --git a/Engine/Plugins/Runtime/StateTree/Source/StateTreeModule/Private/StateTreeExecutionContext.cpp b/Engine/Plugins/Runtime/StateTree/Source/StateTreeModule/Private/StateTreeExecutionContext.cpp
--- a/Engine/Plugins/Runtime/StateTree/Source/StateTreeModule/Private/StateTreeExecutionContext.cpp
+++ b/Engine/Plugins/Runtime/StateTree/Source/StateTreeModule/Private/StateTreeExecutionContext.cpp
@@ -2289,20 +2289,35 @@ void FStateTreeExecutionContext::UpdateInstanceData(const TSharedPtr<FSelectStat
 	// Temporaries that are moved to the new instance data buffer.
 	// Node instance data are moved in EnterState when the ActiveNodeIndex is incremented.
 	// The array is clear at the end of EnterState.
-	TArray<TOptional<FInstancedStruct*>, FNonconcurrentLinearArrayAllocator> TempInstanceStructs;
-	TempInstanceStructs.Reserve(EstimatedNumStructs);
+	TArray<TOptional<FInstancedStruct*>, FNonconcurrentLinearArrayAllocator> TemporaryInstanceStructs;
+	TemporaryInstanceStructs.Reserve(EstimatedNumStructs);
+
+	using FTemporaryInstanceKey = TTuple<FActiveFrameID, FStateTreeIndex16, FStateTreeDataHandle>;
+	TArray<FTemporaryInstanceKey, TInlineAllocator<8, FNonconcurrentLinearArrayAllocator>> TemporaryInstanceStructsToConsume;
 
 	constexpr int32 ExpectedAmountOfFrames = 8;
 	TArray<FCompactStateTreeParameters, TInlineAllocator<ExpectedAmountOfFrames, FNonconcurrentLinearArrayAllocator>> TempParams;
 
-	TArrayView<FStateTreeTemporaryInstanceData> TempInstances = Storage.GetMutableTemporaryInstances();
-	auto FindInstanceTempData = [&TempInstances](FActiveFrameID FrameID, FStateTreeDataHandle DataHandle)
+	TArrayView<FStateTreeTemporaryInstanceData> TemporaryInstances = Storage.GetMutableTemporaryInstances();
+	auto FindInstanceTempData = [&TemporaryInstances](FActiveFrameID FrameID, FStateTreeIndex16 OwnerNodeIndex, FStateTreeDataHandle DataHandle)
 		{
-			FStateTreeTemporaryInstanceData* TempData = TempInstances.FindByPredicate([&FrameID, &DataHandle](const FStateTreeTemporaryInstanceData& Data)
+			FStateTreeTemporaryInstanceData* TempData = TemporaryInstances.FindByPredicate([&FrameID, &OwnerNodeIndex, &DataHandle](const FStateTreeTemporaryInstanceData& Data)
 				{
-					return Data.FrameID == FrameID && Data.DataHandle == DataHandle;
+					return Data.FrameID == FrameID
+						&& Data.OwnerNodeIndex == OwnerNodeIndex
+						&& Data.DataHandle == DataHandle;
 				});
-			return TempData ? &TempData->Instance : nullptr;
+				return TempData ? &TempData->Instance : nullptr;
+		};
+
+	auto ConsumeTemporaryInstanceData = [&FindInstanceTempData, &TemporaryInstanceStructsToConsume](FActiveFrameID FrameID, FStateTreeIndex16 OwnerNodeIndex, FStateTreeDataHandle DataHandle)
+		{
+			FInstancedStruct* FoundInstance = FindInstanceTempData(FrameID, OwnerNodeIndex, DataHandle);
+			if (FoundInstance)
+			{
+				TemporaryInstanceStructsToConsume.Emplace(FrameID, OwnerNodeIndex, DataHandle);
+			}
+			return FoundInstance;
 		};
 
 	constexpr TOptional<FInstancedStruct*> NoExistingTemporaryDuplicateWrapper = TOptional<FInstancedStruct*>(nullptr);
@@ -2362,7 +2377,7 @@ void FStateTreeExecutionContext::UpdateInstanceData(const TSharedPtr<FSelectStat
 					CurrentGlobalInstanceIndexBase = BaseIndex;
 
 					InstanceStructs.AddDefaulted(CurrentStateTree->NumGlobalInstanceData);
-					TempInstanceStructs.AddZeroed(CurrentStateTree->NumGlobalInstanceData);
+					TemporaryInstanceStructs.AddDefaulted(CurrentStateTree->NumGlobalInstanceData);
 
 					// Global Evals
 					for (int32 EvalIndex = CurrentStateTree->EvaluatorsBegin; EvalIndex < (CurrentStateTree->EvaluatorsBegin + CurrentStateTree->EvaluatorsNum); EvalIndex++)
@@ -2371,7 +2386,7 @@ void FStateTreeExecutionContext::UpdateInstanceData(const TSharedPtr<FSelectStat
 						const FConstStructView EvalInstanceData = CurrentStateTree->DefaultInstanceData.GetStruct(Eval.InstanceTemplateIndex.Get());
 						const int32 EvalInstanceDataIndex = BaseIndex + Eval.InstanceDataHandle.GetIndex();
 						InstanceStructs[EvalInstanceDataIndex] = EvalInstanceData;
-						TempInstanceStructs[EvalInstanceDataIndex] = TemporaryToBeMovedLater;
+						TemporaryInstanceStructs[EvalInstanceDataIndex] = TemporaryToBeMovedLater;
 					}
 
 					// Global tasks
@@ -2381,7 +2396,7 @@ void FStateTreeExecutionContext::UpdateInstanceData(const TSharedPtr<FSelectStat
 						const FConstStructView TaskInstanceData = CurrentStateTree->DefaultInstanceData.GetStruct(Task.InstanceTemplateIndex.Get());
 						const int32 TaskInstanceDataIndex = BaseIndex + Task.InstanceDataHandle.GetIndex();
 						InstanceStructs[TaskInstanceDataIndex] = TaskInstanceData;
-						TempInstanceStructs[TaskInstanceDataIndex] = TemporaryToBeMovedLater;
+						TemporaryInstanceStructs[TaskInstanceDataIndex] = TemporaryToBeMovedLater;
 					}
 
 					if (bIsFrameCommon)
@@ -2411,7 +2426,7 @@ void FStateTreeExecutionContext::UpdateInstanceData(const TSharedPtr<FSelectStat
 					const bool bIsStateCommon = CurrentFrame->ActiveStates.Contains(SelectedState.GetStateID());
 
 					InstanceStructs.AddDefaulted(CurrentState->InstanceDataNum);
-					TempInstanceStructs.AddZeroed(CurrentState->InstanceDataNum);
+					TemporaryInstanceStructs.AddDefaulted(CurrentState->InstanceDataNum);
 
 					const int32 CurrentStateParameterInstanceDataIndex = CurrentFrameBaseIndex + CurrentState->ParameterDataHandle.GetIndex();
 
@@ -2451,7 +2466,7 @@ void FStateTreeExecutionContext::UpdateInstanceData(const TSharedPtr<FSelectStat
 							check(CurrentState->ParameterDataHandle.IsValid());
 
 							const FCompactStateTreeParameters* Params = nullptr;
-							if (FInstancedStruct* TempParamsInstanceData = FindInstanceTempData(CurrentFrameID, CurrentState->ParameterDataHandle))
+							if (FInstancedStruct* TempParamsInstanceData = FindInstanceTempData(CurrentFrameID, FStateTreeIndex16::Invalid, CurrentState->ParameterDataHandle))
 							{
 								// If we have temp data for the parameters, then setup the instance data with just a type, so that we can steal the temp data below (TempInstanceStructs).
 								// We expect overridden linked assets to hit this code path. 
@@ -2504,11 +2519,11 @@ void FStateTreeExecutionContext::UpdateInstanceData(const TSharedPtr<FSelectStat
 					{
 						if (bCanHaveTempData)
 						{
-							TempInstanceStructs[CurrentStateParameterInstanceDataIndex] = FindInstanceTempData(CurrentFrameID, CurrentState->ParameterDataHandle);
+							TemporaryInstanceStructs[CurrentStateParameterInstanceDataIndex] = ConsumeTemporaryInstanceData(CurrentFrameID, FStateTreeIndex16::Invalid, CurrentState->ParameterDataHandle);
 						}
 						else
 						{
-							TempInstanceStructs[CurrentStateParameterInstanceDataIndex] = NoExistingTemporaryDuplicateWrapper;
+							TemporaryInstanceStructs[CurrentStateParameterInstanceDataIndex] = NoExistingTemporaryDuplicateWrapper;
 						}
 					}
 
@@ -2516,7 +2531,7 @@ void FStateTreeExecutionContext::UpdateInstanceData(const TSharedPtr<FSelectStat
 					{
 						const int32 EventDataInstanceDataIndex = CurrentFrameBaseIndex + CurrentState->EventDataIndex.Get();
 						InstanceStructs[EventDataInstanceDataIndex] = FConstStructView(FStateTreeSharedEvent::StaticStruct());
-						TempInstanceStructs[EventDataInstanceDataIndex] = NoExistingTemporaryDuplicateWrapper;
+						TemporaryInstanceStructs[EventDataInstanceDataIndex] = NoExistingTemporaryDuplicateWrapper;
 					}
 
 					for (int32 TaskIndex = CurrentState->TasksBegin; TaskIndex < (CurrentState->TasksBegin + CurrentState->TasksNum); TaskIndex++)
@@ -2525,7 +2540,7 @@ void FStateTreeExecutionContext::UpdateInstanceData(const TSharedPtr<FSelectStat
 						const FConstStructView TaskInstanceData = CurrentStateTree->DefaultInstanceData.GetStruct(Task.InstanceTemplateIndex.Get());
 						const int32 TaskInstanceDataIndex = CurrentFrameBaseIndex + Task.InstanceDataHandle.GetIndex();
 						InstanceStructs[TaskInstanceDataIndex] = TaskInstanceData;
-						TempInstanceStructs[TaskInstanceDataIndex] = NoExistingTemporaryDuplicateWrapper;
+						TemporaryInstanceStructs[TaskInstanceDataIndex] = NoExistingTemporaryDuplicateWrapper;
 					}
 
 					if (bIsStateCommon)
@@ -2564,10 +2579,16 @@ void FStateTreeExecutionContext::UpdateInstanceData(const TSharedPtr<FSelectStat
 	// Add new instance data.
 	InstanceData.Append(Owner,
 		MakeArrayView(InstanceStructs.GetData() + NumCommonInstanceData, InstanceStructs.Num() - NumCommonInstanceData),
-		MakeArrayView(TempInstanceStructs.GetData() + NumCommonInstanceData, TempInstanceStructs.Num() - NumCommonInstanceData));
+		MakeArrayView(TemporaryInstanceStructs.GetData() + NumCommonInstanceData, TemporaryInstanceStructs.Num() - NumCommonInstanceData));
+
+	// Remove consumed temporary instances. We need to postpone it after the Append because it uses a direct pointer to the memory address.
+	for (const FTemporaryInstanceKey& Key : TemporaryInstanceStructsToConsume)
+	{
+		Storage.RemoveTemporaryInstance(Key.Get<0>(), Key.Get<1>(), Key.Get<2>());
+	}
 }
 
-FStateTreeDataView FStateTreeExecutionContext::MoveTemporaryToInstance(const FStateTreeExecutionFrame& CurrentFrame, const FStateTreeDataHandle DataHandle)
+FStateTreeDataView FStateTreeExecutionContext::MoveTemporaryToInstance(const FStateTreeExecutionFrame& CurrentFrame, const FStateTreeIndex16 NodeIndex, const FStateTreeDataHandle DataHandle)
 {
 	check(DataHandle.GetSource() == EStateTreeDataSourceType::ActiveInstanceData
 		|| DataHandle.GetSource() == EStateTreeDataSourceType::ActiveInstanceDataObject
@@ -2605,6 +2626,7 @@ FStateTreeDataView FStateTreeExecutionContext::MoveTemporaryToInstance(const FSt
 			{
 				check(DataView.GetStruct() == TemporarySource->GetScriptStruct());
 				FMemory::Memswap(DataView.GetMutableMemory(), TemporarySource->GetMutableMemory(), DataView.GetStruct()->GetStructureSize());
+				Storage.RemoveTemporaryInstance(CurrentFrame.FrameID, NodeIndex, DataHandle);
 			}
 		}
 		else if (StructDataHandle != DataHandle)
@@ -3643,7 +3665,7 @@ EStateTreeRunStatus FStateTreeExecutionContext::EnterState(const TSharedPtr<FSel
 			{
 				const int32 AssetTaskIndex = CurrentState.TasksBegin + StateTaskIndex;
 				const FStateTreeTaskBase& Task = CurrentStateTree->Nodes[AssetTaskIndex].Get<const FStateTreeTaskBase>();
-				const FStateTreeDataView TaskInstanceView = MoveTemporaryToInstance(CurrentFrame, Task.InstanceDataHandle);
+				const FStateTreeDataView TaskInstanceView = MoveTemporaryToInstance(CurrentFrame, FStateTreeIndex16(AssetTaskIndex), Task.InstanceDataHandle);
 
 				// Ignore disabled task
 				if (Task.bTaskEnabled == false)
@@ -4131,7 +4153,7 @@ void FStateTreeExecutionContext::MoveGlobalTemporariesToInstances(FStateTreeExec
 	for (int32 EvalIndex = CurrentStateTree->EvaluatorsBegin; EvalIndex < EvaluatorsEnd; ++EvalIndex)
 	{
 		const FStateTreeEvaluatorBase& Eval = CurrentStateTree->Nodes[EvalIndex].Get<const FStateTreeEvaluatorBase>();
-		MoveTemporaryToInstance(Frame, Eval.InstanceDataHandle);
+		MoveTemporaryToInstance(Frame, FStateTreeIndex16(EvalIndex), Eval.InstanceDataHandle);
 	}
 
 	// Start Global tasks
@@ -4141,7 +4163,7 @@ void FStateTreeExecutionContext::MoveGlobalTemporariesToInstances(FStateTreeExec
 	for (int32 AssetTaskIndex = CurrentStateTree->GlobalTasksBegin; AssetTaskIndex < AssetTaskMax; ++AssetTaskIndex)
 	{
 		const FStateTreeTaskBase& Task = CurrentStateTree->Nodes[AssetTaskIndex].Get<const FStateTreeTaskBase>();
-		MoveTemporaryToInstance(Frame, Task.InstanceDataHandle);
+		MoveTemporaryToInstance(Frame, FStateTreeIndex16(AssetTaskIndex), Task.InstanceDataHandle);
 	}
 }
 
@@ -4192,7 +4214,7 @@ EStateTreeRunStatus FStateTreeExecutionContext::StartGlobalsForFrameInternal(con
 		FStateTreeDataView EvalInstanceView;
 		if constexpr (bMoveFromTemporary)
 		{
-			EvalInstanceView = MoveTemporaryToInstance(CurrentFrame, Eval.InstanceDataHandle);
+			EvalInstanceView = MoveTemporaryToInstance(CurrentFrame, FStateTreeIndex16(EvalIndex), Eval.InstanceDataHandle);
 		}
 		else if constexpr (bOnActiveInstances)
 		{
@@ -4261,7 +4283,7 @@ EStateTreeRunStatus FStateTreeExecutionContext::StartGlobalsForFrameInternal(con
 		FStateTreeDataView TaskDataView;
 		if constexpr (bMoveFromTemporary)
 		{
-			TaskDataView = MoveTemporaryToInstance(CurrentFrame, Task.InstanceDataHandle);
+			TaskDataView = MoveTemporaryToInstance(CurrentFrame, FStateTreeIndex16(AssetTaskIndex), Task.InstanceDataHandle);
 		}
 
 		// Ignore disabled task
diff --git a/Engine/Plugins/Runtime/StateTree/Source/StateTreeModule/Private/StateTreeInstanceData.cpp b/Engine/Plugins/Runtime/StateTree/Source/StateTreeModule/Private/StateTreeInstanceData.cpp
--- a/Engine/Plugins/Runtime/StateTree/Source/StateTreeModule/Private/StateTreeInstanceData.cpp
+++ b/Engine/Plugins/Runtime/StateTree/Source/StateTreeModule/Private/StateTreeInstanceData.cpp
@@ -614,6 +614,21 @@ FStructView FStateTreeInstanceStorage::AddTemporaryInstance(UObject& InOwner, co
 	return TempInstance->Instance;
 }
 
+void FStateTreeInstanceStorage::RemoveTemporaryInstance(const UE::StateTree::FActiveFrameID FrameID, const FStateTreeIndex16 OwnerNodeIndex, const FStateTreeDataHandle DataHandle)
+{
+	const int32 Index = TemporaryInstances.IndexOfByPredicate([&FrameID, &OwnerNodeIndex, &DataHandle](const FStateTreeTemporaryInstanceData& TempInstance)
+				{
+					return TempInstance.FrameID == FrameID
+						&& TempInstance.OwnerNodeIndex == OwnerNodeIndex
+						&& TempInstance.DataHandle == DataHandle;
+		});
+
+	if (ensure(Index != INDEX_NONE))
+	{
+		TemporaryInstances.RemoveAtSwap(Index);
+	}
+}
+
 FStructView FStateTreeInstanceStorage::GetMutableTemporaryStruct(UE::StateTree::FActiveFrameID FrameID, const FStateTreeDataHandle DataHandle)
 {
 	FStateTreeTemporaryInstanceData* ExistingInstance = TemporaryInstances.FindByPredicate([FrameID, &DataHandle](const FStateTreeTemporaryInstanceData& TempInstance)
diff --git a/Engine/Plugins/Runtime/StateTree/Source/StateTreeModule/Public/StateTreeExecutionContext.h b/Engine/Plugins/Runtime/StateTree/Source/StateTreeModule/Public/StateTreeExecutionContext.h
--- a/Engine/Plugins/Runtime/StateTree/Source/StateTreeModule/Public/StateTreeExecutionContext.h
+++ b/Engine/Plugins/Runtime/StateTree/Source/StateTreeModule/Public/StateTreeExecutionContext.h
@@ -1410,7 +1410,7 @@ PRAGMA_ENABLE_DEPRECATION_WARNINGS
 	 * Move the data from the temporary buffer to the allocated instance data relative to given frame.
 	 * @return instanced data of the specified handle relative to given frame.
 	 */
-	UE_API FStateTreeDataView MoveTemporaryToInstance(const FStateTreeExecutionFrame& Frame, const FStateTreeDataHandle Handle);
+	UE_API FStateTreeDataView MoveTemporaryToInstance(const FStateTreeExecutionFrame& Frame, const FStateTreeIndex16 NodeIndex, const FStateTreeDataHandle Handle);
 
 	/** @return data view of the specified handle relative to given frame. */
 	[[nodiscard]] UE_API FStateTreeDataView GetDataView(const FStateTreeExecutionFrame* ParentFrame, const FStateTreeExecutionFrame& CurrentFrame, const FStateTreeDataHandle Handle);
diff --git a/Engine/Plugins/Runtime/StateTree/Source/StateTreeModule/Public/StateTreeInstanceData.h b/Engine/Plugins/Runtime/StateTree/Source/StateTreeModule/Public/StateTreeInstanceData.h
--- a/Engine/Plugins/Runtime/StateTree/Source/StateTreeModule/Public/StateTreeInstanceData.h
+++ b/Engine/Plugins/Runtime/StateTree/Source/StateTreeModule/Public/StateTreeInstanceData.h
@@ -290,6 +290,9 @@ struct FStateTreeInstanceStorage
 	 */
 	UE_API FStructView AddTemporaryInstance(UObject& InOwner, const FStateTreeExecutionFrame& Frame, const FStateTreeIndex16 OwnerNodeIndex, const FStateTreeDataHandle DataHandle, FConstStructView NewInstanceData);
 
+	/** Removes a temporary instance data associated with specified frame and data handle. */
+	UE_API void RemoveTemporaryInstance(const UE::StateTree::FActiveFrameID FrameID, const FStateTreeIndex16 OwnerNodeIndex, const FStateTreeDataHandle DataHandle);
+
 	UE_DEPRECATED(5.7, "Use the version with FrameID instead.")
 	FStructView GetMutableTemporaryStruct(const FStateTreeExecutionFrame& Frame, const FStateTreeDataHandle DataHandle) { return GetMutableTemporaryStruct(Frame.FrameID, DataHandle); }
 
