[StateTree] Add trace tag to StateTreeExecutionContext for profiling.

#rnx

diff --git a/Engine/Plugins/Runtime/StateTree/Source/StateTreeModule/Private/StateTreeExecutionContext.cpp b/Engine/Plugins/Runtime/StateTree/Source/StateTreeModule/Private/StateTreeExecutionContext.cpp
--- a/Engine/Plugins/Runtime/StateTree/Source/StateTreeModule/Private/StateTreeExecutionContext.cpp
+++ b/Engine/Plugins/Runtime/StateTree/Source/StateTreeModule/Private/StateTreeExecutionContext.cpp
@@ -1398,6 +1398,7 @@ void FStateTreeExecutionContext::SetUpdatePhaseInExecutionState(FStateTreeExecut
 EStateTreeRunStatus FStateTreeExecutionContext::Start(FStartParameters Parameters)
 {
 	CSV_SCOPED_TIMING_STAT_EXCLUSIVE(StateTree_Start);
+	TRACE_CPUPROFILER_EVENT_SCOPE(FStateTreeExecutionContext::Start);
 	UE_STATETREE_CRASH_REPORTER_SCOPE(&Owner, &RootStateTree, UE::StateTree::ExecutionContext::Private::Name_Start.Resolve());
 
 	using namespace UE::StateTree;
@@ -1598,6 +1599,7 @@ EStateTreeRunStatus FStateTreeExecutionContext::Start(FStartParameters Parameter
 EStateTreeRunStatus FStateTreeExecutionContext::Stop(EStateTreeRunStatus CompletionStatus)
 {
 	CSV_SCOPED_TIMING_STAT_EXCLUSIVE(StateTree_Stop);
+	TRACE_CPUPROFILER_EVENT_SCOPE(FStateTreeExecutionContext::Stop);
 	UE_STATETREE_CRASH_REPORTER_SCOPE(&Owner, &RootStateTree, UE::StateTree::ExecutionContext::Private::Name_Stop.Resolve());
 
 	if (!IsValid())
@@ -1733,6 +1735,7 @@ EStateTreeRunStatus FStateTreeExecutionContext::TickPostlude()
 EStateTreeRunStatus FStateTreeExecutionContext::Tick(const float DeltaTime)
 {
 	CSV_SCOPED_TIMING_STAT_EXCLUSIVE(StateTree_Tick);
+	TRACE_CPUPROFILER_EVENT_SCOPE(FStateTreeExecutionContext::Tick);
 	UE_STATETREE_CRASH_REPORTER_SCOPE(&Owner, &RootStateTree, UE::StateTree::ExecutionContext::Private::Name_Tick.Resolve());
 
 	TGuardValue<bool> ScheduledNextTickScope(bAllowedToScheduleNextTick, false);
@@ -1752,6 +1755,7 @@ EStateTreeRunStatus FStateTreeExecutionContext::Tick(const float DeltaTime)
 EStateTreeRunStatus FStateTreeExecutionContext::TickUpdateTasks(const float DeltaTime)
 {
 	CSV_SCOPED_TIMING_STAT_EXCLUSIVE(StateTree_Tick);
+	TRACE_CPUPROFILER_EVENT_SCOPE(FStateTreeExecutionContext::TickUpdateTasks);
 	UE_STATETREE_CRASH_REPORTER_SCOPE(&Owner, &RootStateTree, UE::StateTree::ExecutionContext::Private::Name_Tick.Resolve());
 
 	TGuardValue<bool> ScheduledNextTickScope(bAllowedToScheduleNextTick, false);
@@ -1770,6 +1774,7 @@ EStateTreeRunStatus FStateTreeExecutionContext::TickUpdateTasks(const float Delt
 EStateTreeRunStatus FStateTreeExecutionContext::TickTriggerTransitions()
 {
 	CSV_SCOPED_TIMING_STAT_EXCLUSIVE(StateTree_Tick);
+	TRACE_CPUPROFILER_EVENT_SCOPE(FStateTreeExecutionContext::TickTriggerTransitions);
 	UE_STATETREE_CRASH_REPORTER_SCOPE(&Owner, &RootStateTree, UE::StateTree::ExecutionContext::Private::Name_Tick.Resolve());
 
 	TGuardValue<bool> ScheduledNextTickScope(bAllowedToScheduleNextTick, false);
