[StateTree] Add new transition targets. Parent, NextParent, and NextSelectableParent. They are reconsiled at compile time and work similar to NextState and NextSelectableState.

diff --git a/Engine/Plugins/Runtime/StateTree/Source/StateTreeEditorModule/Private/Customizations/StateTreeStateLinkDetails.cpp b/Engine/Plugins/Runtime/StateTree/Source/StateTreeEditorModule/Private/Customizations/StateTreeStateLinkDetails.cpp
--- a/Engine/Plugins/Runtime/StateTree/Source/StateTreeEditorModule/Private/Customizations/StateTreeStateLinkDetails.cpp
+++ b/Engine/Plugins/Runtime/StateTree/Source/StateTreeEditorModule/Private/Customizations/StateTreeStateLinkDetails.cpp
@@ -161,70 +161,30 @@ TSharedRef<SWidget> FStateTreeStateLinkDetails::GenerateStatePicker()
 				];
 		};
 
-		MenuBuilder.AddMenuEntry(
-			FUIAction(
-				FExecuteAction::CreateSP(this, &FStateTreeStateLinkDetails::SetTransitionByType, EStateTreeTransitionType::None),
-				FCanExecuteAction(),
-				FIsActionChecked::CreateSPLambda(this, [this]
-				{
-					return GetTransitionType() == EStateTreeTransitionType::None; 
-				})),
-			MakeMetaStateWidget(EStateTreeTransitionType::None),
-			FName(),
-			LOCTEXT("TransitionNoneTooltip", "No transition."),
-			EUserInterfaceActionType::Check);
-
-		MenuBuilder.AddMenuEntry(
-			FUIAction(
-				FExecuteAction::CreateSP(this, &FStateTreeStateLinkDetails::SetTransitionByType, EStateTreeTransitionType::NextState),
-				FCanExecuteAction(),
-				FIsActionChecked::CreateSPLambda(this, [this]
-				{
-					return GetTransitionType() == EStateTreeTransitionType::NextState; 
-				})),
-			MakeMetaStateWidget(EStateTreeTransitionType::NextState),
-			FName(),
-			LOCTEXT("TransitionNextTooltip", "Goto next sibling State."),
-			EUserInterfaceActionType::Check);
-
-		MenuBuilder.AddMenuEntry(
-			FUIAction(
-				FExecuteAction::CreateSP(this, &FStateTreeStateLinkDetails::SetTransitionByType, EStateTreeTransitionType::NextSelectableState),
-				FCanExecuteAction(),
-				FIsActionChecked::CreateSPLambda(this, [this]
-				{
-					return GetTransitionType() == EStateTreeTransitionType::NextSelectableState; 
-			})),
-			MakeMetaStateWidget(EStateTreeTransitionType::NextSelectableState),
-			FName(),
-			LOCTEXT("TransitionNextSelectableTooltip", "Goto next sibling state, whose enter conditions pass."),
-			EUserInterfaceActionType::Check);
-
-		MenuBuilder.AddMenuEntry(
-			FUIAction(
-				FExecuteAction::CreateSP(this, &FStateTreeStateLinkDetails::SetTransitionByType, EStateTreeTransitionType::Succeeded),
-				FCanExecuteAction(),
-				FIsActionChecked::CreateSPLambda(this, [this]
-				{
-					return GetTransitionType() == EStateTreeTransitionType::Succeeded; 
-				})),
-			MakeMetaStateWidget(EStateTreeTransitionType::Succeeded),
-			FName(),
-			LOCTEXT("TransitionTreeSuccessTooltip", "Complete tree with success."),
-			EUserInterfaceActionType::Check);
-
-		MenuBuilder.AddMenuEntry(
-			FUIAction(
-				FExecuteAction::CreateSP(this, &FStateTreeStateLinkDetails::SetTransitionByType, EStateTreeTransitionType::Failed),
-				FCanExecuteAction(),
-				FIsActionChecked::CreateSPLambda(this, [this]
-				{
-					return GetTransitionType() == EStateTreeTransitionType::Failed; 
-				})),
-			MakeMetaStateWidget(EStateTreeTransitionType::Failed),
-			FName(),
-			LOCTEXT("TransitionTreeFailedTooltip", "Complete tree with failure."),
-			EUserInterfaceActionType::Check);
+		auto AddMenuEntry = [this, &MakeMetaStateWidget, &MenuBuilder](EStateTreeTransitionType Type, FText Tooltip)
+		{
+			MenuBuilder.AddMenuEntry(
+				FUIAction(
+					FExecuteAction::CreateSP(this, &FStateTreeStateLinkDetails::SetTransitionByType, Type),
+					FCanExecuteAction(),
+					FIsActionChecked::CreateSPLambda(this, [this, Type]
+					{
+						return GetTransitionType() == Type;
+					})),
+				MakeMetaStateWidget(Type),
+				FName(),
+				Tooltip,
+				EUserInterfaceActionType::Check);
+		};
+
+		AddMenuEntry(EStateTreeTransitionType::None, LOCTEXT("TransitionNoneTooltip", "No transition."));
+		AddMenuEntry(EStateTreeTransitionType::Parent, LOCTEXT("TransitionParentTooltip", "Goto the parent state."));
+		AddMenuEntry(EStateTreeTransitionType::NextState, LOCTEXT("TransitionNextTooltip", "Goto next sibling State."));
+		AddMenuEntry(EStateTreeTransitionType::NextSelectableState, LOCTEXT("TransitionNextSelectableTooltip", "Goto next sibling state, whose enter conditions pass."));
+		AddMenuEntry(EStateTreeTransitionType::NextParent, LOCTEXT("TransitionNextParentTooltip", "Goto next sibling of the parent State."));
+		AddMenuEntry(EStateTreeTransitionType::NextSelectableParent, LOCTEXT("TransitionNextParentSelectableTooltip", "Goto next sibling of the parent state, whose enter conditions pass."));
+		AddMenuEntry(EStateTreeTransitionType::Succeeded, LOCTEXT("TransitionTreeSuccessTooltip", "Complete tree with success."));
+		AddMenuEntry(EStateTreeTransitionType::Failed, LOCTEXT("TransitionTreeFailedTooltip", "Complete tree with failure."));
 	}
 
 	MenuBuilder.BeginSection("States", LOCTEXT("States", "States"));
diff --git a/Engine/Plugins/Runtime/StateTree/Source/StateTreeEditorModule/Private/SStateTreeViewRow.cpp b/Engine/Plugins/Runtime/StateTree/Source/StateTreeEditorModule/Private/SStateTreeViewRow.cpp
--- a/Engine/Plugins/Runtime/StateTree/Source/StateTreeEditorModule/Private/SStateTreeViewRow.cpp
+++ b/Engine/Plugins/Runtime/StateTree/Source/StateTreeEditorModule/Private/SStateTreeViewRow.cpp
@@ -1971,6 +1971,11 @@ const FSlateBrush* SStateTreeViewRow::GetTransitionsIcon(const EStateTreeTransit
 			case EStateTreeTransitionType::NextSelectableState:
 				IconType |= IconNext;
 				break;
+			case EStateTreeTransitionType::Parent:
+			case EStateTreeTransitionType::NextParent:
+			case EStateTreeTransitionType::NextSelectableParent:
+				IconType |= IconParent;
+				break;
 			case EStateTreeTransitionType::GotoState:
 				IconType |= IconGoto;
 				break;
diff --git a/Engine/Plugins/Runtime/StateTree/Source/StateTreeEditorModule/Private/StateTreeCompiler.cpp b/Engine/Plugins/Runtime/StateTree/Source/StateTreeEditorModule/Private/StateTreeCompiler.cpp
--- a/Engine/Plugins/Runtime/StateTree/Source/StateTreeEditorModule/Private/StateTreeCompiler.cpp
+++ b/Engine/Plugins/Runtime/StateTree/Source/StateTreeEditorModule/Private/StateTreeCompiler.cpp
@@ -2148,7 +2148,27 @@ bool FStateTreeCompiler::ResolveTransitionStateAndFallback(const UStateTreeState
 			return false;
 		}
 	}
-	else if (Link.LinkType == EStateTreeTransitionType::NextState || Link.LinkType == EStateTreeTransitionType::NextSelectableState)
+	else if (Link.LinkType == EStateTreeTransitionType::Parent)
+	{
+		// Find next parent.
+		const UStateTreeState* Parent = SourceState && SourceState->Parent ? SourceState->Parent : nullptr;
+		if (Parent == nullptr)
+		{
+			Log.Reportf(EMessageSeverity::Error,
+				TEXT("Failed to resolve transition, there's no selectable parent state."));
+			return false;
+		}
+		OutTransitionHandle = GetStateHandle(Parent->ID);
+		if (!OutTransitionHandle.IsValid())
+		{
+			Log.Reportf(EMessageSeverity::Error,
+				TEXT("Failed to resolve transition parent state, no handle found for '%s'."),
+				*Parent->Name.ToString());
+			return false;
+		}
+	}
+	else if (Link.LinkType == EStateTreeTransitionType::NextState
+		|| Link.LinkType == EStateTreeTransitionType::NextSelectableState)
 	{
 		// Find next state.
 		const UStateTreeState* NextState = SourceState ? SourceState->GetNextSelectableSiblingState() : nullptr;
@@ -2167,6 +2187,26 @@ bool FStateTreeCompiler::ResolveTransitionStateAndFallback(const UStateTreeState
 			return false;
 		}
 	}
+	else if (Link.LinkType == EStateTreeTransitionType::NextParent
+		|| Link.LinkType == EStateTreeTransitionType::NextSelectableParent)
+	{
+		// Find next parent.
+		const UStateTreeState* NextParent = SourceState && SourceState->Parent ? SourceState->Parent->GetNextSelectableSiblingState() : nullptr;
+		if (NextParent == nullptr)
+		{
+			Log.Reportf(EMessageSeverity::Error,
+				TEXT("Failed to resolve transition, there's no selectable next state for the parent state."));
+			return false;
+		}
+		OutTransitionHandle = GetStateHandle(NextParent->ID);
+		if (!OutTransitionHandle.IsValid())
+		{
+			Log.Reportf(EMessageSeverity::Error,
+				TEXT("Failed to resolve transition next state, no handle found for '%s'."),
+				*NextParent->Name.ToString());
+			return false;
+		}
+	}
 	else if(Link.LinkType == EStateTreeTransitionType::Failed)
 	{
 		OutTransitionHandle = FStateTreeStateHandle::Failed;
@@ -2186,7 +2226,9 @@ PRAGMA_DISABLE_DEPRECATION_WARNINGS
 	}
 PRAGMA_ENABLE_DEPRECATION_WARNINGS
 
-	if (Link.LinkType == EStateTreeTransitionType::NextSelectableState)
+	// Set fallback type
+	if (Link.LinkType == EStateTreeTransitionType::NextSelectableState
+		|| Link.LinkType == EStateTreeTransitionType::NextSelectableParent)
 	{
 		OutFallback = EStateTreeSelectionFallback::NextSelectableSibling;
 	}
diff --git a/Engine/Plugins/Runtime/StateTree/Source/StateTreeEditorModule/Private/StateTreeDescriptionHelpers.cpp b/Engine/Plugins/Runtime/StateTree/Source/StateTreeEditorModule/Private/StateTreeDescriptionHelpers.cpp
--- a/Engine/Plugins/Runtime/StateTree/Source/StateTreeEditorModule/Private/StateTreeDescriptionHelpers.cpp
+++ b/Engine/Plugins/Runtime/StateTree/Source/StateTreeEditorModule/Private/StateTreeDescriptionHelpers.cpp
@@ -26,6 +26,12 @@ FText GetStateLinkDesc(const UStateTreeEditorData* EditorData, const FStateTreeS
 			? LOCTEXT("TransitionNoneRich", "<i>None</>")
 			: LOCTEXT("TransitionNone", "None");
 	}
+	if (Link.LinkType == EStateTreeTransitionType::Parent)
+	{
+		return Formatting == EStateTreeNodeFormatting::RichText
+			? LOCTEXT("TransitionParentRich", "<i>Parent</>")
+			: LOCTEXT("TransitionParent", "Parent");
+	}
 	if (Link.LinkType == EStateTreeTransitionType::NextState)
 	{
 		return Formatting == EStateTreeNodeFormatting::RichText
@@ -38,6 +44,18 @@ FText GetStateLinkDesc(const UStateTreeEditorData* EditorData, const FStateTreeS
 			? LOCTEXT("TransitionNextSelectableStateRich", "<i>Next Selectable State</>")
 			: LOCTEXT("TransitionNextSelectableState", "Next Selectable State");
 	}
+	if (Link.LinkType == EStateTreeTransitionType::NextParent)
+	{
+		return Formatting == EStateTreeNodeFormatting::RichText
+			? LOCTEXT("TransitionNextParentRich", "<i>Next Parent</>")
+			: LOCTEXT("TransitionNextParent", "Next Parent");
+	}
+	if (Link.LinkType == EStateTreeTransitionType::NextSelectableParent)
+	{
+		return Formatting == EStateTreeNodeFormatting::RichText
+			? LOCTEXT("TransitionNextSelectableParentRich", "<i>Next Selectable Parent</>")
+			: LOCTEXT("TransitionNextSelectableParent", "Next Selectable Parent");
+	}
 	if (Link.LinkType == EStateTreeTransitionType::Succeeded)
 	{
 		return Formatting == EStateTreeNodeFormatting::RichText
@@ -86,6 +104,10 @@ const FSlateBrush* GetStateLinkIcon(const UStateTreeEditorData* EditorData, cons
 	{
 		return FStateTreeEditorStyle::Get().GetBrush("StateTreeEditor.Transition.None");
 	}
+	if (Link.LinkType == EStateTreeTransitionType::Parent)
+	{
+		return FStateTreeEditorStyle::Get().GetBrush("StateTreeEditor.Transition.Parent");
+	}
 	if (Link.LinkType == EStateTreeTransitionType::NextState)
 	{
 		return FStateTreeEditorStyle::Get().GetBrush("StateTreeEditor.Transition.Next");
@@ -94,6 +116,14 @@ const FSlateBrush* GetStateLinkIcon(const UStateTreeEditorData* EditorData, cons
 	{
 		return FStateTreeEditorStyle::Get().GetBrush("StateTreeEditor.Transition.Next");
 	}
+	if (Link.LinkType == EStateTreeTransitionType::NextParent)
+	{
+		return FStateTreeEditorStyle::Get().GetBrush("StateTreeEditor.Transition.Parent");
+	}
+	if (Link.LinkType == EStateTreeTransitionType::NextSelectableParent)
+	{
+		return FStateTreeEditorStyle::Get().GetBrush("StateTreeEditor.Transition.Parent");
+	}
 	if (Link.LinkType == EStateTreeTransitionType::Succeeded)
 	{
 		return FStateTreeEditorStyle::Get().GetBrush("StateTreeEditor.Transition.Succeeded");
diff --git a/Engine/Plugins/Runtime/StateTree/Source/StateTreeModule/Public/StateTreeTypes.h b/Engine/Plugins/Runtime/StateTree/Source/StateTreeModule/Public/StateTreeTypes.h
--- a/Engine/Plugins/Runtime/StateTree/Source/StateTreeModule/Public/StateTreeTypes.h
+++ b/Engine/Plugins/Runtime/StateTree/Source/StateTreeModule/Public/StateTreeTypes.h
@@ -75,13 +75,22 @@ enum class EStateTreeTransitionType : uint8
 	
 	/** Transition to the specified state. */
 	GotoState,
+
+	/** Transition to the parent state. */
+	Parent,
 	
 	/** Transition to the next sibling state. */
 	NextState,
 
-	/** Transition to the next selectable sibling state */
+	/** Transition to the next selectable sibling state. */
 	NextSelectableState,
 
+	/** Transition to the next sibling state of the parent state. */
+	NextParent,
+
+	/** Transition to the next selectable sibling state of the parent state. */
+	NextSelectableParent,
+
 	NotSet UE_DEPRECATED(all, "Use None instead."),
 };
 
diff --git a/Engine/Plugins/Runtime/StateTree/Source/StateTreeTestSuite/Private/StateTreeTestTypes.h b/Engine/Plugins/Runtime/StateTree/Source/StateTreeTestSuite/Private/StateTreeTestTypes.h
--- a/Engine/Plugins/Runtime/StateTree/Source/StateTreeTestSuite/Private/StateTreeTestTypes.h
+++ b/Engine/Plugins/Runtime/StateTree/Source/StateTreeTestSuite/Private/StateTreeTestTypes.h
@@ -43,6 +43,11 @@ struct FStateTreeTestLog
 struct FTestStateTreeExecutionContext : public FStateTreeExecutionContext
 {
 	FTestStateTreeExecutionContext(UObject& InOwner, const UStateTree& InStateTree, FStateTreeInstanceData& InInstanceData, const FOnCollectStateTreeExternalData& CollectExternalDataCallback = {}, const EStateTreeRecordTransitions RecordTransitions = EStateTreeRecordTransitions::No)
+		: FTestStateTreeExecutionContext(&InOwner, &InStateTree, InInstanceData, CollectExternalDataCallback, RecordTransitions)
+	{
+	}
+
+	FTestStateTreeExecutionContext(TNotNull<UObject*> InOwner, TNotNull<const UStateTree*> InStateTree, FStateTreeInstanceData& InInstanceData, const FOnCollectStateTreeExternalData& CollectExternalDataCallback = {}, const EStateTreeRecordTransitions RecordTransitions = EStateTreeRecordTransitions::No)
 		: FStateTreeExecutionContext(InOwner, InStateTree, InInstanceData, CollectExternalDataCallback, RecordTransitions)
 	{
 		// Handle getting pointer to the test log.
diff --git a/Engine/Plugins/Runtime/StateTree/Source/StateTreeTestSuite/Private/StateTreeTransitionTest.cpp b/Engine/Plugins/Runtime/StateTree/Source/StateTreeTestSuite/Private/StateTreeTransitionTest.cpp
--- a/Engine/Plugins/Runtime/StateTree/Source/StateTreeTestSuite/Private/StateTreeTransitionTest.cpp
+++ b/Engine/Plugins/Runtime/StateTree/Source/StateTreeTestSuite/Private/StateTreeTransitionTest.cpp
@@ -1059,7 +1059,6 @@ struct FStateTreeTest_RegularTransitions : FStateTreeTestBase
 };
 IMPLEMENT_STATE_TREE_INSTANT_TEST(FStateTreeTest_RegularTransitions, "System.StateTree.Transition.RegularTransitions");
 
-
 struct FStateTreeTest_RequestTransition : FStateTreeTestBase
 {
 	virtual bool InstantTest() override
@@ -2408,6 +2407,311 @@ struct FStateTreeTest_MultipleTransition_TemporaryFrame_GlobalTask : FStateTreeT
 };
 IMPLEMENT_STATE_TREE_INSTANT_TEST(FStateTreeTest_MultipleTransition_TemporaryFrame_GlobalTask, "System.StateTree.Transition.MultipleTransition.TemporaryFrame.GlobalTask");
 
+struct FStateTreeTest_TransitionTypes : FStateTreeTestBase
+{
+	virtual bool InstantTest() override
+	{
+		//Root
+		// State1
+		//  State1C1
+		//  State1C2
+		//  State1C3
+		// State2
+		// State3
+		struct FTree
+		{
+			UStateTree* StateTree = nullptr;
+			UStateTreeState* Root = nullptr;
+			UStateTreeState* State1 = nullptr;
+			UStateTreeState* State1C1 = nullptr;
+			UStateTreeState* State1C2 = nullptr;
+			UStateTreeState* State1C3 = nullptr;
+			UStateTreeState* State2 = nullptr;
+			UStateTreeState* State3 = nullptr;
+			FGuid EvaluatorA_ID;
+		};
+
+		auto Compile = [this](TNotNull<UStateTree*> StateTree, bool bExpectedToCompile)
+			{
+				FStateTreeCompilerLog Log;
+				FStateTreeCompiler Compiler(Log);
+				const bool bResult = Compiler.Compile(StateTree);
+				if (bExpectedToCompile)
+				{
+					AITEST_TRUE(TEXT("StateTree should compiled"), bResult);
+					return bResult;
+				}
+				else
+				{
+					AITEST_FALSE(TEXT("StateTree should not compiled"), bResult);
+					return !bResult;
+				}
+			};
+
+
+		auto Start = [this](FTestStateTreeExecutionContext& Exec)
+			{
+				const bool bInitSucceeded = Exec.IsValid();
+				AITEST_TRUE("StateTree should init", bInitSucceeded);
+
+				EStateTreeRunStatus Status = Exec.Start();
+				AITEST_TRUE(TEXT("State1C1 should enter"), Exec.Expect("State1C1Task2", TEXT("EnterState")));
+				Exec.LogClear();
+				return true;
+			};
+
+		auto MakeBasicTree = [this]()
+			{
+				FTree Result;
+				Result.StateTree = &NewStateTree();
+				UStateTreeEditorData* EditorData = Cast<UStateTreeEditorData>(Result.StateTree->EditorData);
+
+				Result.Root = &EditorData->AddSubTree(FName("Root"));
+				Result.State1 = &Result.Root->AddChildState(FName("State1"));
+				Result.State1C1 = &Result.State1->AddChildState(FName("State1C1"));
+				Result.State1C2 = &Result.State1->AddChildState(FName("State1C2"));
+				Result.State1C3 = &Result.State1->AddChildState(FName("State1C3"));
+				Result.State2 = &Result.Root->AddChildState(FName("State2"));
+				Result.State3 = &Result.Root->AddChildState(FName("State3"));
+
+				// Tasks
+				Result.State1C1->AddTask<FTestTask_PrintValue>(FName(TEXT("State1C1Task1"))).GetInstanceData().Value = 1;
+				Result.State1C1->AddTask<FTestTask_Stand>(FName(TEXT("State1C1Task2"))).GetNode().TicksToCompletion = 1;
+				Result.State1C2->AddTask<FTestTask_PrintValue>(FName(TEXT("State1C2Task1"))).GetInstanceData().Value = 2;
+				Result.State1C2->AddTask<FTestTask_Stand>(FName(TEXT("State1C2Task2"))).GetNode().TicksToCompletion = 1;
+				Result.State1C3->AddTask<FTestTask_PrintValue>(FName(TEXT("State1C3Task1"))).GetInstanceData().Value = 3;
+				Result.State1C3->AddTask<FTestTask_Stand>(FName(TEXT("State1C3Task2"))).GetNode().TicksToCompletion = 1;
+				Result.State2->AddTask<FTestTask_PrintValue>(FName(TEXT("State2Task1"))).GetInstanceData().Value = 4;
+				Result.State3->AddTask<FTestTask_PrintValue>(FName(TEXT("State3Task1"))).GetInstanceData().Value = 5;
+
+				// Evaluators
+				TStateTreeEditorNode<FTestEval_A>& Evaluator = EditorData->AddEvaluator<FTestEval_A>(FName("Evaluator1"));
+				Evaluator.GetInstanceData().bBoolA = true;
+				Result.EvaluatorA_ID = Evaluator.ID;
+
+				return Result;
+			};
+
+		auto AddPropertyBinding = [this](TNotNull<UStateTree*> StateTree, const FGuid& SourceID, FName SourcePath, const FGuid& TargetID, FName TargetPath)
+			{
+				UStateTreeEditorData* EditorData = Cast<UStateTreeEditorData>(StateTree->EditorData);
+				EditorData->AddPropertyBinding(FPropertyBindingPath(SourceID, SourcePath), FPropertyBindingPath(TargetID, TargetPath));
+				return true;
+			};
+
+		// None -> state should remain the same state after transition trigger
+		{
+			FTree Tree = MakeBasicTree();
+
+			Tree.State1C1->AddTransition(EStateTreeTransitionTrigger::OnStateCompleted, EStateTreeTransitionType::None);
+
+			if (!Compile(Tree.StateTree, /*bExpectedToCompile*/false))
+			{
+				return false;
+			}
+		}
+
+		// NextState & NextParent -> should select next sibling State1C2, State1C3
+		{
+			FTree Tree = MakeBasicTree();
+
+			Tree.State1C1->AddTransition(EStateTreeTransitionTrigger::OnStateCompleted, EStateTreeTransitionType::NextState);
+			Tree.State1C2->AddTransition(EStateTreeTransitionTrigger::OnStateCompleted, EStateTreeTransitionType::NextState);
+			// State1C3 NextState doesn't compile.
+			Tree.State1C3->AddTransition(EStateTreeTransitionTrigger::OnStateCompleted, EStateTreeTransitionType::NextState);
+
+			if (!Compile(Tree.StateTree, /*bExpectedToCompile*/false))
+			{
+				return false;
+			}
+
+			Tree.State1C3->Transitions.Reset();
+			Tree.State1C3->AddTransition(EStateTreeTransitionTrigger::OnStateCompleted, EStateTreeTransitionType::NextParent);
+
+			if (!Compile(Tree.StateTree, /*bExpectedToCompile*/true))
+			{
+				return false;
+			}
+
+			FStateTreeInstanceData InstanceData;
+			FTestStateTreeExecutionContext Exec(Tree.StateTree, Tree.StateTree, InstanceData);
+			if (!Start(Exec))
+			{
+				return false;
+			}
+
+			// Go to State1C2
+			EStateTreeRunStatus Status = Exec.Tick(0.1f);
+			AITEST_TRUE(TEXT("State1C1Task1 should complete"), Exec.Expect("State1C1Task2", TEXT("StateCompleted")));
+			AITEST_TRUE(TEXT("Should transition to in State1C2"), Exec.ExpectInActiveStates("Root", "State1", "State1C2"));
+			Exec.LogClear();
+
+			// Go to State1C3
+			Status = Exec.Tick(0.1f);
+			AITEST_TRUE(TEXT("State1C1Task1 should complete"), Exec.Expect("State1C2Task2", TEXT("StateCompleted")));
+			AITEST_TRUE(TEXT("Should transition to in State1C2"), Exec.ExpectInActiveStates("Root", "State1", "State1C3"));
+			Exec.LogClear();
+
+			// Go to State2
+			Status = Exec.Tick(0.1f);
+			AITEST_TRUE(TEXT("State1C1Task1 should complete"), Exec.Expect("State1C3Task2", TEXT("StateCompleted")));
+			AITEST_TRUE(TEXT("Should transition to in State1C2"), Exec.ExpectInActiveStates("Root", "State2"));
+			Exec.LogClear();
+
+			Exec.Stop();
+		}
+
+		// NextState -> select next sibling (State1C2) but fail because can't enter. Use NextSelectableState.
+		{
+			FTree Tree = MakeBasicTree();
+
+			Tree.State1C1->AddTransition(EStateTreeTransitionTrigger::OnStateCompleted, EStateTreeTransitionType::NextState);
+			TStateTreeEditorNode<FStateTreeCompareBoolCondition>& State1C2Condition = Tree.State1C2->AddEnterCondition<FStateTreeCompareBoolCondition>();
+			State1C2Condition.GetNode().EvaluationMode = EStateTreeConditionEvaluationMode::ForcedFalse;
+			State1C2Condition.GetInstanceData().bLeft = false;
+
+			if (!AddPropertyBinding(Tree.StateTree, Tree.EvaluatorA_ID, FName("bBoolA"), State1C2Condition.ID, FName("bLeft")))
+			{
+				return false;
+			}
+
+			if (!Compile(Tree.StateTree, /*bExpectedToCompile*/true))
+			{
+				return false;
+			}
+
+			{
+				FStateTreeInstanceData InstanceData;
+				FTestStateTreeExecutionContext Exec(Tree.StateTree, Tree.StateTree, InstanceData);
+				if (!Start(Exec))
+				{
+					return false;
+				}
+
+				const EStateTreeRunStatus Status = Exec.Tick(1.1f);
+				AITEST_TRUE(TEXT("State1C1Task2 should complete"), Exec.Expect("State1C1Task2", TEXT("StateCompleted")));
+				AITEST_TRUE(TEXT("State1C1Task2 should exit"), Exec.Expect("State1C1Task2", TEXT("ExitState")));
+				AITEST_TRUE(TEXT("State1C1Task2 should enter"), Exec.Expect("State1C1Task2", TEXT("EnterState")));
+				AITEST_FALSE(TEXT("State1C2Task2 should enter"), Exec.Expect("State1C2Task2", TEXT("EnterState")));
+				AITEST_TRUE(TEXT("Should transition to root and back"), Exec.ExpectInActiveStates("Root", "State1", "State1C1"));
+				Exec.LogClear();
+
+				Exec.Stop();
+			}
+
+			// Use NextSelectableState (go to C3)
+			Tree.State1C1->Transitions.Reset();
+			Tree.State1C1->AddTransition(EStateTreeTransitionTrigger::OnStateCompleted, EStateTreeTransitionType::NextSelectableState);
+			if (!Compile(Tree.StateTree, /*bExpectedToCompile*/true))
+			{
+				return false;
+			}
+
+			{
+				FStateTreeInstanceData InstanceData;
+				FTestStateTreeExecutionContext Exec(Tree.StateTree, Tree.StateTree, InstanceData);
+				if (!Start(Exec))
+				{
+					return false;
+				}
+
+				const EStateTreeRunStatus Status = Exec.Tick(1.1f);
+				AITEST_TRUE(TEXT("State1C1Task2 should complete"), Exec.Expect("State1C1Task2", TEXT("StateCompleted")));
+				AITEST_TRUE(TEXT("State1C1Task2 should exit"), Exec.Expect("State1C1Task2", TEXT("ExitState")));
+				AITEST_FALSE(TEXT("State1C2Task2 should enter"), Exec.Expect("State1C2Task2", TEXT("EnterState")));
+				AITEST_TRUE(TEXT("State1C3Task2 should enter"), Exec.Expect("State1C3Task2", TEXT("EnterState")));
+				AITEST_TRUE(TEXT("Should transition to State1C3"), Exec.ExpectInActiveStates("Root", "State1", "State1C3"));
+				Exec.LogClear();
+
+				Exec.Stop();
+			}
+		}
+
+		// NextParent -> select next parent (State2) but fail because can't enter. Use NextSelectableParent.
+		{
+			FTree Tree = MakeBasicTree();
+
+			Tree.State1C1->AddTransition(EStateTreeTransitionTrigger::OnStateCompleted, EStateTreeTransitionType::NextParent);
+			TStateTreeEditorNode<FStateTreeCompareBoolCondition>& State2Condition = Tree.State2->AddEnterCondition<FStateTreeCompareBoolCondition>();
+			State2Condition.GetNode().EvaluationMode = EStateTreeConditionEvaluationMode::ForcedFalse;
+			State2Condition.GetInstanceData().bLeft = false;
+
+			if (!AddPropertyBinding(Tree.StateTree, Tree.EvaluatorA_ID, FName("bBoolA"), State2Condition.ID, FName("bLeft")))
+			{
+				return false;
+			}
+
+			if (!Compile(Tree.StateTree, /*bExpectedToCompile*/true))
+			{
+				return false;
+			}
+
+			{
+				FStateTreeInstanceData InstanceData;
+				FTestStateTreeExecutionContext Exec(Tree.StateTree, Tree.StateTree, InstanceData);
+				if (!Start(Exec))
+				{
+					return false;
+				}
+
+				const EStateTreeRunStatus Status = Exec.Tick(1.1f);
+				AITEST_TRUE(TEXT("State1C1Task2 should complete"), Exec.Expect("State1C1Task2", TEXT("StateCompleted")));
+				AITEST_TRUE(TEXT("State1C1Task2 should exit"), Exec.Expect("State1C1Task2", TEXT("ExitState")));
+				AITEST_TRUE(TEXT("State1C1Task2 should enter"), Exec.Expect("State1C1Task2", TEXT("EnterState")));
+				AITEST_FALSE(TEXT("State2Task1 should not enter"), Exec.Expect("State2Task1", TEXT("EnterState4")));
+				AITEST_TRUE(TEXT("Should transition to root and back"), Exec.ExpectInActiveStates("Root", "State1", "State1C1"));
+				Exec.LogClear();
+
+				Exec.Stop();
+			}
+
+			// Use NextSelectableParent
+			Tree.State1C1->Transitions.Reset();
+			Tree.State1C1->AddTransition(EStateTreeTransitionTrigger::OnStateCompleted, EStateTreeTransitionType::NextSelectableParent);
+			if (!Compile(Tree.StateTree, /*bExpectedToCompile*/true))
+			{
+				return false;
+			}
+
+			{
+				FStateTreeInstanceData InstanceData;
+				FTestStateTreeExecutionContext Exec(Tree.StateTree, Tree.StateTree, InstanceData);
+				if (!Start(Exec))
+				{
+					return false;
+				}
+
+				const EStateTreeRunStatus Status = Exec.Tick(1.1f);
+				AITEST_TRUE(TEXT("State1C1Task2 should complete"), Exec.Expect("State1C1Task2", TEXT("StateCompleted")));
+				AITEST_TRUE(TEXT("State1C1Task2 should exit"), Exec.Expect("State1C1Task2", TEXT("ExitState")));
+				AITEST_FALSE(TEXT("State1C2Task2 should not enter"), Exec.Expect("State1C2Task2", TEXT("EnterState")));
+				AITEST_FALSE(TEXT("State2Task1 should not enter"), Exec.Expect("State2Task1", TEXT("EnterState4")));
+				AITEST_TRUE(TEXT("State3Task1 should enter"), Exec.Expect("State3Task1", TEXT("EnterState5")));
+				AITEST_TRUE(TEXT("Should transition to State3"), Exec.ExpectInActiveStates("Root", "State3"));
+				Exec.LogClear();
+
+				Exec.Stop();
+			}
+		}
+
+		// NextParent fail to compile when there is no next parent
+		{
+			FTree Tree = MakeBasicTree();
+
+			Tree.State2->AddTransition(EStateTreeTransitionTrigger::OnStateCompleted, EStateTreeTransitionType::NextParent);
+
+			if (!Compile(Tree.StateTree, /*bExpectedToCompile*/false))
+			{
+				return false;
+			}
+		}
+
+		return true;
+	}
+};
+IMPLEMENT_STATE_TREE_INSTANT_TEST(FStateTreeTest_TransitionTypes, "System.StateTree.Transition.Types");
+
+
 } // namespace UE::StateTree::Tests
 
 UE_ENABLE_OPTIMIZATION_SHIP
