[StateTree] Fix bad implicit cast with TObjectPtr and ? operator.

diff --git a/Engine/Plugins/Runtime/StateTree/Source/StateTreeEditorModule/Private/StateTreeCompiler.cpp b/Engine/Plugins/Runtime/StateTree/Source/StateTreeEditorModule/Private/StateTreeCompiler.cpp
--- a/Engine/Plugins/Runtime/StateTree/Source/StateTreeEditorModule/Private/StateTreeCompiler.cpp
+++ b/Engine/Plugins/Runtime/StateTree/Source/StateTreeEditorModule/Private/StateTreeCompiler.cpp
@@ -2151,7 +2151,7 @@ bool FStateTreeCompiler::ResolveTransitionStateAndFallback(const UStateTreeState
 	else if (Link.LinkType == EStateTreeTransitionType::Parent)
 	{
 		// Find next parent.
-		const UStateTreeState* Parent = SourceState && SourceState->Parent ? SourceState->Parent : nullptr;
+		const UStateTreeState* Parent = SourceState && SourceState->Parent ? SourceState->Parent.Get() : nullptr;
 		if (Parent == nullptr)
 		{
 			Log.Reportf(EMessageSeverity::Error,
