[StateTree] Fix bad copy in unittest.

#rnx

diff --git a/Engine/Plugins/Runtime/StateTree/Source/StateTreeEditorModule/Public/StateTreeEditorNode.h b/Engine/Plugins/Runtime/StateTree/Source/StateTreeEditorModule/Public/StateTreeEditorNode.h
--- a/Engine/Plugins/Runtime/StateTree/Source/StateTreeEditorModule/Public/StateTreeEditorNode.h
+++ b/Engine/Plugins/Runtime/StateTree/Source/StateTreeEditorModule/Public/StateTreeEditorNode.h
@@ -61,6 +61,15 @@ struct FStateTreeEditorNode
 		return FGuid::Combine(ID, FGuid::NewDeterministicGuid(TEXT("Node Struct")));
 	}
 
+	/**
+	 * Get ID for the node instance data, which is different from the ID for the template.
+	 * @return ID for the instance data
+	 */
+	FGuid GetInstanceDataID() const
+	{
+		return ID;
+	}
+
 	FStateTreeDataView GetInstance() const
 	{
 		return InstanceObject ? FStateTreeDataView(InstanceObject) : FStateTreeDataView(const_cast<FInstancedStruct&>(Instance));
diff --git a/Engine/Plugins/Runtime/StateTree/Source/StateTreeTestSuite/Private/StateTreeDelegateTest.cpp b/Engine/Plugins/Runtime/StateTree/Source/StateTreeTestSuite/Private/StateTreeDelegateTest.cpp
--- a/Engine/Plugins/Runtime/StateTree/Source/StateTreeTestSuite/Private/StateTreeDelegateTest.cpp
+++ b/Engine/Plugins/Runtime/StateTree/Source/StateTreeTestSuite/Private/StateTreeDelegateTest.cpp
@@ -580,7 +580,7 @@ struct FStateTreeTest_Delegate_DispatcherOnNodeListenerOnTransition : FStateTree
 		UStateTreeEditorData& EditorData = *Cast<UStateTreeEditorData>(StateTree.EditorData);
 
 		UStateTreeState& Root = EditorData.AddSubTree(FName(TEXT("Root")));
-		TStateTreeEditorNode<FTestTask_DispatcherOnNodeAndInstance> DispatcherTaskNode = Root.AddTask<FTestTask_DispatcherOnNodeAndInstance>(TEXT("Dispatcher"));
+		TStateTreeEditorNode<FTestTask_DispatcherOnNodeAndInstance>& DispatcherTaskNode = Root.AddTask<FTestTask_DispatcherOnNodeAndInstance>(TEXT("Dispatcher"));
 		UStateTreeState& State1 = Root.AddChildState(FName(TEXT("State1")));
 		FStateTreeTransition& DelegateOnNodeTransition = State1.AddTransition(EStateTreeTransitionTrigger::OnDelegate, EStateTreeTransitionType::NextState);
 
@@ -592,7 +592,7 @@ struct FStateTreeTest_Delegate_DispatcherOnNodeListenerOnTransition : FStateTree
 			FPropertyBindingPath(DelegateOnNodeTransition.ID, GET_MEMBER_NAME_CHECKED(FStateTreeTransition, DelegateListener)));
 
 		EditorData.AddPropertyBinding(
-			FPropertyBindingPath(DispatcherTaskNode.ID, GET_MEMBER_NAME_CHECKED(FTestTask_DispatcherOnNodeAndInstance_InstanceData, DispatcherOnInstance)),
+			FPropertyBindingPath(DispatcherTaskNode.GetInstanceDataID(), GET_MEMBER_NAME_CHECKED(FTestTask_DispatcherOnNodeAndInstance_InstanceData, DispatcherOnInstance)),
 			FPropertyBindingPath(DelegateOnInstanceTransition.ID, GET_MEMBER_NAME_CHECKED(FStateTreeTransition, DelegateListener)));
 
 		{
