[StateTree] Refactor ModifyStateSelectionRulesPerTransition for GetStateSelectionRulesForTransition. Remove the OUT macro from all previous ModifyStateSelectionRulesPerTransition call.

#rnx

diff --git a/Engine/Plugins/Runtime/StateTree/Source/StateTreeModule/Private/StateTreeExecutionContext.cpp b/Engine/Plugins/Runtime/StateTree/Source/StateTreeModule/Private/StateTreeExecutionContext.cpp
--- a/Engine/Plugins/Runtime/StateTree/Source/StateTreeModule/Private/StateTreeExecutionContext.cpp
+++ b/Engine/Plugins/Runtime/StateTree/Source/StateTreeModule/Private/StateTreeExecutionContext.cpp
@@ -5838,7 +5838,7 @@ bool FStateTreeExecutionContext::TriggerTransitions()
 							TransitionArgs.Priority = Transition.Priority;
 							TransitionArgs.TransitionEvent = DelayedTransition.CapturedEvent;
 							TransitionArgs.Fallback = Transition.Fallback;
-							Transition.ModifyStateSelectionRulesPerTransition(OUT TransitionArgs.SelectionRules);
+							TransitionArgs.SelectionRules = Transition.GetStateSelectionRulesForTransition(TransitionArgs.SelectionRules);
 							if (RequestTransitionInternal(CurrentFrame, Handler.StateID, TargetState, TransitionArgs))
 							{
 								// If the transition was successfully requested with a specific event, consume and remove the event, it's been used.
@@ -5932,7 +5932,7 @@ PRAGMA_ENABLE_DEPRECATION_WARNINGS
 						TransitionArgs.Priority = Transition.Priority;
 						TransitionArgs.TransitionEvent = TransitionEvent;
 						TransitionArgs.Fallback = Transition.Fallback;
-						Transition.ModifyStateSelectionRulesPerTransition(OUT TransitionArgs.SelectionRules);
+						TransitionArgs.SelectionRules = Transition.GetStateSelectionRulesForTransition(TransitionArgs.SelectionRules);
 						if (RequestTransitionInternal(CurrentFrame, Handler.StateID, TargetState, TransitionArgs))
 						{
 							// If the transition was successfully requested with a specific event, consume and remove the event, it's been used.
@@ -6097,7 +6097,7 @@ PRAGMA_ENABLE_DEPRECATION_WARNINGS
 									FTransitionArguments TransitionArgs;
 									TransitionArgs.Priority = EStateTreeTransitionPriority::Normal;
 									TransitionArgs.Fallback = Transition.Fallback;
-									Transition.ModifyStateSelectionRulesPerTransition(OUT TransitionArgs.SelectionRules);
+									TransitionArgs.SelectionRules = Transition.GetStateSelectionRulesForTransition(TransitionArgs.SelectionRules);
 									if (RequestTransitionInternal(CurrentFrame, CurrentStateID, TargetState, TransitionArgs))
 									{
 										RequestedTransition->Source = FStateTreeTransitionSource(CurrentFrame.StateTree, FStateTreeIndex16(TransitionIndex), Transition.State, Transition.Priority);
@@ -8051,7 +8051,7 @@ bool FStateTreeExecutionContext::SelectStateInternal_TryFollowTransitions(
 				SelectStateArgs.TransitionEvent = SelectedStateTransitionEvent;
 				SelectStateArgs.Fallback = Transition.Fallback;
 				SelectStateArgs.SelectionRules = NextStateTree->StateSelectionRules;
-				Transition.ModifyStateSelectionRulesPerTransition(OUT SelectStateArgs.SelectionRules);
+				SelectStateArgs.SelectionRules = Transition.GetStateSelectionRulesForTransition(SelectStateArgs.SelectionRules);
 
 				OutSelectionResult->SelectedStates.Reset();
 				OutSelectionResult->SelectedFrames.Reset();
diff --git a/Engine/Plugins/Runtime/StateTree/Source/StateTreeModule/Public/StateTreeTypes.h b/Engine/Plugins/Runtime/StateTree/Source/StateTreeModule/Public/StateTreeTypes.h
--- a/Engine/Plugins/Runtime/StateTree/Source/StateTreeModule/Public/StateTreeTypes.h
+++ b/Engine/Plugins/Runtime/StateTree/Source/StateTreeModule/Public/StateTreeTypes.h
@@ -650,22 +650,18 @@ struct FCompactStateTransition
 		return !Delay.IsEmpty();
 	}
 
-	/** Toggles flags for state selection based on this transition's configuration */
-	void ModifyStateSelectionRulesPerTransition(EStateTreeStateSelectionRules& OutModifiedStateSelectionRules) const
+	/** @return flags for state selection based on this transition's configuration */
+	EStateTreeStateSelectionRules GetStateSelectionRulesForTransition(EStateTreeStateSelectionRules ModifiedStateSelectionRules) const
 	{
 		switch (ChangeTypeTargetStateRule)
 		{
 		case EStateTreeTransitionChangeTypeRules::ForceChanged:
-			OutModifiedStateSelectionRules |= EStateTreeStateSelectionRules::ReselectedStateCreatesNewStates;
-			return;
+			return ModifiedStateSelectionRules | EStateTreeStateSelectionRules::ReselectedStateCreatesNewStates;
 		case EStateTreeTransitionChangeTypeRules::ForceSustained:
-			OutModifiedStateSelectionRules &= ~EStateTreeStateSelectionRules::ReselectedStateCreatesNewStates;
-			return;
+			return ModifiedStateSelectionRules & ~EStateTreeStateSelectionRules::ReselectedStateCreatesNewStates;
 		case EStateTreeTransitionChangeTypeRules::Default:
-			return;
 		default:
-			checkf(false, TEXT("Unsupported EStateTreeTransitionChangeTypeRules: %u"), static_cast<uint8>(ChangeTypeTargetStateRule));
-			return;
+			return ModifiedStateSelectionRules;
 		}
 	}
 
