[State Tree] Fixed states order is reversed when inserting before another state in the editor

diff --git a/Engine/Plugins/Runtime/StateTree/Source/StateTreeEditorModule/Private/StateTreeViewModel.cpp b/Engine/Plugins/Runtime/StateTree/Source/StateTreeEditorModule/Private/StateTreeViewModel.cpp
--- a/Engine/Plugins/Runtime/StateTree/Source/StateTreeEditorModule/Private/StateTreeViewModel.cpp
+++ b/Engine/Plugins/Runtime/StateTree/Source/StateTreeEditorModule/Private/StateTreeViewModel.cpp
@@ -1847,7 +1847,6 @@ void FStateTreeViewModel::MoveSelectedStates(UStateTreeState* TargetState, const
 		const FScopedTransaction Transaction(LOCTEXT("MoveTransaction", "Move"));
 
 		TSet<UStateTreeState*> AffectedParents;
-		TSet<UStateTreeState*> AffectedStates;
 
 		UStateTreeState* TargetParent = TargetState->Parent;
 		if (RelativeLocation == EStateTreeViewModelInsert::Into)
@@ -1886,10 +1885,11 @@ void FStateTreeViewModel::MoveSelectedStates(UStateTreeState* TargetState, const
 			}
 		}
 
+		TSet<UStateTreeState*> AffectedStates;
 		// Add in reverse order to keep the original order.
-		for (int32 i = States.Num() - 1; i >= 0; i--)
+		for (int32 Idx = States.Num() - 1; Idx >= 0; Idx--)
 		{
-			if (UStateTreeState* SelectedState = States[i])
+			if (UStateTreeState* SelectedState = States[Idx])
 			{
 				AffectedStates.Add(SelectedState);
 
@@ -1920,7 +1920,10 @@ void FStateTreeViewModel::MoveSelectedStates(UStateTreeState* TargetState, const
 						if (RelativeLocation == EStateTreeViewModelInsert::Before)
 						{
 							// Before
-							ArrayToMoveTo.Insert(SelectedState, TargetIndex);
+							const int32 IndexToInsertBefore = TargetIndex - (States.Num() - 1 - Idx);
+							check(ArrayToMoveTo.IsValidIndex(IndexToInsertBefore));
+
+							ArrayToMoveTo.Insert(SelectedState, IndexToInsertBefore);
 							SelectedState->Parent = TargetParent;
 						}
 						else if (RelativeLocation == EStateTreeViewModelInsert::After)
