[StateTree] Make status not trivially copiable. The inner pointer might point to itself. Unreal container do use trivailly allocated stucture. They do not call the move or copy constructor.

diff --git a/Engine/Plugins/Runtime/StateTree/Source/StateTreeModule/Public/StateTreeTasksStatus.h b/Engine/Plugins/Runtime/StateTree/Source/StateTreeModule/Public/StateTreeTasksStatus.h
--- a/Engine/Plugins/Runtime/StateTree/Source/StateTreeModule/Public/StateTreeTasksStatus.h
+++ b/Engine/Plugins/Runtime/StateTree/Source/StateTreeModule/Public/StateTreeTasksStatus.h
@@ -3,6 +3,7 @@
 #pragma once
 
 #include "Misc/NotNull.h"
+#include "Traits/IsTriviallyRelocatable.h"
 #include "StateTreeTasksStatus.generated.h"
 
 class FArchive;
@@ -370,3 +371,10 @@ struct TStructOpsTypeTraits<FStateTreeTasksCompletionStatus> : public TStructOps
 		WithNetSerializer = true,
 	};
 };
+
+// FStateTreeTasksCompletionStatus is not trivially relocatable. Flag their use in containers and other contexts.
+template <UE::CDerivedFrom<FStateTreeTasksCompletionStatus> T>
+struct TIsTriviallyRelocatable<T>
+{
+	static inline constexpr bool Value = false;
+};
