[StateTree] Add menu action to select the transition's target state.

diff --git a/Engine/Plugins/Runtime/StateTree/Source/StateTreeEditorModule/Private/Customizations/Widgets/SStateTreeContextMenuButton.cpp b/Engine/Plugins/Runtime/StateTree/Source/StateTreeEditorModule/Private/Customizations/Widgets/SStateTreeContextMenuButton.cpp
--- a/Engine/Plugins/Runtime/StateTree/Source/StateTreeEditorModule/Private/Customizations/Widgets/SStateTreeContextMenuButton.cpp
+++ b/Engine/Plugins/Runtime/StateTree/Source/StateTreeEditorModule/Private/Customizations/Widgets/SStateTreeContextMenuButton.cpp
@@ -22,14 +22,7 @@ void SStateTreeContextMenuButton::Construct(const FArguments& InArgs, const TSha
 	{
 		if (UStateTreeState* OwnerState = OwnerStateWeak.Get())
 		{
-			for (const FStateTreeTransition& Transition : OwnerState->Transitions)
-			{
-				if (NodeID == Transition.ID)
-				{
-					bIsStateTransition = true;
-					break;
-				}
-			}
+			bIsStateTransition = OwnerState->GetTransitionByID(NodeID) != nullptr;
 		}
 	}
 
@@ -57,7 +50,8 @@ void SStateTreeContextMenuButton::Construct(const FArguments& InArgs, const TSha
 
 FReply SStateTreeContextMenuButton::OnMouseButtonUp(const FGeometry& MyGeometry, const FPointerEvent& MouseEvent)
 {
-	StateTreeViewModel->BringNodeToFocus(OwnerStateWeak.Get(), NodeID);
+	UStateTreeState* OwnerState = OwnerStateWeak.Get();
+	StateTreeViewModel->BringNodeToFocus(OwnerState, NodeID);
 
 	if (MouseEvent.GetEffectingButton() == EKeys::RightMouseButton)
 	{
@@ -75,6 +69,13 @@ FReply SStateTreeContextMenuButton::OnMouseButtonUp(const FGeometry& MyGeometry,
 			return FReply::Handled();
 		}
 	}
+	// Go to transition target with Alt-Click
+	else if (bIsStateTransition && MouseEvent.GetModifierKeys().IsAltDown())
+	{
+		UStateTreeState* GotoState = OwnerState ? StateTreeViewModel->GetTransitionToState(OwnerState, NodeID) : nullptr;
+		StateTreeViewModel->SetSelection(GotoState);
+		return FReply::Handled();
+	}
 
 	return SButton::OnMouseButtonUp(MyGeometry, MouseEvent);
 }
@@ -85,6 +86,33 @@ TSharedRef<SWidget> SStateTreeContextMenuButton::MakeContextMenu() const
 
 	if (StateTreeViewModel.IsValid() && OwnerStateWeak.IsValid())
 	{
+		MenuBuilder.BeginSection(FName("View"), LOCTEXT("View", "View"));
+
+		// Navigate
+		if (bIsStateTransition)
+		{
+			MenuBuilder.AddMenuEntry(
+				LOCTEXT("Navigate", "Navigate"),
+				LOCTEXT("NavigateTooltip", "Navigate to transition target. (Alt-Click)"),
+				FSlateIcon(FAppStyle::GetAppStyleSetName(), "Icons.BrowseContent"),
+				FUIAction(
+					FExecuteAction::CreateSPLambda(this, [this]()
+						{
+							UStateTreeState* OwnerState = OwnerStateWeak.Get();
+							UStateTreeState* GotoState = OwnerState ? StateTreeViewModel->GetTransitionToState(OwnerState, NodeID) : nullptr;
+							StateTreeViewModel->SetSelection(GotoState);
+						}),
+					FCanExecuteAction::CreateSPLambda(this, [this]()
+						{
+							UStateTreeState* OwnerState = OwnerStateWeak.Get();
+							return OwnerState != nullptr && StateTreeViewModel->GetTransitionToState(OwnerState, NodeID) != nullptr;
+						})
+				));
+		}
+
+
+		MenuBuilder.EndSection();
+
 		MenuBuilder.BeginSection(FName("Edit"), LOCTEXT("Edit", "Edit"));
 
 		// Copy
diff --git a/Engine/Plugins/Runtime/StateTree/Source/StateTreeEditorModule/Private/StateTreeState.cpp b/Engine/Plugins/Runtime/StateTree/Source/StateTreeEditorModule/Private/StateTreeState.cpp
--- a/Engine/Plugins/Runtime/StateTree/Source/StateTreeEditorModule/Private/StateTreeState.cpp
+++ b/Engine/Plugins/Runtime/StateTree/Source/StateTreeEditorModule/Private/StateTreeState.cpp
@@ -441,7 +441,7 @@ const UStateTreeState* UStateTreeState::GetRootState() const
 	return RootState;
 }
 
-const UStateTreeState* UStateTreeState::GetNextSiblingState() const
+UStateTreeState* UStateTreeState::GetNextSiblingState() const
 {
 	if (!Parent)
 	{
@@ -464,7 +464,7 @@ const UStateTreeState* UStateTreeState::GetNextSiblingState() const
 	return nullptr;
 }
 
-const UStateTreeState* UStateTreeState::GetNextSelectableSiblingState() const
+UStateTreeState* UStateTreeState::GetNextSelectableSiblingState() const
 {
 	if (!Parent)
 	{
@@ -480,7 +480,7 @@ const UStateTreeState* UStateTreeState::GetNextSelectableSiblingState() const
 	for (int32 ChildIdx = StartChildIndex + 1; ChildIdx < Parent->Children.Num(); ChildIdx++)
 	{
 		// Select the next enabled and selectable sibling
-		const UStateTreeState* State =Parent->Children[ChildIdx];
+		UStateTreeState* State = Parent->Children[ChildIdx];
 		if (State->SelectionBehavior != EStateTreeStateSelectionBehavior::None
 			&& State->bEnabled)
 		{
@@ -493,7 +493,7 @@ const UStateTreeState* UStateTreeState::GetNextSelectableSiblingState() const
 
 FString UStateTreeState::GetPath() const
 {
-	TArray<const UStateTreeState*> States;
+	TArray<const UStateTreeState*, TInlineAllocator<8>> States;
 	for (const UStateTreeState* CurrState = this; CurrState; CurrState = CurrState->Parent)
 	{
 		States.Add(CurrState);
diff --git a/Engine/Plugins/Runtime/StateTree/Source/StateTreeEditorModule/Private/StateTreeViewModel.cpp b/Engine/Plugins/Runtime/StateTree/Source/StateTreeEditorModule/Private/StateTreeViewModel.cpp
--- a/Engine/Plugins/Runtime/StateTree/Source/StateTreeEditorModule/Private/StateTreeViewModel.cpp
+++ b/Engine/Plugins/Runtime/StateTree/Source/StateTreeEditorModule/Private/StateTreeViewModel.cpp
@@ -858,6 +858,35 @@ void FStateTreeViewModel::GetLinkStates(FGuid StateID, TArray<FGuid>& LinkingIn,
 	}
 }
 
+UStateTreeState* FStateTreeViewModel::GetTransitionToState(TNotNull<const UStateTreeState*> State, const FGuid TransitionID) const
+{
+	const FStateTreeTransition* Transition = State->GetTransitionByID(TransitionID);
+	ensure(Transition);
+	if (Transition == nullptr)
+	{
+		return nullptr;
+	}
+
+	switch (Transition->State.LinkType)
+	{
+	case EStateTreeTransitionType::GotoState:
+		return GetMutableStateByID(Transition->State.ID);
+	case EStateTreeTransitionType::Parent:
+		return State->Parent ? State->Parent.Get() : nullptr;
+	case EStateTreeTransitionType::NextState:
+	case EStateTreeTransitionType::NextSelectableState:
+		return State->GetNextSelectableSiblingState();
+	case EStateTreeTransitionType::NextParent:
+	case EStateTreeTransitionType::NextSelectableParent:
+		return State->Parent ? State->Parent->GetNextSelectableSiblingState() : nullptr;
+	case EStateTreeTransitionType::None:
+	case EStateTreeTransitionType::Succeeded:
+	case EStateTreeTransitionType::Failed:
+	default:
+		return nullptr;
+	}
+}
+
 void FStateTreeViewModel::PostUndo(bool bSuccess)
 {
 	// TODO: see if we can narrow this down.
diff --git a/Engine/Plugins/Runtime/StateTree/Source/StateTreeEditorModule/Public/StateTreeState.h b/Engine/Plugins/Runtime/StateTree/Source/StateTreeEditorModule/Public/StateTreeState.h
--- a/Engine/Plugins/Runtime/StateTree/Source/StateTreeEditorModule/Public/StateTreeState.h
+++ b/Engine/Plugins/Runtime/StateTree/Source/StateTreeEditorModule/Public/StateTreeState.h
@@ -237,8 +237,8 @@ public:
 	UE_API void OnTreeCompiled(const UStateTree& StateTree);
 
 	UE_API const UStateTreeState* GetRootState() const;
-	UE_API const UStateTreeState* GetNextSiblingState() const;
-	UE_API const UStateTreeState* GetNextSelectableSiblingState() const;
+	UE_API UStateTreeState* GetNextSiblingState() const;
+	UE_API UStateTreeState* GetNextSelectableSiblingState() const;
 
 	/** @return the path of the state as string. */
 	UE_API FString GetPath() const;
@@ -339,6 +339,14 @@ public:
 		return Transition;
 	}
 
+	const FStateTreeTransition* GetTransitionByID(const FGuid TransitionID) const
+	{
+		return Transitions.FindByPredicate([TransitionID](const FStateTreeTransition& Other)
+			{
+				return Other.ID == TransitionID;
+			});
+	}
+
 	FGuid GetEventID() const
 	{
 		return FGuid::Combine(ID, FGuid::NewDeterministicGuid(TEXT("Event")));
diff --git a/Engine/Plugins/Runtime/StateTree/Source/StateTreeEditorModule/Public/StateTreeViewModel.h b/Engine/Plugins/Runtime/StateTree/Source/StateTreeEditorModule/Public/StateTreeViewModel.h
--- a/Engine/Plugins/Runtime/StateTree/Source/StateTreeEditorModule/Public/StateTreeViewModel.h
+++ b/Engine/Plugins/Runtime/StateTree/Source/StateTreeEditorModule/Public/StateTreeViewModel.h
@@ -96,6 +96,9 @@ public:
 	/** Find the states that are linked to the provided StateID. */
 	UE_API void GetLinkStates(FGuid StateID, TArray<FGuid>& LinkingIn, TArray<FGuid>& LinkedOut) const;
 
+	/** Find the states that the transition go to. */
+	UE_API UStateTreeState* GetTransitionToState(TNotNull<const UStateTreeState*> State, FGuid TransitionID) const;
+
 	// Gets and sets StateTree view expansion state store in the asset.
 	UE_API void SetPersistentExpandedStates(TSet<TWeakObjectPtr<UStateTreeState>>& InExpandedStates);
 	UE_API void GetPersistentExpandedStates(TSet<TWeakObjectPtr<UStateTreeState>>& OutExpandedStates);
