[State Tree]
replaced UE_REQUIRES with c++ requires in AsyncExecutionContext

#Synced-CL 47889454

diff --git a/Engine/Plugins/Runtime/StateTree/Source/StateTreeModule/Private/StateTreeAsyncExecutionContext.cpp b/Engine/Plugins/Runtime/StateTree/Source/StateTreeModule/Private/StateTreeAsyncExecutionContext.cpp
--- a/Engine/Plugins/Runtime/StateTree/Source/StateTreeModule/Private/StateTreeAsyncExecutionContext.cpp
+++ b/Engine/Plugins/Runtime/StateTree/Source/StateTreeModule/Private/StateTreeAsyncExecutionContext.cpp
@@ -56,8 +56,7 @@ TStateTreeStrongExecutionContext<bWithWriteAccess>::~TStateTreeStrongExecutionCo
 }
 
 template <bool bWithWriteAccess>
-template<bool bWriteAccess UE_REQUIRES_DEFINITION(bWriteAccess)>
-bool TStateTreeStrongExecutionContext<bWithWriteAccess>::SendEvent(const FGameplayTag Tag, const FConstStructView Payload, const FName Origin) const
+bool TStateTreeStrongExecutionContext<bWithWriteAccess>::SendEvent(const FGameplayTag Tag, const FConstStructView Payload, const FName Origin) const requires bWithWriteAccess
 {
 	static_assert(bWithWriteAccess);
 
@@ -72,9 +71,8 @@ bool TStateTreeStrongExecutionContext<bWithWriteAccess>::SendEvent(const FGamepl
 }
 
 template <bool bWithWriteAccess>
-template<bool bWriteAccess UE_REQUIRES_DEFINITION(bWriteAccess)>
 bool TStateTreeStrongExecutionContext<bWithWriteAccess>::RequestTransition(FStateTreeStateHandle TargetState,
-	EStateTreeTransitionPriority Priority, EStateTreeSelectionFallback Fallback) const
+	EStateTreeTransitionPriority Priority, EStateTreeSelectionFallback Fallback) const requires bWithWriteAccess
 {
 	static_assert(bWithWriteAccess);
 
@@ -98,9 +96,7 @@ bool TStateTreeStrongExecutionContext<bWithWriteAccess>::RequestTransition(FStat
 }
 
 template <bool bWithWriteAccess>
-template<bool bWriteAccess UE_REQUIRES_DEFINITION(bWriteAccess)>
-bool TStateTreeStrongExecutionContext<bWithWriteAccess>::BroadcastDelegate(const FStateTreeDelegateDispatcher& Dispatcher)
-const
+bool TStateTreeStrongExecutionContext<bWithWriteAccess>::BroadcastDelegate(const FStateTreeDelegateDispatcher& Dispatcher) const requires bWithWriteAccess
 {
 	static_assert(bWithWriteAccess);
 
@@ -127,8 +123,7 @@ const
 }
 
 template <bool bWithWriteAccess>
-template<bool bWriteAccess UE_REQUIRES_DEFINITION(bWriteAccess)>
-bool TStateTreeStrongExecutionContext<bWithWriteAccess>::BindDelegate(const FStateTreeDelegateListener& Listener, FSimpleDelegate Delegate) const
+bool TStateTreeStrongExecutionContext<bWithWriteAccess>::BindDelegate(const FStateTreeDelegateListener& Listener, FSimpleDelegate Delegate) const requires bWithWriteAccess
 {
 	static_assert(bWithWriteAccess);
 
@@ -153,8 +148,7 @@ bool TStateTreeStrongExecutionContext<bWithWriteAccess>::BindDelegate(const FSta
 }
 
 template <bool bWithWriteAccess>
-template<bool bWriteAccess UE_REQUIRES_DEFINITION(bWriteAccess)>
-bool TStateTreeStrongExecutionContext<bWithWriteAccess>::UnbindDelegate(const FStateTreeDelegateListener& Listener) const
+bool TStateTreeStrongExecutionContext<bWithWriteAccess>::UnbindDelegate(const FStateTreeDelegateListener& Listener) const requires bWithWriteAccess
 {
 	static_assert(bWithWriteAccess);
 
@@ -177,8 +171,7 @@ bool TStateTreeStrongExecutionContext<bWithWriteAccess>::UnbindDelegate(const FS
 }
 
 template <bool bWithWriteAccess>
-template<bool bWriteAccess UE_REQUIRES_DEFINITION(bWriteAccess)>
-bool TStateTreeStrongExecutionContext<bWithWriteAccess>::FinishTask(EStateTreeFinishTaskType FinishType) const
+bool TStateTreeStrongExecutionContext<bWithWriteAccess>::FinishTask(EStateTreeFinishTaskType FinishType) const requires bWithWriteAccess
 {
 	static_assert(bWithWriteAccess);
 
@@ -237,8 +230,7 @@ bool TStateTreeStrongExecutionContext<bWithWriteAccess>::FinishTask(EStateTreeFi
 }
 
 template <bool bWithWriteAccess>
-template<bool bWriteAccess UE_REQUIRES_DEFINITION(bWriteAccess)>
-bool TStateTreeStrongExecutionContext<bWithWriteAccess>::UpdateScheduledTickRequest(UE::StateTree::FScheduledTickHandle Handle, FStateTreeScheduledTick ScheduledTick) const
+bool TStateTreeStrongExecutionContext<bWithWriteAccess>::UpdateScheduledTickRequest(UE::StateTree::FScheduledTickHandle Handle, FStateTreeScheduledTick ScheduledTick) const requires bWithWriteAccess
 {
 	static_assert(bWithWriteAccess);
 
@@ -342,8 +334,7 @@ FStateTreeDataView TStateTreeStrongExecutionContext<bWithWriteAccess>::GetInstan
 }
 
 template <bool bWithWriteAccess>
-template<bool bWriteAccess UE_REQUIRES_DEFINITION(bWriteAccess)>
-void TStateTreeStrongExecutionContext<bWithWriteAccess>::ScheduleNextTick(TNotNull<UObject*> Owner, TNotNull<const UStateTree*> RootStateTree, FStateTreeInstanceStorage& Storage, UE::StateTree::ETickReason Reason)
+void TStateTreeStrongExecutionContext<bWithWriteAccess>::ScheduleNextTick(TNotNull<UObject*> Owner, TNotNull<const UStateTree*> RootStateTree, FStateTreeInstanceStorage& Storage, UE::StateTree::ETickReason Reason) requires bWithWriteAccess
 {
 	static_assert(bWithWriteAccess);
 
@@ -357,16 +348,6 @@ void TStateTreeStrongExecutionContext<bWithWriteAccess>::ScheduleNextTick(TNotNu
 template struct TStateTreeStrongExecutionContext<false>;
 template struct TStateTreeStrongExecutionContext<true>;
 
-//@todo: remove all these instantiations once past 5.6 because of c++20 support
-template bool STATETREEMODULE_API TStateTreeStrongExecutionContext<true>::SendEvent(const FGameplayTag Tag, const FConstStructView Payload, const FName Origin) const;
-template bool STATETREEMODULE_API TStateTreeStrongExecutionContext<true>::RequestTransition(FStateTreeStateHandle TargetState, EStateTreeTransitionPriority Priority, EStateTreeSelectionFallback Fallback) const;
-template bool STATETREEMODULE_API TStateTreeStrongExecutionContext<true>::BroadcastDelegate(const FStateTreeDelegateDispatcher& Dispatcher) const;
-template bool STATETREEMODULE_API TStateTreeStrongExecutionContext<true>::BindDelegate(const FStateTreeDelegateListener& Listener, FSimpleDelegate Delegate) const;
-template bool STATETREEMODULE_API TStateTreeStrongExecutionContext<true>::UnbindDelegate(const FStateTreeDelegateListener& Listener) const;
-template bool STATETREEMODULE_API TStateTreeStrongExecutionContext<true>::FinishTask(EStateTreeFinishTaskType FinishType) const;
-template bool STATETREEMODULE_API TStateTreeStrongExecutionContext<true>::UpdateScheduledTickRequest(UE::StateTree::FScheduledTickHandle Handle, FStateTreeScheduledTick ScheduledTick) const;
-template void STATETREEMODULE_API TStateTreeStrongExecutionContext<true>::ScheduleNextTick(TNotNull<UObject*> Owner, TNotNull<const UStateTree*> RootStateTree, FStateTreeInstanceStorage& Storage, UE::StateTree::ETickReason Reason);
-
 FStateTreeWeakExecutionContext::FStateTreeWeakExecutionContext(const FStateTreeExecutionContext& Context)
 	: Owner(Context.GetOwner())
 	, StateTree(Context.GetStateTree())
diff --git a/Engine/Plugins/Runtime/StateTree/Source/StateTreeModule/Public/StateTreeAsyncExecutionContext.h b/Engine/Plugins/Runtime/StateTree/Source/StateTreeModule/Public/StateTreeAsyncExecutionContext.h
--- a/Engine/Plugins/Runtime/StateTree/Source/StateTreeModule/Public/StateTreeAsyncExecutionContext.h
+++ b/Engine/Plugins/Runtime/StateTree/Source/StateTreeModule/Public/StateTreeAsyncExecutionContext.h
@@ -84,14 +84,11 @@ struct TStateTreeStrongExecutionContext
 		return Storage;
 	}
 
-	//@todo: Replace UE_REQUIRES with requires and UE_API once past 5.6 for c++20 support
-
 	/**
 	 * Sends event for the StateTree.
 	 * @return false if the context is not valid or the event could not be sent.
 	 */
-	template<bool bWriteAccess = bWithWriteAccess UE_REQUIRES(bWriteAccess)>
-	bool SendEvent(const FGameplayTag Tag, const FConstStructView Payload = FConstStructView(), const FName Origin = FName()) const;
+	UE_API bool SendEvent(const FGameplayTag Tag, const FConstStructView Payload = FConstStructView(), const FName Origin = FName()) const requires bWithWriteAccess;
 
 	/**
 	 * Requests transition to a state.
@@ -103,16 +100,14 @@ struct TStateTreeStrongExecutionContext
 	 * @param Fallback of the transition if it fails to select the target state.
 	 * @return false if the context is not valid or doesn't have a valid frame anymore or the request failed.
 	 */
-	template<bool bWriteAccess = bWithWriteAccess UE_REQUIRES(bWriteAccess)>
-	bool RequestTransition(FStateTreeStateHandle TargetState, EStateTreeTransitionPriority Priority = EStateTreeTransitionPriority::Normal, EStateTreeSelectionFallback Fallback = EStateTreeSelectionFallback::None) const;
+	UE_API bool RequestTransition(FStateTreeStateHandle TargetState, EStateTreeTransitionPriority Priority = EStateTreeTransitionPriority::Normal, EStateTreeSelectionFallback Fallback = EStateTreeSelectionFallback::None) const requires bWithWriteAccess;
 
 	/**
 	 * Broadcasts the delegate.
 	 * It executes bound delegates immediately and triggers bound transitions (when transitions are evaluated).
 	 * @return false if the context is not valid or doesn't have a valid frame anymore or the broadcast failed.
 	 */
-	template<bool bWriteAccess = bWithWriteAccess UE_REQUIRES(bWriteAccess)>
-	bool BroadcastDelegate(const FStateTreeDelegateDispatcher& Dispatcher) const;
+	UE_API bool BroadcastDelegate(const FStateTreeDelegateDispatcher& Dispatcher) const requires bWithWriteAccess;
 
 	/**
 	 * Registers the delegate to the listener.
@@ -120,15 +115,13 @@ struct TStateTreeStrongExecutionContext
 	 * The listener is bound to a dispatcher in the editor.
 	 * @return false if the context is not valid or doesn't have a valid frame anymore or the bind failed.
 	 */
-	template<bool bWriteAccess = bWithWriteAccess UE_REQUIRES(bWriteAccess)>
-	bool BindDelegate(const FStateTreeDelegateListener& Listener, FSimpleDelegate Delegate) const;
+	UE_API bool BindDelegate(const FStateTreeDelegateListener& Listener, FSimpleDelegate Delegate) const requires bWithWriteAccess;
 
 	/**
 	 * Unregisters the callback bound to the listener.
 	 * @return false if the context is not valid or the unbind failed.
 	 */
-	template<bool bWriteAccess = bWithWriteAccess UE_REQUIRES(bWriteAccess)>
-	bool UnbindDelegate(const FStateTreeDelegateListener& Listener) const;
+	UE_API bool UnbindDelegate(const FStateTreeDelegateListener& Listener) const requires bWithWriteAccess;
 
 	/**
 	 * Finishes a task.
@@ -136,15 +129,13 @@ struct TStateTreeStrongExecutionContext
 	 * If called outside of the tick processing, then the request is buffered and handled on the next tick.
 	 * @return false if the context is not valid or doesn't have a valid frame anymore or finish task failed.
 	 */
-	template<bool bWriteAccess = bWithWriteAccess UE_REQUIRES(bWriteAccess)>
-	bool FinishTask(EStateTreeFinishTaskType FinishType) const;
+	UE_API bool FinishTask(EStateTreeFinishTaskType FinishType) const requires bWithWriteAccess;
 
 	/**
 	 * Updates the scheduled tick of a previous request.
 	 * @return false if the context is not valid.
 	 */
-	template<bool bWriteAccess = bWithWriteAccess UE_REQUIRES(bWriteAccess)>
-	bool UpdateScheduledTickRequest(UE::StateTree::FScheduledTickHandle Handle, FStateTreeScheduledTick ScheduledTick) const;
+	UE_API bool UpdateScheduledTickRequest(UE::StateTree::FScheduledTickHandle Handle, FStateTreeScheduledTick ScheduledTick) const requires bWithWriteAccess;
 
 	/**
 	 * Get the Instance Data of recorded node. Only callable on lvalue because it would have lost the access track.
@@ -190,8 +181,7 @@ private:
 
 	UE_API FStateTreeDataView GetInstanceDataPtrInternal() const;
 
-	template<bool bWriteAccess = bWithWriteAccess UE_REQUIRES(bWriteAccess)>
-	static void ScheduleNextTick(TNotNull<UObject*> Owner, TNotNull<const UStateTree*> RootStateTree, FStateTreeInstanceStorage& Storage, UE::StateTree::ETickReason Reason);
+	UE_API static void ScheduleNextTick(TNotNull<UObject*> Owner, TNotNull<const UStateTree*> RootStateTree, FStateTreeInstanceStorage& Storage, UE::StateTree::ETickReason Reason) requires bWithWriteAccess;
 
 	TStrongObjectPtr<UObject> Owner;
 	TStrongObjectPtr<const UStateTree> StateTree;
