[StateTree] Enable async compilation when the state tree asset loads. Doesn't impact PIE.

#Synced-CL 48251233

diff --git a/Engine/Plugins/Runtime/StateTree/Source/StateTreeEditorModule/Private/StateTreeCompilerManager.cpp b/Engine/Plugins/Runtime/StateTree/Source/StateTreeEditorModule/Private/StateTreeCompilerManager.cpp
--- a/Engine/Plugins/Runtime/StateTree/Source/StateTreeEditorModule/Private/StateTreeCompilerManager.cpp
+++ b/Engine/Plugins/Runtime/StateTree/Source/StateTreeEditorModule/Private/StateTreeCompilerManager.cpp
@@ -43,7 +43,7 @@ FAutoConsoleVariableRef CVarEnableCompileSynchronouslyInPIESession(
 	TEXT("Any new asset will be compile synchronously event is async was requested.")
 );
 
-int32 NumberOfQueuedCompilationPerBatch = 1;
+int32 NumberOfQueuedCompilationPerBatch = 2;
 FAutoConsoleVariableRef CVarNumberOfQueuedCompilationPerBatch(
 	TEXT("StateTree.Compiler.NumberOfQueuedCompilationPerBatch"),
 	NumberOfQueuedCompilationPerBatch,
diff --git a/Engine/Plugins/Runtime/StateTree/Source/StateTreeEditorModule/Private/StateTreeEditorModule.cpp b/Engine/Plugins/Runtime/StateTree/Source/StateTreeEditorModule/Private/StateTreeEditorModule.cpp
--- a/Engine/Plugins/Runtime/StateTree/Source/StateTreeEditorModule/Private/StateTreeEditorModule.cpp
+++ b/Engine/Plugins/Runtime/StateTree/Source/StateTreeEditorModule/Private/StateTreeEditorModule.cpp
@@ -57,7 +57,7 @@ namespace UE::StateTree::Editor
 	 * That ensures that dependent state tree assets are compiled before being used.
 	 * When compiling in the editor (once the asset are loaded), you might have a dependency bug.
 	 */
-	bool bEnableQueuedCompilationOnAssetLoad = false;
+	bool bEnableQueuedCompilationOnAssetLoad = true;
 	FAutoConsoleVariableRef CVarEnableQueuedCompilationOnAssetLoad(
 		TEXT("StateTree.Compiler.EnableQueuedCompilationOnAssetLoad"),
 		bEnableQueuedCompilationOnAssetLoad,
