StateTree - do not initialize node class cache in non-editor runs.

#Synced-CL 48301749

diff --git a/Engine/Plugins/Runtime/StateTree/Source/StateTreeEditorModule/Private/StateTreeEditorModule.cpp b/Engine/Plugins/Runtime/StateTree/Source/StateTreeEditorModule/Private/StateTreeEditorModule.cpp
--- a/Engine/Plugins/Runtime/StateTree/Source/StateTreeEditorModule/Private/StateTreeEditorModule.cpp
+++ b/Engine/Plugins/Runtime/StateTree/Source/StateTreeEditorModule/Private/StateTreeEditorModule.cpp
@@ -103,10 +103,14 @@ void FStateTreeEditorModule::StartupModule()
 	UE::StateTree::Delegates::Private::OnStateTreeAssetLoaded.BindRaw(this, &FStateTreeEditorModule::HandleStateTreeAssetLoaded);
 	UE::StateTree::Delegates::Private::OnAppendToClassSchema.BindRaw(this, &FStateTreeEditorModule::HandleAppendToStateTreeClassSchema);
 
-	OnPostEngineInitHandle = FCoreDelegates::OnPostEngineInit.AddLambda([this]()
-		{
-			NodeClassCache = UE::StateTree::Editor::InitNodeClassCache();
-		});
+	// class cache isn't needed in non-editor runs
+	if (GIsEditor)
+	{
+		OnPostEngineInitHandle = FCoreDelegates::OnPostEngineInit.AddLambda([this]()
+			{
+				NodeClassCache = UE::StateTree::Editor::InitNodeClassCache();
+			});
+	}
 
 #if WITH_STATETREE_TRACE_DEBUGGER
 	FStateTreeDebuggerCommands::Register();
@@ -178,7 +182,10 @@ void FStateTreeEditorModule::ShutdownModule()
 	FStateTreeDebuggerCommands::Unregister();
 #endif // WITH_STATETREE_TRACE_DEBUGGER
 
-	FCoreDelegates::OnPostEngineInit.Remove(OnPostEngineInitHandle);
+	if (OnPostEngineInitHandle.IsValid())
+	{
+		FCoreDelegates::OnPostEngineInit.Remove(OnPostEngineInitHandle);
+	}
 	NodeClassCache.Reset();
 
 	UE::StateTree::Compiler::FCompilerManager::Shutdown();
