[StateTree] Do not mark the asset as modified in MarkAsPubliclyModified. It doesn't work as intended right now. It supposed to mark the asset dirty for runtime compilation, see UStateTree::IsReadyToRun.

#rnx
#Synced-CL 48343893

diff --git a/Engine/Plugins/Runtime/StateTree/Source/StateTreeEditorModule/Private/StateTreeEditingSubsystem.cpp b/Engine/Plugins/Runtime/StateTree/Source/StateTreeEditorModule/Private/StateTreeEditingSubsystem.cpp
--- a/Engine/Plugins/Runtime/StateTree/Source/StateTreeEditorModule/Private/StateTreeEditingSubsystem.cpp
+++ b/Engine/Plugins/Runtime/StateTree/Source/StateTreeEditorModule/Private/StateTreeEditingSubsystem.cpp
@@ -54,20 +54,22 @@ bool UStateTreeEditingSubsystem::NeedsRecompile(TNotNull<const UStateTree*> InSt
 
 void UStateTreeEditingSubsystem::MarkAsPubliclyModified(TNotNull<UStateTree*> InStateTree)
 {
-	constexpr bool bPubliclyModified = true;
-	InStateTree->MarkAsModified(bPubliclyModified);
 	if (UE::StateTree::Editor::bEnableQueuedCompilationWhenAssetDirty)
 	{
+		//@todo fix mark as modified to only trigger a runtime compilation
+		constexpr bool bPubliclyModified = true;
+		InStateTree->MarkAsModified(bPubliclyModified);
 		UE::StateTree::Compiler::FCompilerManager::QueueForCompilation(InStateTree);
 	}
 }
 
 void UStateTreeEditingSubsystem::MarkAsModified(TNotNull<UStateTree*> InStateTree)
 {
-	constexpr bool bPubliclyModified = false;
-	InStateTree->MarkAsModified(bPubliclyModified);
 	if (UE::StateTree::Editor::bEnableQueuedCompilationWhenAssetDirty)
 	{
+		//@todo fix mark as modified to only trigger a runtime compilation
+		constexpr bool bPubliclyModified = false;
+		InStateTree->MarkAsModified(bPubliclyModified);
 		UE::StateTree::Compiler::FCompilerManager::QueueForCompilation(InStateTree);
 	}
 }
