Initial Commit State Centric View
- Disabled via `StateTree.Editor.Experimental.EnableStateCentricView`

#tests editor
#Synced-CL 48254188

diff --git a/Engine/Plugins/Experimental/UAF/UAFStateTree/Source/UAFStateTreeUncookedOnly/Private/AnimNextStateTreeFunctionLibraryHelper.cpp b/Engine/Plugins/Experimental/UAF/UAFStateTree/Source/UAFStateTreeUncookedOnly/Private/AnimNextStateTreeFunctionLibraryHelper.cpp
--- a/Engine/Plugins/Experimental/UAF/UAFStateTree/Source/UAFStateTreeUncookedOnly/Private/AnimNextStateTreeFunctionLibraryHelper.cpp
+++ b/Engine/Plugins/Experimental/UAF/UAFStateTree/Source/UAFStateTreeUncookedOnly/Private/AnimNextStateTreeFunctionLibraryHelper.cpp
@@ -4,14 +4,24 @@
 #include "RigVMBlueprintGeneratedClass.h"
 #include "UncookedOnlyUtils.h"
 
+static bool bStateTreeSearchForFunctions = true;
+static FAutoConsoleVariableRef CVarStateTreeSearchForFunctions(
+	TEXT("a.UAF.Experimental.StateTreeSearchForFunctions"),
+	bStateTreeSearchForFunctions,
+	TEXT("Temp CVar. This disables a slow global asset search that happens per details panel in UAF ST. We have sometimes 6x details panel in ST, so don't do that."));
+
 const TArray<FName> UAnimNextStateTreeFunctionLibraryHelper::GetExposedAnimNextFunctionNames()
 {
 	// Can't cache here since the list of functions is dynamic.
 	TArray<FName> Result;
 
 	TMap<FAssetData, FRigVMGraphFunctionHeaderArray> FunctionExports;
-	UE::UAF::UncookedOnly::FUtils::GetExportedFunctionsFromAssetRegistry(UE::UAF::AnimNextPublicGraphFunctionsExportsRegistryTag, FunctionExports);
-	UE::UAF::UncookedOnly::FUtils::GetExportedFunctionsFromAssetRegistry(UE::UAF::ControlRigAssetPublicGraphFunctionsExportsRegistryTag, FunctionExports);
+
+	if (bStateTreeSearchForFunctions)
+	{
+		UE::UAF::UncookedOnly::FUtils::GetExportedFunctionsFromAssetRegistry(UE::UAF::AnimNextPublicGraphFunctionsExportsRegistryTag, FunctionExports);
+		UE::UAF::UncookedOnly::FUtils::GetExportedFunctionsFromAssetRegistry(UE::UAF::ControlRigAssetPublicGraphFunctionsExportsRegistryTag, FunctionExports);
+	}
 	
 	for (const TPair<FAssetData, FRigVMGraphFunctionHeaderArray>& Export : FunctionExports.Array())
 	{
diff --git a/Engine/Plugins/Runtime/StateTree/Source/StateTreeEditorModule/Private/Customizations/StateTreeEditorNodeDetails.cpp b/Engine/Plugins/Runtime/StateTree/Source/StateTreeEditorModule/Private/Customizations/StateTreeEditorNodeDetails.cpp
--- a/Engine/Plugins/Runtime/StateTree/Source/StateTreeEditorModule/Private/Customizations/StateTreeEditorNodeDetails.cpp
+++ b/Engine/Plugins/Runtime/StateTree/Source/StateTreeEditorModule/Private/Customizations/StateTreeEditorNodeDetails.cpp
@@ -1804,6 +1804,11 @@ FText FStateTreeEditorNodeDetails::GetName() const
 {
 	check(StructProperty);
 
+	if (!EditorData.IsValid())
+	{
+		return FText::GetEmpty();
+	}
+
 	// Multiple names do not make sense, just if only one node is selected.
 	TArray<void*> RawNodeData;
 	StructProperty->AccessRawData(RawNodeData);
diff --git a/Engine/Plugins/Runtime/StateTree/Source/StateTreeEditorModule/Private/SStateTreeView.cpp b/Engine/Plugins/Runtime/StateTree/Source/StateTreeEditorModule/Private/SStateTreeView.cpp
--- a/Engine/Plugins/Runtime/StateTree/Source/StateTreeEditorModule/Private/SStateTreeView.cpp
+++ b/Engine/Plugins/Runtime/StateTree/Source/StateTreeEditorModule/Private/SStateTreeView.cpp
@@ -22,6 +22,11 @@
 #include "IDetailsView.h"
 #include "PropertyEditorModule.h"
 
+// @TODO: Dependency onto state centric to enable view toggling. 
+// Long term plan is to have this in a separate tab.
+#include "StateCentricView/StateCentricViewSettings.h"
+#include "StateCentricView/View/SStateCentricViewMainStateNode.h"
+
 #define LOCTEXT_NAMESPACE "StateTreeEditor"
 
 SStateTreeView::SStateTreeView()
@@ -86,10 +91,7 @@ void SStateTreeView::Construct(const FArguments& InArgs, TSharedRef<FStateTreeVi
 		.AllowOverscroll(EAllowOverscroll::Yes)
 		.ExternalScrollbar(VerticalScrollBar);
 
-	ChildSlot
-	[
-		SNew(SVerticalBox)
-
+	TSharedRef<SWidget> StateTreeView = SNew(SVerticalBox)
 		+ SVerticalBox::Slot()
 		.VAlign(VAlign_Center)
 		.AutoHeight()
@@ -169,8 +171,47 @@ void SStateTreeView::Construct(const FArguments& InArgs, TSharedRef<FStateTreeVi
 		.AutoHeight()
 		[
 			HorizontalScrollBar
-		]
-	];
+		];
+
+	if (!UStateCentricViewSettings::Get().IsStateCentricViewEnabled())
+	{
+		// No state centric view, only tree view
+		ChildSlot
+		[
+			StateTreeView
+		];
+	}
+	else
+	{
+		// @TODO: Wrap main state centric view with something that accepts a null state
+		UStateTreeState* StateCentricViewState = StateTreeViewModel->GetViewState();
+		ensure(StateCentricViewState);
+		ChildSlot
+		[
+			SNew(SSplitter)
+			+ SSplitter::Slot()
+			.Value(0.15f)
+			[
+				StateTreeView
+			]
+			+SSplitter::Slot()
+			.Value(0.85f)
+			[
+				SNew(SScrollBox)
+				.Orientation(Orient_Vertical)
+				+ SScrollBox::Slot()
+				.FillSize(1.0f)
+				[
+					SNew(SBorder)
+					.BorderImage(FAppStyle::Get().GetBrush("Brushes.Recessed"))
+					.Padding(FMargin(2.0, 4.0))
+					[
+						SNew(UE::StateTree::Editor::StateCentricView::SStateCentricViewMainStateNode, InStateTreeViewModel, StateCentricViewState)
+					]
+				]
+			]
+		];
+	}
 
 	UpdateTree(true);
 
diff --git a/Engine/Plugins/Runtime/StateTree/Source/StateTreeEditorModule/Private/SStateTreeViewRow.cpp b/Engine/Plugins/Runtime/StateTree/Source/StateTreeEditorModule/Private/SStateTreeViewRow.cpp
--- a/Engine/Plugins/Runtime/StateTree/Source/StateTreeEditorModule/Private/SStateTreeViewRow.cpp
+++ b/Engine/Plugins/Runtime/StateTree/Source/StateTreeEditorModule/Private/SStateTreeViewRow.cpp
@@ -27,6 +27,10 @@
 #include "Widgets/Text/SRichTextBlock.h"
 #include "Widgets/Input/SButton.h"
 
+// @TODO: Dependency onto state centric to enable hierarchy driven colors. 
+// Should instead be driven by view model.
+#include "StateCentricView/StateCentricViewSettings.h"
+
 #define LOCTEXT_NAMESPACE "StateTreeEditor"
 
 namespace UE::StateTree::Editor
@@ -1247,19 +1251,21 @@ FSlateColor SStateTreeViewRow::GetTitleColor(const float Alpha, const float Ligh
 	const UStateTreeEditorData* EditorData = WeakEditorData.Get();
 
 	FLinearColor Color(FColor(31, 151, 167));
-	
-	if (State != nullptr && EditorData != nullptr)
+
+	if (State && StateTreeViewModel)
 	{
-		if (const FStateTreeEditorColor* FoundColor = EditorData->FindColor(State->ColorRef))
+		Color = StateTreeViewModel->GetStateColor(State);
+
+		// @TODO: Dependency onto state centric to enable hierarchy driven colors. 
+		// Should instead be driven by view model.
+		if (UStateCentricViewSettings::GetSchemaSettingsDefaultFallback(StateTreeViewModel.Get()).bEnableHierarchyDrivenColors)
 		{
-			if (IsRootState() || State->Type == EStateTreeStateType::Subtree)
-			{
-				Color = UE::StateTree::Editor::LerpColorSRGB(FoundColor->Color, FColor::Black, 0.25f);
-			}
-			else
-			{
-				Color = FoundColor->Color;
-			}
+			Color = StateTreeViewModel->GetHierarchyBasedBackgroundColor(State);
+		}
+
+		if (IsRootState() || State->Type == EStateTreeStateType::Subtree)
+		{
+			Color = UE::StateTree::Editor::LerpColorSRGB(Color, FColor::Black, 0.25f);
 		}
 	}
 
diff --git a/Engine/Plugins/Runtime/StateTree/Source/StateTreeEditorModule/Private/StateCentricView/Layout/ExtendableNodeTypes.h b/Engine/Plugins/Runtime/StateTree/Source/StateTreeEditorModule/Private/StateCentricView/Layout/ExtendableNodeTypes.h
new file mode 100644
--- /dev/null
+++ b/Engine/Plugins/Runtime/StateTree/Source/StateTreeEditorModule/Private/StateCentricView/Layout/ExtendableNodeTypes.h
@@ -0,0 +1,22 @@
+// Copyright Epic Games, Inc. All Rights Reserved.
+
+#pragma once
+
+#include "CoreMinimal.h"
+#include "Types/SlateStructs.h"
+#include "ExtendableNodeTypes.generated.h"
+
+/**
+ * Experimental. LOD / displayed info that an extendable node may support.
+ */
+UENUM(Experimental)
+enum class EExtendableNodeLOD : uint8
+{
+	Hidden,					// View provides no context, if appropriate should not even give a placeholder
+	Collapsed,				// View provides no context, but visible as placeholder for when LOD changes
+	Minimal,				// View provides minimal context deemed needed
+	Medium,					// View is filtered, giving some but not all detail
+	Full,					// View would not reasonably show more info
+	Unused UMETA(Hidden),	// Indicate node itself does not change based on LOD. Node may still accept LOD to pass onto children.
+	Default UMETA(Hidden),	// Indicate no preference at API level, exact behavior implementation specific
+};
diff --git a/Engine/Plugins/Runtime/StateTree/Source/StateTreeEditorModule/Private/StateCentricView/Layout/SExtendableNode.cpp b/Engine/Plugins/Runtime/StateTree/Source/StateTreeEditorModule/Private/StateCentricView/Layout/SExtendableNode.cpp
new file mode 100644
--- /dev/null
+++ b/Engine/Plugins/Runtime/StateTree/Source/StateTreeEditorModule/Private/StateCentricView/Layout/SExtendableNode.cpp
@@ -0,0 +1,733 @@
+// Copyright Epic Games, Inc. All Rights Reserved.
+
+#include "StateCentricView/Layout/SExtendableNode.h"
+
+#include "Widgets/Layout/SBorder.h"
+#include "Widgets/Layout/SBox.h"
+#include "Widgets/SBoxPanel.h"
+#include "Widgets/SOverlay.h"
+
+#define LOCTEXT_NAMESPACE "SExtendableNode"
+
+namespace UE::StateTree::Editor::StateCentricView
+{
+
+
+namespace Private
+{
+
+	/** Dummy widget to make widget reflector path clearer */
+	class SHorizontalExtensionBox : public SStackBox
+	{
+	};
+
+	/** Dummy widget to make widget reflector path clearer */
+	class SVerticalExtensionBox : public SStackBox
+	{
+	};
+
+	/** Dummy widget to make widget reflector path clearer */
+	class SHorizontalExtensionSplitter : public SSplitter
+	{
+	};
+
+	/** Dummy widget to make widget reflector path clearer */
+	class SVerticalExtensionSplitter : public SSplitter
+	{
+	};
+
+	/** Dummy widget to make widget reflector path clearer */
+	class SMainNodeStackBox : public SStackBox
+	{
+	};
+
+	/** 
+	 * Support ability to offset layer ID. Needed to handle overlapping UI that is laid out flat correctly
+	 * IE: Two slots next to each other in a box, but they need to have a defined overlap order
+	 */
+	class SOverlayExtensionBox : public SOverlay
+	{
+	public:
+		virtual int32 OnPaint(const FPaintArgs& Args, const FGeometry& AllottedGeometry, const FSlateRect& MyCullingRect, FSlateWindowElementList& OutDrawElements, int32 LayerId, const FWidgetStyle& InWidgetStyle, bool bParentEnabled) const override
+		{
+			int32 OffsetLayerId = LayerId + Offset;
+			return SOverlay::OnPaint(Args, AllottedGeometry, MyCullingRect, OutDrawElements, OffsetLayerId, InWidgetStyle, bParentEnabled);
+		}
+
+	public:
+
+		int32 Offset = 0;
+	};
+
+} // namespace Private
+
+
+void SExtendableNode::Construct(const FArguments& InArgs)
+{
+	Config = InArgs._Config;
+
+	if (Config.InitialLOD == EExtendableNodeLOD::Default)
+	{
+		// Store whatever default was in constructor
+		Config.InitialLOD = CurrentLOD;
+	}
+	else
+	{
+		CurrentLOD = Config.InitialLOD;
+	}
+
+	Construct_GatherExtensionsFromExternalSubsystems();
+	Construct_InitializeNodePersistentPreLOD();
+	RegenerateWidgetForLOD(CurrentLOD);
+}
+
+void SExtendableNode::AddExtension(TSharedRef<SWidget> ExtensionWidget, const FExtensionAddNodeParams& ExtensionParams)
+{
+	const FExtensionAddLayoutInfo& LayoutInfo = ExtensionParams.LayoutInfo;
+
+	enum class EMainNodeExtensionDirection
+	{
+		After,	// Bottom, Right, & Foreground, add after main node in the layout
+		Before,	// Top, Left, & Background, add before main node in the layout
+	};
+
+	// Adding in a cardinal direction is only unique w.r.t. orientation & direction relative to main node
+	auto AddCardinalExtension = [&](const EOrientation InOrientation, const EMainNodeExtensionDirection InAddDirection)
+	{
+		// We know we want to insert, but not where yet. Forward declare so we can keep insert logic consolidated.
+		TSharedPtr<SStackBox> LayoutToInsert = nullptr;
+		int8* MainNodeInsertIndex = nullptr;
+
+		auto WrapCurrentLayout = [&]()
+		{
+			if (InOrientation == Orient_Horizontal)
+			{
+				SAssignNew(CurrentLayoutBox, StateCentricView::Private::SHorizontalExtensionBox)
+					.Orientation(InOrientation);
+			}
+			else
+			{
+				SAssignNew(CurrentLayoutBox, StateCentricView::Private::SVerticalExtensionBox)
+					.Orientation(InOrientation);
+			}
+
+			if (!InnerLayoutBox)
+			{
+				TSharedRef<SWidget> MainNodeWidget = MainNodeOverlay 
+					? StaticCastSharedRef<SWidget>(MainNodeOverlay.ToSharedRef())
+					: StaticCastSharedRef<SWidget>(MainNodeBox.ToSharedRef());
+
+				// @TODO: Remove. We don't support cardinal layouts after adding a splitter.
+				// Plan is to add splitter support as a feature toggle to base layout. See: 48250292.
+				if (CurrentSplitterBox)
+				{
+					MainNodeWidget = CurrentSplitterBox.ToSharedRef();
+				}
+
+				CurrentLayoutBox->AddSlot()
+				.Expose(MainNodeLayoutSlot)
+				.SizeParam(MainNodeLayoutInfo.SizeParam ? MainNodeLayoutInfo.SizeParam : LayoutInfo.SizeParam)
+				.MinSize(MainNodeLayoutInfo.MinSize.IsSet() ? MainNodeLayoutInfo.MinSize : LayoutInfo.MinSize)
+				.MaxSize(MainNodeLayoutInfo.MaxSize.IsSet() ? MainNodeLayoutInfo.MaxSize : LayoutInfo.MaxSize)
+				.Padding(MainNodeLayoutInfo.Padding.IsSet() ? MainNodeLayoutInfo.Padding : LayoutInfo.Padding)
+				.HAlign(MainNodeLayoutInfo.HAlign.Get({ LayoutInfo.HAlign.Get(HAlign_Fill) }))
+				.VAlign(MainNodeLayoutInfo.VAlign.Get({ LayoutInfo.VAlign.Get(VAlign_Fill) }))
+				[
+					MainNodeWidget
+				];
+			}
+			else
+			{
+				CurrentLayoutBox->AddSlot()
+				.SizeParam(MainNodeLayoutInfo.SizeParam ? MainNodeLayoutInfo.SizeParam : LayoutInfo.SizeParam)
+				.MinSize(MainNodeLayoutInfo.MinSize.IsSet() ? MainNodeLayoutInfo.MinSize : LayoutInfo.MinSize)
+				.MaxSize(MainNodeLayoutInfo.MaxSize.IsSet() ? MainNodeLayoutInfo.MaxSize : LayoutInfo.MaxSize)
+				.Padding(MainNodeLayoutInfo.Padding.IsSet() ? MainNodeLayoutInfo.Padding : LayoutInfo.Padding)
+				.HAlign(MainNodeLayoutInfo.HAlign.Get({ LayoutInfo.HAlign.Get(HAlign_Fill) }))
+				.VAlign(MainNodeLayoutInfo.VAlign.Get({ LayoutInfo.VAlign.Get(VAlign_Fill) }))
+				[
+					InnerLayoutBox.ToSharedRef()
+				];
+			}
+			MainNodeCurrentLayoutIndex = 0;
+
+			// Replace the actual widget shown by this composite widget.
+			if (EntireNodeOverlaySlot)
+			{
+				(*EntireNodeOverlaySlot)
+				[
+					CurrentLayoutBox.ToSharedRef()
+				];
+			}
+			else
+			{
+				ChildSlot
+				[
+					CurrentLayoutBox.ToSharedRef()
+				];
+			}
+
+			// Technically not a part of wrapping, but if we just created a layout we'll always insert to it. So assign here.
+			LayoutToInsert = CurrentLayoutBox;
+			MainNodeInsertIndex = &MainNodeCurrentLayoutIndex;
+		};
+
+		if (!CurrentLayoutBox)
+		{
+			WrapCurrentLayout();
+		}
+		else if (CurrentLayoutBox->GetOrientation() != InOrientation)
+		{
+			if (ExtensionParams.AddLocationBehavior == EExtensionAddLocationBehavior::Default
+				|| ExtensionParams.AddLocationBehavior == EExtensionAddLocationBehavior::Wrap)
+			{
+				InnerLayoutBox = CurrentLayoutBox;
+				MainNodeInnerLayoutIndex = MainNodeCurrentLayoutIndex;
+
+				WrapCurrentLayout();
+			}
+			else if (ensure(ExtensionParams.AddLocationBehavior == EExtensionAddLocationBehavior::UseInnerLayout))
+			{
+				if (!InnerLayoutBox)
+				{
+					if (InOrientation == Orient_Horizontal)
+					{
+						SAssignNew(InnerLayoutBox, StateCentricView::Private::SHorizontalExtensionBox)
+							.Orientation(InOrientation);
+					}
+					else
+					{
+						SAssignNew(InnerLayoutBox, StateCentricView::Private::SVerticalExtensionBox)
+							.Orientation(InOrientation);
+					}
+
+					TSharedRef<SWidget> MainNodeWidget = MainNodeOverlay 
+						? StaticCastSharedRef<SWidget>(MainNodeOverlay.ToSharedRef())
+						: StaticCastSharedRef<SWidget>(MainNodeBox.ToSharedRef());
+					
+					(*MainNodeLayoutSlot)
+					[
+						InnerLayoutBox.ToSharedRef()
+					];
+
+					InnerLayoutBox->AddSlot()
+					.Expose(MainNodeLayoutSlot)
+					.SizeParam(MainNodeLayoutInfo.SizeParam ? MainNodeLayoutInfo.SizeParam : LayoutInfo.SizeParam)
+					.MinSize(MainNodeLayoutInfo.MinSize.IsSet() ? MainNodeLayoutInfo.MinSize : LayoutInfo.MinSize)
+					.MaxSize(MainNodeLayoutInfo.MaxSize.IsSet() ? MainNodeLayoutInfo.MaxSize : LayoutInfo.MaxSize)
+					.Padding(MainNodeLayoutInfo.Padding.IsSet() ? MainNodeLayoutInfo.Padding : LayoutInfo.Padding)
+					.HAlign(MainNodeLayoutInfo.HAlign.Get({ LayoutInfo.HAlign.Get(HAlign_Fill) }))
+					.VAlign(MainNodeLayoutInfo.VAlign.Get({ LayoutInfo.VAlign.Get(VAlign_Fill) }))
+					[
+						MainNodeWidget
+					];
+				}
+
+				LayoutToInsert = InnerLayoutBox;
+				MainNodeInsertIndex = &MainNodeInnerLayoutIndex;
+			}
+		}
+		else
+		{
+			LayoutToInsert = CurrentLayoutBox;
+			MainNodeInsertIndex = &MainNodeCurrentLayoutIndex;
+		}
+
+		int8 InsertIndex = INDEX_NONE;
+		if (InAddDirection == EMainNodeExtensionDirection::After)
+		{
+			InsertIndex = ExtensionParams.AddIndex == INDEX_NONE
+				? IntCastChecked<int8>(LayoutToInsert->NumSlots())
+				: *MainNodeInsertIndex + ExtensionParams.AddIndex + 1;
+		}
+		else if (ensure(InAddDirection == EMainNodeExtensionDirection::Before))
+		{
+			InsertIndex = ExtensionParams.AddIndex == INDEX_NONE
+				? 0
+				: *MainNodeInsertIndex - ExtensionParams.AddIndex;
+
+			(*MainNodeInsertIndex)++;
+		}
+
+		LayoutToInsert->InsertSlot(InsertIndex)
+		.SizeParam(LayoutInfo.SizeParam)
+		.MinSize(LayoutInfo.MinSize)
+		.MaxSize(LayoutInfo.MaxSize)
+		.Padding(LayoutInfo.Padding)
+		.HAlign(LayoutInfo.HAlign.Get(HAlign_Fill))
+		.VAlign(LayoutInfo.VAlign.Get(VAlign_Fill))
+		[
+			ExtensionWidget
+		];
+	};
+
+	// Same in overlays are only unique per direction relative to main node
+	auto AddOverlayExtension = [&](const EMainNodeExtensionDirection InAddDirection)
+	{
+		// This is possible to support but not important at the moment
+		checkf(ExtensionParams.AddIndex == INDEX_NONE, TEXT("Unimplemented. SExtendableNode does not currently support specified overlay indexes: %s"), *GetNodeName().ToString());
+
+		// We know we want to insert, but not where yet. Forward declare so we can keep insert logic consolidated.
+		TSharedPtr<SOverlay> OverlayToInsert = nullptr;
+		int8* OverlayToInsertLowestZOrder = nullptr;
+
+		if (ExtensionParams.AddLocationBehavior == EExtensionAddLocationBehavior::Default
+			|| ExtensionParams.AddLocationBehavior == EExtensionAddLocationBehavior::OverlayMainNode)
+		{
+			if (!MainNodeOverlay)
+			{
+				TSharedPtr<StateCentricView::Private::SOverlayExtensionBox> OverlayExtensionBox = SNew(StateCentricView::Private::SOverlayExtensionBox);
+				OverlayExtensionBox->Offset = MainNodeLayoutInfo.OverlayOffset;
+				MainNodeOverlay = OverlayExtensionBox;
+				
+				MainNodeOverlay->AddSlot()
+				.Padding(MainNodeLayoutInfo.Padding)
+				.HAlign(MainNodeLayoutInfo.HAlign.Get(HAlign_Fill))
+				.VAlign(MainNodeLayoutInfo.VAlign.Get(VAlign_Fill))
+				[
+					MainNodeBox.ToSharedRef()
+				];
+
+				// Update correct slot. Starting with main node which 
+				// will be inside any overlays / layout widgets
+
+				// @TODO: Remove. Splitter aspect should be a feature toggle on layout. Not a dedicated incompitable layout. See: 48250292.
+				if (MainNodeSplitterSlot)
+				{
+					(*MainNodeSplitterSlot)
+					[
+						MainNodeOverlay.ToSharedRef()
+					];
+				}
+
+				else if (MainNodeLayoutSlot)
+				{
+					(*MainNodeLayoutSlot)
+					[
+						MainNodeOverlay.ToSharedRef()
+					];
+				}
+				else if (EntireNodeOverlaySlot)
+				{
+					(*EntireNodeOverlaySlot)
+					[
+						MainNodeOverlay.ToSharedRef()
+					];
+				}
+				else
+				{
+					ChildSlot
+					[
+						MainNodeOverlay.ToSharedRef()
+					];
+				}
+
+				// Default to -1 so when we decrement it's -2, which overlay supports for inserting behind.
+				MainNodeLowestZOrder = INDEX_NONE;
+			}
+
+			OverlayToInsert = MainNodeOverlay;
+			OverlayToInsertLowestZOrder = &MainNodeLowestZOrder;
+		}
+		else if (ensure(ExtensionParams.AddLocationBehavior == EExtensionAddLocationBehavior::OverlayEntireNode))
+		{
+			// We don't support entire node overlays after adding a splitter.
+			// Plan is to add splitter support as a feature toggle to base layout. See: 48250292.
+			if (!ensure(!CurrentSplitterBox))
+			{
+				return;
+			}
+
+			if (!EntireNodeOverlay)
+			{
+				TSharedPtr<StateCentricView::Private::SOverlayExtensionBox> OverlayExtensionBox = SNew(StateCentricView::Private::SOverlayExtensionBox);
+				OverlayExtensionBox->Offset = MainNodeLayoutInfo.OverlayOffset;
+				EntireNodeOverlay = OverlayExtensionBox;
+
+				TSharedRef<SWidget> MainNodeWidget = MainNodeOverlay
+					? StaticCastSharedRef<SWidget>(MainNodeOverlay.ToSharedRef())
+					: StaticCastSharedRef<SWidget>(MainNodeBox.ToSharedRef());
+
+				EntireNodeOverlay->AddSlot()
+				.Padding(MainNodeLayoutInfo.Padding)
+				.HAlign(MainNodeLayoutInfo.HAlign.Get(HAlign_Fill))
+				.VAlign(MainNodeLayoutInfo.VAlign.Get(VAlign_Fill))
+				.Expose(EntireNodeOverlaySlot)
+				[
+					CurrentLayoutBox ? CurrentLayoutBox.ToSharedRef() : MainNodeWidget
+				];
+
+				ChildSlot
+				[
+					EntireNodeOverlay.ToSharedRef()
+				];
+
+				// Default to -1 so when we decrement it's -2, which overlay supports for inserting behind.
+				EntireNodeLowestZOrder = INDEX_NONE;
+			}
+
+			OverlayToInsert = EntireNodeOverlay;
+			OverlayToInsertLowestZOrder = &EntireNodeLowestZOrder;
+		}
+
+		if (InAddDirection == EMainNodeExtensionDirection::After)
+		{
+			// Just add at end and let overlay automatically increment Z order based on last child
+			OverlayToInsert->AddSlot()
+			.Padding(LayoutInfo.Padding)
+			.HAlign(LayoutInfo.HAlign.Get(HAlign_Fill))
+			.VAlign(LayoutInfo.VAlign.Get(VAlign_Fill))
+			[
+				ExtensionWidget
+			];
+		}
+		else if (ensure(InAddDirection == EMainNodeExtensionDirection::Before))
+		{
+			// @TODO: Logic for inserts, currently not supported by this API. Possibly should belong in SOverlay
+			//TPanelChildren<SOverlay::FOverlaySlot>* OverlayChildren = static_cast<TPanelChildren<SOverlay::FOverlaySlot>*>(Overlay->GetChildren());
+			//for (int32 ChildIdx = 0; ChildIdx < OverlayChildren->Num(); ++ChildIdx)
+			//{
+			//	SOverlay::FOverlaySlot& OverlaySlot = (*OverlayChildren)[ChildIdx];
+			//	OverlaySlot.SetZOrder(OverlaySlot.GetZOrder() + 1);
+			//}
+			
+			// Use hack of inserting at -2 or lower to insert before widgets at Z Order 0.
+			if (OverlayToInsertLowestZOrder)
+			{
+				(*OverlayToInsertLowestZOrder)--;
+				OverlayToInsert->AddSlot(*OverlayToInsertLowestZOrder)
+				.Padding(LayoutInfo.Padding)
+				.HAlign(LayoutInfo.HAlign.Get(HAlign_Fill))
+				.VAlign(LayoutInfo.VAlign.Get(VAlign_Fill))
+				[
+					ExtensionWidget
+				];
+			}
+		}
+
+	};
+
+	// Add in a cardinal direction with a splitter
+	// @TODO: This should be removed and just a feature toggle on add cardinal direction.
+	auto AddSplitExtension = [&](const EOrientation InOrientation, const EMainNodeExtensionDirection InAddDirection)
+	{
+		if (!CurrentSplitterBox)
+		{
+			if (InOrientation == Orient_Horizontal)
+			{
+				SAssignNew(CurrentSplitterBox, StateCentricView::Private::SHorizontalExtensionSplitter)
+					.PhysicalSplitterHandleSize(0.0f)
+					.Orientation(InOrientation);
+			}
+			else
+			{
+				SAssignNew(CurrentSplitterBox, StateCentricView::Private::SVerticalExtensionSplitter)
+					.PhysicalSplitterHandleSize(0.0f)
+					.Orientation(InOrientation);
+			}
+				
+			CurrentSplitterBox->AddSlot()
+			.Expose(MainNodeSplitterSlot)
+			.SizeRule(MainNodeLayoutInfo.SplitterSizeRule)
+			.Value(MainNodeLayoutInfo.SplitterValue)
+			.OnSlotResized(MainNodeLayoutInfo.SplitterOnSlotResized.Get({}))
+			[
+				MainNodeBox.ToSharedRef()
+			];
+
+			// We don't support arbitrary nesting order with current splitter impl
+			if (ensure(!MainNodeOverlay)
+				&& ensure(!MainNodeLayoutSlot)
+				&& ensure(!EntireNodeOverlaySlot))
+			{
+				ChildSlot
+				[
+					CurrentSplitterBox.ToSharedRef()
+				];
+			}
+		}
+
+		// We don't support nesting or ordered insert for splitters since splitter container support is temp
+		if (ensure(InOrientation == CurrentSplitterBox->GetOrientation())
+			&& ensure(ExtensionParams.AddIndex == INDEX_NONE))
+		{
+			int8 InsertIndex = INDEX_NONE;
+			if (InAddDirection == EMainNodeExtensionDirection::After)
+			{
+				InsertIndex = IntCastChecked<int8>(CurrentSplitterBox->NumSlots());
+			}
+			else if (ensure(InAddDirection == EMainNodeExtensionDirection::Before))
+			{
+				InsertIndex = 0;
+			}
+
+			CurrentSplitterBox->AddSlot(InsertIndex)
+			.SizeRule(LayoutInfo.SplitterSizeRule)
+			.Value(LayoutInfo.SplitterValue)
+			.OnSlotResized(LayoutInfo.SplitterOnSlotResized.Get({}))
+			[
+				ExtensionWidget
+			];
+		}
+
+	};
+
+	switch (ExtensionParams.Location)
+	{
+		case EExtensionLocation::Bottom:
+		{
+			AddCardinalExtension(EOrientation::Orient_Vertical, EMainNodeExtensionDirection::After);
+			break;
+		}
+		case EExtensionLocation::Left:
+		{
+			AddCardinalExtension(EOrientation::Orient_Horizontal, EMainNodeExtensionDirection::Before);
+			break;
+		}
+		case EExtensionLocation::Right:
+		{
+			AddCardinalExtension(EOrientation::Orient_Horizontal, EMainNodeExtensionDirection::After);
+			break;
+		}
+		case EExtensionLocation::Top:
+		{
+			AddCardinalExtension(EOrientation::Orient_Vertical, EMainNodeExtensionDirection::Before);
+			break;
+		}
+		case EExtensionLocation::SubHeader:
+		{
+			if (!SubHeaderBox)
+			{
+				SAssignNew(SubHeaderBox, SStackBox)
+				.Orientation(EOrientation::Orient_Vertical);
+
+				int8 AfterHeaderIndex = NodeHeaderWidget ? 1 : 0;
+
+				MainNodeBox->InsertSlot(AfterHeaderIndex)
+				.AutoSize()
+				[
+					SubHeaderBox.ToSharedRef()
+				];
+			}
+
+			int8 InsertIndex = ExtensionParams.AddIndex == INDEX_NONE
+				? IntCastChecked<int8>(SubHeaderBox->NumSlots())
+				: ExtensionParams.AddIndex;
+
+			SubHeaderBox->InsertSlot(InsertIndex)
+			.SizeParam(LayoutInfo.SizeParam)
+			.MinSize(LayoutInfo.MinSize)
+			.MaxSize(LayoutInfo.MaxSize)
+			.Padding(LayoutInfo.Padding)
+			.HAlign(LayoutInfo.HAlign.Get(HAlign_Fill))
+			.VAlign(LayoutInfo.VAlign.Get(VAlign_Fill))
+			[
+				ExtensionWidget
+			];
+
+			break;
+		}
+		case EExtensionLocation::Footnote:
+		{
+			if (!FootnoteBox)
+			{
+				SAssignNew(FootnoteBox, SStackBox)
+					.Orientation(EOrientation::Orient_Vertical);
+
+				// Footnote is always last, so just add
+				MainNodeBox->AddSlot()
+				.AutoSize()
+				[
+					FootnoteBox.ToSharedRef()
+				];
+			}
+
+			int8 InsertIndex = ExtensionParams.AddIndex == INDEX_NONE
+				? IntCastChecked<int8>(FootnoteBox->NumSlots())
+				: ExtensionParams.AddIndex;
+
+			FootnoteBox->InsertSlot(InsertIndex)
+			.SizeParam(LayoutInfo.SizeParam)
+			.MinSize(LayoutInfo.MinSize)
+			.MaxSize(LayoutInfo.MaxSize)
+			.Padding(LayoutInfo.Padding)
+			.HAlign(LayoutInfo.HAlign.Get(HAlign_Fill))
+			.VAlign(LayoutInfo.VAlign.Get(VAlign_Fill))
+			[
+				ExtensionWidget
+			];
+
+			break;
+		}
+		case EExtensionLocation::Foreground:
+		{
+			AddOverlayExtension(EMainNodeExtensionDirection::After);
+			break;
+		}
+		case EExtensionLocation::Background:
+		{
+			AddOverlayExtension(EMainNodeExtensionDirection::Before);
+			break;
+		}
+		case EExtensionLocation::SplitBottom:
+		{
+			AddSplitExtension(EOrientation::Orient_Vertical, EMainNodeExtensionDirection::After);
+			break;
+		}
+		case EExtensionLocation::SplitLeft:
+		{
+			AddSplitExtension(EOrientation::Orient_Horizontal, EMainNodeExtensionDirection::Before);
+			break;
+		}
+		case EExtensionLocation::SplitRight:
+		{
+			AddSplitExtension(EOrientation::Orient_Horizontal, EMainNodeExtensionDirection::After);
+			break;
+		}
+		case EExtensionLocation::SplitTop:
+		{
+			AddSplitExtension(EOrientation::Orient_Vertical, EMainNodeExtensionDirection::Before);
+			break;
+		}
+		default:
+			ensureMsgf(false, TEXT("Invalid extension location: %s - %u"), *GetNodeName().ToString(), ExtensionParams.Location);
+			break;
+	}
+
+#if UE_EXTENDABLE_NODE_DEBUG
+	AddedExtensions.Add({ ExtensionParams, ExtensionWidget });
+#endif // UE_EXTENDABLE_NODE_DEBUG
+}
+
+FScopedExtensionAddNodeParams SExtendableNode::AddExtension(TSharedRef<SWidget> ExtensionWidget)
+{
+	return FScopedExtensionAddNodeParams(StaticCastSharedRef<SExtendableNode>(AsShared()), ExtensionWidget);
+}
+
+void SExtendableNode::SetMainNodeLayoutInfo(const FExtensionAddLayoutInfo& InLayoutInfo)
+{
+	MainNodeLayoutInfo = InLayoutInfo;
+}
+
+FScopedMainNodeLayoutInfo SExtendableNode::SetScopedMainNodeLayoutInfo()
+{
+	return FScopedMainNodeLayoutInfo(StaticCastSharedRef<SExtendableNode>(AsShared()));
+}
+
+void SExtendableNode::RegenerateWidgetForLOD(const EExtendableNodeLOD InLOD)
+{
+	CurrentLOD = InLOD;
+
+	ResetCachedNodeAndLayout();
+
+	Construct_GenerateMainNode(CurrentLOD);
+	Construct_GenerateBuiltInExtensions(CurrentLOD);
+	Construct_GenerateExternalExtensions(CurrentLOD);
+}
+
+EExtendableNodeLOD SExtendableNode::GetCurrentLOD() const
+{
+	return CurrentLOD;
+}
+
+void SExtendableNode::Construct_GatherExtensionsFromExternalSubsystems()
+{
+	// @TODO: Gather from a default ExtendableNode subsytem using our unique name ID & add to config
+}
+
+void SExtendableNode::Construct_InitializeNodePersistentPreLOD()
+{
+	// No-op
+}
+
+void SExtendableNode::Construct_GenerateMainNode(const EExtendableNodeLOD InLOD)
+{
+	TSharedRef<SWidget> Header = Construct_GenerateNodeHeader(InLOD);
+	TSharedRef<SWidget> Body = Construct_GenerateNodeBody(InLOD);
+
+	SAssignNew(MainNodeBox, StateCentricView::Private::SMainNodeStackBox)
+	.Orientation(EOrientation::Orient_Vertical);
+
+	if (Header != SNullWidget::NullWidget)
+	{
+		NodeHeaderWidget = Header;
+		MainNodeBox->AddSlot()
+		.Padding(NodeHeaderPadding)
+		.SizeParam(NodeHeaderSizeParam.Get(FAuto()))
+		[
+			NodeHeaderWidget.ToSharedRef()
+		];
+	}
+
+	if (Body != SNullWidget::NullWidget)
+	{
+		NodeBodyWidget = Body;
+		MainNodeBox->AddSlot()
+		.Padding(NodeBodyPadding)
+		.SizeParam(NodeBodySizeParam.Get(FAuto()))
+		[
+			NodeBodyWidget.ToSharedRef()
+		];
+	}
+
+	ChildSlot
+	[
+		MainNodeBox.ToSharedRef()
+	];
+}
+
+void SExtendableNode::Construct_GenerateBuiltInExtensions(const EExtendableNodeLOD InLOD)
+{
+	// No-op
+}
+
+void SExtendableNode::Construct_GenerateExternalExtensions(const EExtendableNodeLOD InLOD)
+{
+	for (FAddExtensionDelegate& ExternalExtension : Config.ExternalExtensions)
+	{
+		ExternalExtension.ExecuteIfBound(StaticCastSharedRef<SExtendableNode>(AsShared()), InLOD);
+	}
+}
+
+TSharedRef<SWidget> SExtendableNode::Construct_GenerateNodeHeader(const EExtendableNodeLOD InLOD)
+{
+	return SNullWidget::NullWidget;
+}
+
+TSharedRef<SWidget> SExtendableNode::Construct_GenerateNodeBody(const EExtendableNodeLOD InLOD)
+{
+	return SNullWidget::NullWidget;
+}
+
+void SExtendableNode::ResetCachedNodeAndLayout()
+{
+	NodeHeaderWidget = nullptr;
+	NodeBodyWidget = nullptr;
+	MainNodeBox = nullptr;
+	SubHeaderBox = nullptr;
+	FootnoteBox = nullptr;
+
+	MainNodeLayoutInfo = {};
+
+	CurrentLayoutBox = nullptr;
+	CurrentSplitterBox = nullptr;
+	InnerLayoutBox = nullptr;
+	EntireNodeOverlay = nullptr;
+	MainNodeOverlay = nullptr;
+	MainNodeLayoutSlot = nullptr;
+	MainNodeSplitterSlot = nullptr;
+	EntireNodeOverlaySlot = nullptr;
+
+	MainNodeCurrentLayoutIndex = INDEX_NONE;
+	MainNodeInnerLayoutIndex = INDEX_NONE;
+	EntireNodeLowestZOrder = INDEX_NONE;
+	MainNodeLowestZOrder = INDEX_NONE;
+
+#if UE_EXTENDABLE_NODE_DEBUG
+	AddedExtensions.Reset();
+#endif // UE_EXTENDABLE_NODE_DEBUG
+}
+
+} // namespace UE::StateTree::Editor::StateCentricView
+
+#undef LOCTEXT_NAMESPACE
diff --git a/Engine/Plugins/Runtime/StateTree/Source/StateTreeEditorModule/Private/StateCentricView/Layout/SExtendableNode.h b/Engine/Plugins/Runtime/StateTree/Source/StateTreeEditorModule/Private/StateCentricView/Layout/SExtendableNode.h
new file mode 100644
--- /dev/null
+++ b/Engine/Plugins/Runtime/StateTree/Source/StateTreeEditorModule/Private/StateCentricView/Layout/SExtendableNode.h
@@ -0,0 +1,569 @@
+// Copyright Epic Games, Inc. All Rights Reserved.
+
+#pragma once
+
+#include "CoreMinimal.h"
+
+#include "StateCentricView/Layout/ExtendableNodeTypes.h"
+#include "Types/SlateEnums.h"
+#include "Widgets/Layout/SSplitter.h"
+#include "Widgets/SBoxPanel.h"
+#include "Widgets/SCompoundWidget.h"
+#include "Widgets/SOverlay.h"
+#include "Widgets/Text/STextBlock.h"
+
+
+#ifndef UE_EXTENDABLE_NODE_DEBUG
+	#define UE_EXTENDABLE_NODE_DEBUG 1
+#endif
+
+namespace UE::StateTree::Editor::StateCentricView
+{
+	class SExtendableNode;
+	struct FScopedExtensionAddNodeParams;
+	struct FScopedMainNodeLayoutInfo;
+}
+
+namespace UE::StateTree::Editor::StateCentricView
+{
+
+
+/** Experimental. Location for a node extension. */
+enum class EExtensionLocation : uint8
+{
+	Bottom,			// Below the node
+	Left,			// Left of node
+	Right,			// Right of node
+	Top,			// Above the node 
+
+	SubHeader,		// Between main node header and body
+	Footnote,		// Below the main node body.
+
+	Foreground,		// Overlay in front of node. 
+	Background,		// Overlay behind node. 
+
+	// @TODO: Remove. Drag to resize should be a feature toggle not a container switch. See: 48250292.
+	// Right now splitter / stackbox have disjoint functionality sets for no reason so we need these.
+	// For now these have intentionally strict, known golden path only, support.
+
+	SplitBottom,	// Below the node
+	SplitLeft,		// Left of node
+	SplitRight,		// Right of node
+	SplitTop,		// Above the node 
+};
+
+
+/** Experimental. How a newly added extension should behave if there are any special considerations. */
+enum class EExtensionAddLocationBehavior : uint8
+{
+	Default,			// Dynamic default depending on what the given location is
+
+	Wrap,				// (Default) If current layout doesn't support location, wrap with a layout that does
+	UseInnerLayout,		// If current layout doesn't support location, use first inner layout that does or create one
+
+	OverlayMainNode,	// (Default) Add this overlay just for the main node, excluding extensions
+	OverlayEntireNode,	// Add this overlay for the entire node including extensions (IE: Entire widget).
+};
+
+
+/**
+ * Experimental. Callback node will fire when it is ready for generate external extensions
+ * If you need domain specific info beyond LOD to determine if an extension is needed. Please subclass & cast ExtendableNode with that info.
+ *
+ * @param ExtendableNode	- Node to extend by programatically calling `SExtendableNode::AddExtension` based on provided LOD,
+ * @param LOD				- Convenience param. LOD for the current node. Same as ExtendableNode's LOD.
+ */
+using FAddExtensionDelegate = TDelegate<void(TSharedRef<SExtendableNode> /*ExtendableNode*/, const EExtendableNodeLOD /*LOD*/)>;
+
+
+/** Experimental. Config / params struct for extendable nodes. Persists as long as the widget exists. Delegates may re-execute deferred as LOD changes. */
+struct FExtendableNodeConfig
+{
+	/** External extensions for this node. Provided at construction and also gathered in SExtendableNode::Construct_GatherExtensionsFromExternalSubsystems. */
+	TArray<FAddExtensionDelegate> ExternalExtensions;
+
+	/** Initial LOD for this widget. Only used if non-default. If default wil get over-written to Construct time initial value. */
+	EExtendableNodeLOD InitialLOD = EExtendableNodeLOD::Default;
+};
+
+
+/** Experimental. Struct holding layout info used when adding extensions */
+struct FExtensionAddLayoutInfo
+{
+	/** Sizing rules, auto = desired size / stretch = fill / strench content = desired size + fill */
+	TOptional<FSizeParam> SizeParam;
+
+	/** Min size of this widget */
+	TAttribute<float> MinSize;
+
+	/** Max size of this widget */
+	TAttribute<float> MaxSize;
+
+	/** Padding of extension */
+	TAttribute<FMargin> Padding;
+
+	/** Horizontal Alignment */
+	TOptional<EHorizontalAlignment> HAlign;
+
+	/** Vertical Alignment */
+	TOptional<EVerticalAlignment> VAlign;
+
+	/** 
+	 * Layer ID Offset applied to overlay. Used to get consistent widget overlap behavior across contiguous entries in stack boxes.
+	 * Only applies to main node layout at the moment.
+	 */
+	int32 OverlayOffset = 0;
+
+	/** Splitter Sizing Rule */
+	TAttribute<SSplitter::ESizeRule> SplitterSizeRule;
+
+	/** Splitter Ratio Value */
+	TAttribute<float> SplitterValue;
+
+	/** Splitter On Resize Callback */
+	TOptional<SSplitter::FOnSlotResized> SplitterOnSlotResized;
+};
+
+
+/** Experimental. Struct holding params for adding a node extension widget. Supports argument chaining via `FScopedExtensionAddNodeParams`. */
+struct FExtensionAddNodeParams
+{
+	/** Location to add extension */
+	EExtensionLocation Location = EExtensionLocation::SubHeader;
+
+	/** Any specific behavior to consider when adding to vien location */
+	EExtensionAddLocationBehavior AddLocationBehavior = EExtensionAddLocationBehavior::Default;
+
+	/** If positive, index to add widget for given location (Ex: Left 0 will be 1st left widget) */
+	int8 AddIndex = INDEX_NONE; 
+
+	/** Size info when adding extension */
+	FExtensionAddLayoutInfo LayoutInfo = {};
+};
+
+
+/** 
+ * Experimental. Widget with a central Main Node (Header / Body), that can be extended horizontally / vertically / etc.
+ * Also provides conditional node & extension re-generation based on LOD.
+ * 
+ * Simplified Node <Extension> Diagram:
+ * 
+ *    *---------------------------------*        *---------*-------------*---------*
+ *    |           <Top>                 |        |         | <Top>       |         |
+ *    *---------*-------------*---------*        |         *-------------*         |
+ *    |         | Header      |         |        |         | Header      |         |
+ *    |         *-------------*         |        |         *-------------*         |
+ *    | <Left>  | <SubHeader> | <Right> |   or   | <Left>  | <SubHeader> | <Right> |
+ *    |         *-------------*         |        |         *-------------*         |
+ *    |         | Body        |         |        |         | Body        |         |
+ *    *---------*-------------*---------*        |         *-------------*         |
+ *    |           <Bottom>              |        |         | <Bottom>    |         |
+ *    *---------------------------------*        *---------*-------------*---------*
+ * 
+ * The corners are order dependant. See `EExtensionAddLocationBehavior` for details.
+ */
+class SExtendableNode : public SCompoundWidget
+{
+
+public:
+
+	SLATE_BEGIN_ARGS(SExtendableNode)
+		: _Config(FExtendableNodeConfig())
+		{
+		}
+
+		/** Initial set of item categories that this view should show - may be adjusted further by things like CanShowClasses or legacy delegate bindings */
+		SLATE_ARGUMENT(FExtendableNodeConfig, Config)
+
+	SLATE_END_ARGS()
+
+public:
+
+	friend struct FScopedMainNodeLayoutInfo;
+
+public:
+
+	void Construct(const FArguments& InArgs);
+
+	/** Adds the given extension widget, constructing layout widgets as needed to support location */
+	void AddExtension(TSharedRef<SWidget> ExtensionWidget, const FExtensionAddNodeParams& ExtensionParams);
+
+	/** Add extension, supports argument chaining syntax to build params. Includes slate-style methods like AutoSize / FillSize / etc. */
+	FScopedExtensionAddNodeParams AddExtension(TSharedRef<SWidget> ExtensionWidget);
+
+	/** 
+	 * Adding Layout incompatible extensions requires 2 slots. Main Node & Extension. This can be used to layout the main node slot. 
+	 * If not set for cardinal extensions will re-use the extension layout for the main node. Ex: Fill next to fill, desired size next to desired size.
+	 * @TODO: Get rid of this reuse behavior. Or make it an option that is disabled by default / can be disabled. The automatic change something else is confusing.
+	 */
+	void SetMainNodeLayoutInfo(const FExtensionAddLayoutInfo& InLayoutInfo);
+
+	/** Same as above, but resets the value once return value goes out of scope */
+	[[nodiscard]] FScopedMainNodeLayoutInfo SetScopedMainNodeLayoutInfo();
+
+	/** Regenerates this widget for the given LOD, also updates current LOD. */
+	virtual void RegenerateWidgetForLOD(const EExtendableNodeLOD InLOD);
+
+	/** Get current LOD of this widget */
+	EExtendableNodeLOD GetCurrentLOD() const;
+
+	/** 
+	 * Get the unique name for this node, used as a key so that external systems can extend this node.
+	 * If you need access to this key to extend then we suggest just hardcoding a copy locally.
+	 */
+	virtual FName GetNodeName() const = 0;
+
+protected:
+
+	/** One time gather of subsystem externally registered extensions for this node, cached into node config. */
+	virtual void Construct_GatherExtensionsFromExternalSubsystems();
+
+	/** One time init for anything persistent for the node across LODs, pre LOD generation. */
+	virtual void Construct_InitializeNodePersistentPreLOD();
+
+	/** 
+	 * Generation widget header & body based on LOD. Also initally populates ChildSlot
+	 * Should not consider extensions as those are transient between LODs & will be re-added if needed. 
+	 */
+	virtual void Construct_GenerateMainNode(const EExtendableNodeLOD InLOD);
+
+	/** Generation of built in extensions based on LOD. */
+	virtual void Construct_GenerateBuiltInExtensions(const EExtendableNodeLOD InLOD);
+
+	/** Generation of external extensions based on LOD. */
+	virtual void Construct_GenerateExternalExtensions(const EExtendableNodeLOD InLOD);
+
+	/** Builds header for this node based on LOD */
+	virtual TSharedRef<SWidget> Construct_GenerateNodeHeader(const EExtendableNodeLOD InLOD);
+
+	/** Builds body for this node based on LOD */
+	virtual TSharedRef<SWidget> Construct_GenerateNodeBody(const EExtendableNodeLOD InLOD);
+
+	/** Resets values stored during node construction, used to ensure a blank state when LOD changes */
+	void ResetCachedNodeAndLayout();
+
+protected:
+
+	/** Persistent config for this node, such as extension delegates. */
+	FExtendableNodeConfig Config = {};
+
+	/** Current LOD for this node. */
+	EExtendableNodeLOD CurrentLOD = EExtendableNodeLOD::Unused;
+
+	/** Widget used for header, if any. Will not use SNullWidget. */
+	TSharedPtr<SWidget> NodeHeaderWidget = nullptr;
+
+	/** Widget used for body, if any. Will not use SNullWidget. */
+	TSharedPtr<SWidget> NodeBodyWidget = nullptr;
+
+	/** Container for header, subheader, and body. */
+	TSharedPtr<SStackBox> MainNodeBox = nullptr;
+
+	/** Widget for SubHeader extensions, nullptr till a single subheader extension exists. */
+	TSharedPtr<SStackBox> SubHeaderBox = nullptr;
+
+	/** Widget for Footnote extensions, nullptr till a single subheader extension exists. */
+	TSharedPtr<SStackBox> FootnoteBox = nullptr;
+
+	/** Padding for header */
+	TAttribute<FMargin> NodeHeaderPadding;
+
+	/** Padding for body */
+	TAttribute<FMargin> NodeBodyPadding;
+
+	/** Sizing for header */
+	TOptional<FSizeParam> NodeHeaderSizeParam;
+
+	/** Sizing for body */
+	TOptional<FSizeParam> NodeBodySizeParam;
+
+private:
+
+	/** This supports custom layout for the main node slot, used for main node on for first time new layouts. */
+	FExtensionAddLayoutInfo MainNodeLayoutInfo = {};
+
+	/** Widget for current layout, nullptr till a horizontal / vertical extension is added */
+	TSharedPtr<SStackBox> CurrentLayoutBox = nullptr;
+
+	/** @TODO: Remove. Box holding current top level splitter. */
+	TSharedPtr<SSplitter> CurrentSplitterBox = nullptr;
+
+	/** Widget for first inner layout, nullptr till an extension with an incomptible location is added */
+	TSharedPtr<SStackBox> InnerLayoutBox = nullptr;
+
+	/** Widget for entire node overlays (includes extensions), nullptr till a foreground / background extension is added */
+	TSharedPtr<SOverlay> EntireNodeOverlay = nullptr;
+
+	/** Widget for main node overlays, nullptr till a foreground / background extension is added */
+	TSharedPtr<SOverlay> MainNodeOverlay = nullptr;
+
+	/** Innermost layout slot holding main node or widget containing it. Remains constant once created no matter future extensions. */
+	SStackBox::FSlot* MainNodeLayoutSlot = nullptr;
+
+	/** @TODO: Remove. Innermost splitter slot holding main node or widget containing it. */
+	SSplitter::FSlot* MainNodeSplitterSlot = nullptr;
+
+	/** Innermost overlay slot holding entire node Remains constant once created no matter future extensions. */
+	SOverlay::FOverlaySlot* EntireNodeOverlaySlot = nullptr;
+
+	/** Index to main node or widget containing it in the current layout */
+	int8 MainNodeCurrentLayoutIndex = INDEX_NONE;
+
+	/** Index to main node or widget containing it in the first inner layout */
+	int8 MainNodeInnerLayoutIndex = INDEX_NONE;
+
+	/** ZOrder for the furthest entire node background overlay (includes extensions). Used to continously insert behind the current background. */
+	int8 EntireNodeLowestZOrder = INDEX_NONE;
+
+	/** ZOrder for the furthest main node background overlay. Used to continously insert behind the current background */
+	int8 MainNodeLowestZOrder = INDEX_NONE;
+
+#if UE_EXTENDABLE_NODE_DEBUG
+	/** Trace of all extensions added in order */
+	TArray<TPair<FExtensionAddNodeParams, TSharedPtr<SWidget>>> AddedExtensions;
+#endif // UE_EXTENDABLE_NODE_DEBUG
+
+};
+
+
+/** Experimental. Scoped struct to hold add param methods & support function chaining syntax.  */
+struct FScopedExtensionAddNodeParams : public FExtensionAddNodeParams
+{
+public:
+	FScopedExtensionAddNodeParams(TSharedRef<SExtendableNode> InAddNode, TSharedPtr<SWidget> InExtensionWidget)
+		: AddNode(InAddNode)
+		, ExtensionWidget(InExtensionWidget)
+	{
+	}
+
+	FScopedExtensionAddNodeParams() = delete;
+	FScopedExtensionAddNodeParams(const FScopedExtensionAddNodeParams&) = delete;
+	FScopedExtensionAddNodeParams& operator=(const FScopedExtensionAddNodeParams&) = delete;
+	FScopedExtensionAddNodeParams(FScopedExtensionAddNodeParams&&) = default;
+	FScopedExtensionAddNodeParams& operator=(FScopedExtensionAddNodeParams&&) = default;
+
+	virtual ~FScopedExtensionAddNodeParams()
+	{
+		if (ensure(AddNode))	
+		{
+			AddNode->AddExtension(ExtensionWidget.ToSharedRef(), *this);
+		}
+	}
+
+public:
+
+	FScopedExtensionAddNodeParams& SetLocation(const EExtensionLocation InLocation)
+	{
+		Location = InLocation;
+		return *this;
+	}
+
+	FScopedExtensionAddNodeParams& SetAddLocationBehavior(const EExtensionAddLocationBehavior InAddLocationBehavior)
+	{
+		AddLocationBehavior = InAddLocationBehavior;
+		return *this;
+	}
+
+	FScopedExtensionAddNodeParams& SetAddIndex(const int8 InAddIndex)
+	{
+		AddIndex = InAddIndex;
+		return *this;
+	}
+
+	FScopedExtensionAddNodeParams& SetAutoSize()
+	{
+		LayoutInfo.SizeParam = FAuto();
+		return *this;
+	}
+
+	FScopedExtensionAddNodeParams& SetFillSize(TAttribute<float> InStretchCoefficient)
+	{
+		LayoutInfo.SizeParam = FStretch(MoveTemp(InStretchCoefficient));
+		return *this;
+	}
+
+	FScopedExtensionAddNodeParams& SetFillContentSize(TAttribute<float> InStretchCoefficient, TAttribute<float> InShrinkStretchCoefficient = TAttribute<float>())
+	{
+		LayoutInfo.SizeParam = FStretchContent(MoveTemp(InStretchCoefficient), MoveTemp(InShrinkStretchCoefficient));
+		return *this;
+	}
+
+	FScopedExtensionAddNodeParams& SetMinSize(TAttribute<float> InMinSize)
+	{
+		LayoutInfo.MinSize = MoveTemp(InMinSize);
+		return *this;
+	}
+
+	FScopedExtensionAddNodeParams& SetMaxSize(TAttribute<float> InMaxSize)
+	{
+		LayoutInfo.MaxSize = MoveTemp(InMaxSize);
+		return *this;
+	}
+
+	FScopedExtensionAddNodeParams& SetPadding(TAttribute<FMargin> InPadding)
+	{
+		LayoutInfo.Padding = MoveTemp(InPadding);
+		return *this;
+	}
+
+	FScopedExtensionAddNodeParams& SetHAlign(TOptional<EHorizontalAlignment> InHAlign)
+	{
+		LayoutInfo.HAlign = MoveTemp(InHAlign);
+		return *this;
+	}
+
+	FScopedExtensionAddNodeParams& SetVAlign(TOptional<EVerticalAlignment> InVAlign)
+	{
+		LayoutInfo.VAlign = MoveTemp(InVAlign);
+		return *this;
+	}
+
+	FScopedExtensionAddNodeParams& SetSplitterSizeRule(TAttribute<SSplitter::ESizeRule> InSplitterSizeRule)
+	{
+		LayoutInfo.SplitterSizeRule = MoveTemp(InSplitterSizeRule);
+		return *this;
+	}
+
+	FScopedExtensionAddNodeParams& SetSplitterValue(TAttribute<float> InSplitterValue)
+	{
+		LayoutInfo.SplitterValue = MoveTemp(InSplitterValue);
+		return *this;
+	}
+
+	FScopedExtensionAddNodeParams& SetSplitterOnSlotResized(TOptional<SSplitter::FOnSlotResized> InSplitterOnSlotResized)
+	{
+		LayoutInfo.SplitterOnSlotResized = MoveTemp(InSplitterOnSlotResized);
+		return *this;
+	}
+
+private:
+	TSharedPtr<SExtendableNode> AddNode = nullptr;
+	TSharedPtr<SWidget> ExtensionWidget = nullptr;
+};
+
+
+/** Experimental. Scoped struct to toggle main node size info & support function chaining syntax. */
+struct FScopedMainNodeLayoutInfo
+{
+
+protected:
+
+	friend class SExtendableNode;
+
+	FScopedMainNodeLayoutInfo(TSharedRef<SExtendableNode> InAddNode)
+		: AddNode(InAddNode)
+		, bCanRestore(false)
+	{
+		PrevAddSizeInfo = AddNode->MainNodeLayoutInfo;
+	}
+
+public:
+
+	FScopedMainNodeLayoutInfo() = delete;
+	FScopedMainNodeLayoutInfo(const FScopedMainNodeLayoutInfo& Other)
+	{
+		*this = Other;
+	}
+
+	FScopedMainNodeLayoutInfo& operator=(const FScopedMainNodeLayoutInfo& Other)
+	{
+		// Only restore in dtor if we are a copy, first instance is used to enable single-line argument chaning syntax
+		bCanRestore = true;
+		AddNode = Other.AddNode;
+		PrevAddSizeInfo = Other.PrevAddSizeInfo;
+
+		return *this;
+	}
+
+	FScopedMainNodeLayoutInfo(FScopedMainNodeLayoutInfo&& Other) = default;
+	FScopedMainNodeLayoutInfo& operator=(FScopedMainNodeLayoutInfo&& Other) = default;
+
+	virtual ~FScopedMainNodeLayoutInfo()
+	{
+		if (ensure(AddNode) && bCanRestore)
+		{
+			AddNode->MainNodeLayoutInfo = PrevAddSizeInfo;
+		}
+	}
+
+public:
+
+	FScopedMainNodeLayoutInfo& SetAutoSize()
+	{
+		AddNode->MainNodeLayoutInfo.SizeParam = FAuto();
+		return *this;
+	}
+
+	FScopedMainNodeLayoutInfo& SetFillSize(TAttribute<float> InStretchCoefficient)
+	{
+		AddNode->MainNodeLayoutInfo.SizeParam = FStretch(MoveTemp(InStretchCoefficient));
+		return *this;
+	}
+
+	FScopedMainNodeLayoutInfo& SetFillContentSize(TAttribute<float> InStretchCoefficient, TAttribute<float> InShrinkStretchCoefficient = TAttribute<float>())
+	{
+		AddNode->MainNodeLayoutInfo.SizeParam = FStretchContent(MoveTemp(InStretchCoefficient), MoveTemp(InShrinkStretchCoefficient));
+		return *this;
+	}
+
+	FScopedMainNodeLayoutInfo& SetMinSize(TAttribute<float> InMinSize)
+	{
+		AddNode->MainNodeLayoutInfo.MinSize = MoveTemp(InMinSize);
+		return *this;
+	}
+
+	FScopedMainNodeLayoutInfo& SetMaxSize(TAttribute<float> InMaxSize)
+	{
+		AddNode->MainNodeLayoutInfo.MaxSize = MoveTemp(InMaxSize);
+		return *this;
+	}
+
+	FScopedMainNodeLayoutInfo& SetPadding(TAttribute<FMargin> InPadding)
+	{
+		AddNode->MainNodeLayoutInfo.Padding = MoveTemp(InPadding);
+		return *this;
+	}
+
+	FScopedMainNodeLayoutInfo& SetHAlign(TOptional<EHorizontalAlignment> InHAlign)
+	{
+		AddNode->MainNodeLayoutInfo.HAlign = MoveTemp(InHAlign);
+		return *this;
+	}
+
+	FScopedMainNodeLayoutInfo& SetVAlign(TOptional<EVerticalAlignment> InVAlign)
+	{
+		AddNode->MainNodeLayoutInfo.VAlign = MoveTemp(InVAlign);
+		return *this;
+	}
+
+	FScopedMainNodeLayoutInfo& SetOverlayOffset(int32 InOverlayOffset)
+	{
+		AddNode->MainNodeLayoutInfo.OverlayOffset = InOverlayOffset;
+		return *this;
+	}
+
+	FScopedMainNodeLayoutInfo& SetSplitterSizeRule(TAttribute<SSplitter::ESizeRule> InSplitterSizeRule)
+	{
+		AddNode->MainNodeLayoutInfo.SplitterSizeRule = MoveTemp(InSplitterSizeRule);
+		return *this;
+	}
+
+	FScopedMainNodeLayoutInfo& SetSplitterValue(TAttribute<float> InSplitterValue)
+	{
+		AddNode->MainNodeLayoutInfo.SplitterValue = MoveTemp(InSplitterValue);
+		return *this;
+	}
+
+	FScopedMainNodeLayoutInfo& SetSplitterOnSlotResized(TOptional<SSplitter::FOnSlotResized> InSplitterOnSlotResized)
+	{
+		AddNode->MainNodeLayoutInfo.SplitterOnSlotResized = MoveTemp(InSplitterOnSlotResized);
+		return *this;
+	}
+
+private:
+	TSharedPtr<SExtendableNode> AddNode = nullptr;
+	FExtensionAddLayoutInfo PrevAddSizeInfo = FExtensionAddLayoutInfo();
+	bool bCanRestore = false;
+};
+
+} // namespace UE::StateTree::Editor::StateCentricView
\ No newline at end of file
diff --git a/Engine/Plugins/Runtime/StateTree/Source/StateTreeEditorModule/Private/StateCentricView/StateCentricViewPerUserSettings.cpp b/Engine/Plugins/Runtime/StateTree/Source/StateTreeEditorModule/Private/StateCentricView/StateCentricViewPerUserSettings.cpp
new file mode 100644
--- /dev/null
+++ b/Engine/Plugins/Runtime/StateTree/Source/StateTreeEditorModule/Private/StateCentricView/StateCentricViewPerUserSettings.cpp
@@ -0,0 +1,36 @@
+// Copyright Epic Games, Inc. All Rights Reserved.
+
+#include "StateCentricView/StateCentricViewPerUserSettings.h"
+
+#include UE_INLINE_GENERATED_CPP_BY_NAME(StateCentricViewPerUserSettings)
+
+UStateCentricViewPerUserSettings::UStateCentricViewPerUserSettings()
+{
+	FStateCentricViewPerUserSetting DefaultPerUserSettings;
+	DefaultPerUserSettings.ParentHireachyLOD = EExtendableNodeLOD::Minimal;
+	DefaultPerUserSettings.InTransitionViewLOD = EExtendableNodeLOD::Minimal;
+	DefaultPerUserSettings.OutTransitionLOD = EExtendableNodeLOD::Full;
+	// @TODO: Feedback here was 2 in passing, but that seems super low.
+	DefaultPerUserSettings.NumParentsShown = 4; 
+	PerSchemaSettings.Emplace(UStateTreeSchema::StaticClass(), DefaultPerUserSettings);
+}
+
+#if WITH_EDITOR
+
+FText UStateCentricViewPerUserSettings::GetSectionText() const
+{
+	return NSLOCTEXT("StateTreeEditor", "StateCentricViewPerUserSettingsName", "State Centric View Per User");
+}
+
+FText UStateCentricViewPerUserSettings::GetSectionDescription() const
+{
+	return NSLOCTEXT("StateTreeEditor", "StateCentricViewPerUserSettingsDescription", "Experimental. Configure per user options for state centric view.");
+}
+
+#endif // WITH_EDITOR
+
+FName UStateCentricViewPerUserSettings::GetCategoryName() const
+{
+	// @TODO: Change me
+	return FName(TEXT("Plugins"));
+}
diff --git a/Engine/Plugins/Runtime/StateTree/Source/StateTreeEditorModule/Private/StateCentricView/StateCentricViewPerUserSettings.h b/Engine/Plugins/Runtime/StateTree/Source/StateTreeEditorModule/Private/StateCentricView/StateCentricViewPerUserSettings.h
new file mode 100644
--- /dev/null
+++ b/Engine/Plugins/Runtime/StateTree/Source/StateTreeEditorModule/Private/StateCentricView/StateCentricViewPerUserSettings.h
@@ -0,0 +1,125 @@
+// Copyright Epic Games, Inc. All Rights Reserved.
+
+#pragma once
+
+#include "Engine/DeveloperSettings.h"
+#include "Layout/ExtendableNodeTypes.h"
+#include "StateCentricView/StateCentricViewSettings.h"
+#include "StateTreeSchema.h"
+
+#include "StateCentricViewPerUserSettings.generated.h"
+
+#define UE_API STATETREEEDITORMODULE_API
+
+DECLARE_MULTICAST_DELEGATE_OneParam(FOnStateCentricViewLOD_Changed, TOptional<EExtendableNodeLOD> /*InLOD*/);
+
+/**
+ * Experimental. Settings per user. This exists locally per user per schema. But also per state tree.
+ * Per state tree this is used to maintain consistent immediate behavior when clicking across states.
+ * Per user this is used to configure a default behavior when no per-state-tree setting has been set.
+ */
+USTRUCT(Experimental, MinimalAPI)
+struct FStateCentricViewPerUserSetting
+{
+	GENERATED_BODY()
+
+public:
+
+	friend UStateCentricViewPerUserSettings;
+
+private:
+
+	/** Amount of info to show for every parent state in the hierachy. */
+	UPROPERTY(EditAnywhere, Category="Layout", Config, meta = (ValidEnumValues = "Minimal, Medium, Full"))
+	TOptional<EExtendableNodeLOD> ParentHireachyLOD = {};
+
+	/** Amount of info to show for incoming transitions to this state. */
+	UPROPERTY(EditAnywhere, Category = "Layout", Config, meta = (ValidEnumValues = "Minimal, Full"))
+	TOptional<EExtendableNodeLOD> InTransitionViewLOD = {};
+
+	/** Amount of info to show for outgoing transitions from this state. */
+	UPROPERTY(EditAnywhere, Category = "Layout", Config, meta = (ValidEnumValues = "Minimal, Full"))
+	TOptional<EExtendableNodeLOD> OutTransitionLOD = {};
+
+	/** Number of parents to show if they exist. */
+	UPROPERTY(EditAnywhere, Category="Layout", Config)
+	TOptional<uint32> NumParentsShown = {};
+
+public:
+
+	UE_API TOptional<EExtendableNodeLOD> GetInTransitionLOD() const
+	{
+		return InTransitionViewLOD;
+	}
+
+	UE_API TOptional<EExtendableNodeLOD> GetOutTransitionLOD() const
+	{
+		return OutTransitionLOD;
+	}
+
+	UE_API void SetInTransitionLOD(TOptional<EExtendableNodeLOD> InLOD)
+	{
+		InTransitionViewLOD = InLOD;
+		OnInTransitionViewLOD_Changed.Broadcast(InTransitionViewLOD);
+	}
+
+	UE_API void SetOutTransitionLOD(TOptional<EExtendableNodeLOD> InLOD)
+	{
+		OutTransitionLOD = InLOD;
+		OnOutTransitionViewLOD_Changed.Broadcast(OutTransitionLOD);
+	}
+
+public:
+
+	/** Callback whenever the LOD for 'In' transition view has changed. */
+	FOnStateCentricViewLOD_Changed OnInTransitionViewLOD_Changed;
+
+	/** Callback whenever the LOD for 'Out' transition view has changed. */
+	FOnStateCentricViewLOD_Changed OnOutTransitionViewLOD_Changed;
+};
+
+
+/**
+ * Experimental. Settings for the state centric view per schema per user per schema.
+ * These settings provide defaults to reset to along with defaults for newly opened statetrees.
+ */
+UCLASS(Experimental, MinimalAPI, config = EditorPerProjectUserSettings)
+class UStateCentricViewPerUserSettings : public UDeveloperSettings
+{
+	GENERATED_BODY()
+
+public:
+
+	UStateCentricViewPerUserSettings();
+
+	UE_API static const UStateCentricViewPerUserSettings& Get()
+	{
+		return *GetDefault<UStateCentricViewPerUserSettings>();
+	}
+
+	UE_API static UStateCentricViewPerUserSettings& GetMutable()
+	{
+		return *GetMutableDefault<UStateCentricViewPerUserSettings>();
+	}
+
+public:
+
+	/** Per schema settings. Defined per schema since we have no need for an editor schema to be defined. */
+	UPROPERTY(EditAnywhere, Category="Appearance", Config, meta = (EditCondition = "StateTreeEditorModule.StateCentricViewSettings.IsStateCentricViewEnabled", EditConditionHides))
+	TMap<TSubclassOf<UStateTreeSchema>, FStateCentricViewPerUserSetting> PerSchemaSettings;
+
+protected:
+
+#if WITH_EDITOR
+	//~ Begin UDeveloperSettings Interface
+	UE_API virtual FText GetSectionText() const override;
+	UE_API virtual FText GetSectionDescription() const override;
+	//~ End UDeveloperSettings Interface
+#endif
+
+	//~ Begin UDeveloperSettings Interface
+	UE_API virtual FName GetCategoryName() const override;
+	//~ End UDeveloperSettings Interface
+};
+
+#undef UE_API
diff --git a/Engine/Plugins/Runtime/StateTree/Source/StateTreeEditorModule/Private/StateCentricView/StateCentricViewSettings.cpp b/Engine/Plugins/Runtime/StateTree/Source/StateTreeEditorModule/Private/StateCentricView/StateCentricViewSettings.cpp
new file mode 100644
--- /dev/null
+++ b/Engine/Plugins/Runtime/StateTree/Source/StateTreeEditorModule/Private/StateCentricView/StateCentricViewSettings.cpp
@@ -0,0 +1,61 @@
+// Copyright Epic Games, Inc. All Rights Reserved.
+
+#include "StateCentricView/StateCentricViewSettings.h"
+#include "StateTreeEditorData.h"
+#include "StateTreeViewModel.h"
+
+#include UE_INLINE_GENERATED_CPP_BY_NAME(StateCentricViewSettings)
+
+bool bEnableStateCentricView = false;
+FAutoConsoleVariableRef CVarEnableStateCentricView(
+	TEXT("StateTree.Editor.Experimental.EnableStateCentricView"),
+	bEnableStateCentricView,
+	TEXT("Set true to enable state centric view")
+);
+
+UStateCentricViewSettings::UStateCentricViewSettings()
+{
+	PerSchemaSettings.Emplace(UStateTreeSchema::StaticClass(), FStateCentricViewPerSchemaSetting());
+}
+
+bool UStateCentricViewSettings::IsStateCentricViewEnabled()
+{
+	return bEnableStateCentricView;
+}
+
+const FStateCentricViewPerSchemaSetting& UStateCentricViewSettings::GetSchemaSettingsDefaultFallback(const FStateTreeViewModel* ViewModel)
+{
+	if (ViewModel)
+	{
+		if (const UStateTree* StateTree = ViewModel->GetStateTree())
+		{
+			if (const UStateTreeSchema* Schema = StateTree->GetSchema())
+			{
+				const TSubclassOf<UStateTreeSchema> SchemaClass = Schema->GetClass();
+				return UStateCentricViewSettings::GetSchemaSettingsDefaultFallback(SchemaClass);
+			}
+		}
+	}
+
+	return UStateCentricViewSettings::Get().PerSchemaSettings[UStateTreeSchema::StaticClass()];
+}
+
+#if WITH_EDITOR
+
+FText UStateCentricViewSettings::GetSectionText() const
+{
+	return NSLOCTEXT("StateTreeEditor", "StateCentricViewSettingsName", "State Centric View");
+}
+
+FText UStateCentricViewSettings::GetSectionDescription() const
+{
+	return NSLOCTEXT("StateTreeEditor", "StateCentricViewSettingsDescription", "Experimental. Configure options for state centric view.");
+}
+
+#endif // WITH_EDITOR
+
+FName UStateCentricViewSettings::GetCategoryName() const
+{
+	// @TODO: Change me
+	return FName(TEXT("Plugins"));
+}
diff --git a/Engine/Plugins/Runtime/StateTree/Source/StateTreeEditorModule/Private/StateCentricView/StateCentricViewSettings.h b/Engine/Plugins/Runtime/StateTree/Source/StateTreeEditorModule/Private/StateCentricView/StateCentricViewSettings.h
new file mode 100644
--- /dev/null
+++ b/Engine/Plugins/Runtime/StateTree/Source/StateTreeEditorModule/Private/StateCentricView/StateCentricViewSettings.h
@@ -0,0 +1,128 @@
+// Copyright Epic Games, Inc. All Rights Reserved.
+
+#pragma once
+
+#include "StateTreeSchema.h"
+
+#include "Engine/DeveloperSettings.h"
+#include "Templates/SubclassOf.h"
+
+#include "StateCentricViewSettings.generated.h"
+
+#define UE_API STATETREEEDITORMODULE_API
+
+class FStateTreeViewModel;
+
+/**
+ * Experimental. Settings per schema. Mainly focused on styling & layout.
+ */
+USTRUCT(Experimental, MinimalAPI)
+struct FStateCentricViewPerSchemaSetting
+{
+	GENERATED_BODY()
+
+public:
+
+	/** 
+	 * True if state centric view colors based on hireachy rather than User Defined colors 
+	 * @TODO: Also this should be part of base ST.
+	 */
+	UPROPERTY(EditAnywhere, Category="Style")
+	bool bEnableHierarchyDrivenColors = false;
+
+	/** 
+	 * True if state centric view shapes nodes based on hireachy rather than default shapes 
+	 * @TODO: Commented out as not yet implemented. Also this should be part of base ST VM.
+	 */
+	//UPROPERTY(EditAnywhere, Category="Style")
+	//bool bEnableHireachyDrivenShapes = false;
+
+	/** True if central / leafmost view should have transitions justified instead of top aligned */
+	UPROPERTY(EditAnywhere, Category = "Layout")
+	bool bJustifyCentralViewTransitions = true;
+
+	/** True if state centric view allows for full detail expansion of parents, else max expansion is entry conditions */
+	UPROPERTY(EditAnywhere, Category="Layout")
+	bool bEnableFullParentDetailExpansion = true;
+
+	/** At the moment we only allow state centric to be toggled left or right of the outliner. */
+	UPROPERTY(EditAnywhere, Category="Layout")
+	bool bEnableOutlinerRightOfStateCentric = false;
+
+	/** True if transition nodes have fill based width. Enabled by default to avoid confusing amounts of random negative space. */
+	UPROPERTY(EditAnywhere, Category="Layout")
+	bool bEnableFillOnTransitionNodes = true;
+
+	/** 
+	 * Max number of parents we allow users to show. At the moment we do not encourage changing this as showing too many parents 
+	 * can detract from the intended workflow of focusing on the central state and jumping around if needed.
+	 */
+	UPROPERTY(EditAnywhere, Category="Layout", meta=(UIMin="1", UIMax="50"))
+	uint32 MaxNumParentsToShow = 4;
+};
+
+
+/**
+ * Experimental. May be removed at any time and API may change at any time.
+ * Toggle via: StateTree.Editor.Experimental.EnableStateCentricView
+ * 
+ * Settings for the state centric view per schema. As state centric is early in dev, there are many toggles.
+ */
+UCLASS(Experimental, MinimalAPI, config = Editor)
+class UStateCentricViewSettings : public UDeveloperSettings
+{
+	GENERATED_BODY()
+
+public:
+
+	UStateCentricViewSettings();
+
+	UE_API static const UStateCentricViewSettings& Get()
+	{
+		return *GetDefault<UStateCentricViewSettings>();
+	}
+
+	UE_API static UStateCentricViewSettings& GetMutable()
+	{
+		return *GetMutableDefault<UStateCentricViewSettings>();
+	}
+
+	UE_API static const FStateCentricViewPerSchemaSetting& GetSchemaSettingsDefaultFallback(const FStateTreeViewModel* ViewModel);
+
+	UE_API static const FStateCentricViewPerSchemaSetting& GetSchemaSettingsDefaultFallback(const TSubclassOf<UStateTreeSchema> InSchema)
+	{
+		if (const FStateCentricViewPerSchemaSetting* SchemaSettings = UStateCentricViewSettings::Get().PerSchemaSettings.Find(InSchema))
+		{
+			return *SchemaSettings;
+		}
+
+		return UStateCentricViewSettings::Get().PerSchemaSettings[UStateTreeSchema::StaticClass()];
+	}
+
+public:
+
+	/** Per schema settings. Defined per schema since we have no need for an editor schema to be defined. */
+	UPROPERTY(EditAnywhere, Category="Appearance", Config, meta = (EditCondition = "StateTreeEditorModule.StateCentricViewSettings.IsStateCentricViewEnabled", EditConditionHides))
+	TMap<TSubclassOf<UStateTreeSchema>, FStateCentricViewPerSchemaSetting> PerSchemaSettings;
+
+public:
+
+	/** Static version of state centric enabled check, exists to hide based on edit condition */
+	UFUNCTION()
+	static UE_API bool IsStateCentricViewEnabled(); 
+
+protected:
+
+#if WITH_EDITOR
+	//~ Begin UDeveloperSettings Interface
+	UE_API virtual FText GetSectionText() const override;
+	UE_API virtual FText GetSectionDescription() const override;
+	//~ End UDeveloperSettings Interface
+#endif
+
+	//~ Begin UDeveloperSettings Interface
+	UE_API virtual FName GetCategoryName() const override;
+	//~ End UDeveloperSettings Interface
+};
+
+#undef UE_API
diff --git a/Engine/Plugins/Runtime/StateTree/Source/StateTreeEditorModule/Private/StateCentricView/View/SMainStateTransitionExtensionNode.cpp b/Engine/Plugins/Runtime/StateTree/Source/StateTreeEditorModule/Private/StateCentricView/View/SMainStateTransitionExtensionNode.cpp
new file mode 100644
--- /dev/null
+++ b/Engine/Plugins/Runtime/StateTree/Source/StateTreeEditorModule/Private/StateCentricView/View/SMainStateTransitionExtensionNode.cpp
@@ -0,0 +1,429 @@
+// Copyright Epic Games, Inc. All Rights Reserved.
+
+#include "StateCentricView/View/SMainStateTransitionExtensionNode.h"
+
+#include "StateCentricView/ViewModel/StateCentricViewEditorDataExtension.h"
+#include "StateCentricView/StateCentricViewSettings.h"
+#include "StateCentricView/View/StateCentricViewUtils.h"
+#include "StateCentricView/View/STransitionNode.h"
+#include "StateTreeEditorStyle.h"
+#include "Styling/SlateBrush.h"
+#include "Styling/StyleColors.h"
+#include "Widgets/Images/SImage.h"
+#include "Widgets/Input/SButton.h"
+#include "Widgets/Layout/SBorder.h"
+#include "Widgets/Text/STextBlock.h"
+
+#define LOCTEXT_NAMESPACE "SMainStateTransitionExtensionNode"
+
+namespace UE::StateTree::Editor::StateCentricView
+{
+
+
+//////////////////////////////////////////////////////////////////////////
+// SMainStateTransitionExtensionNode
+
+SMainStateTransitionExtensionNode::~SMainStateTransitionExtensionNode()
+{
+	if (ViewModel)
+	{
+		if (UStateTreeEditorData* EditorData = const_cast<UStateTreeEditorData*>(ViewModel->GetStateTreeEditorData()))
+		{
+			if (UStateCentricViewEditorDataExtension* StateCentricEditorData = EditorData->GetExtension<UStateCentricViewEditorDataExtension>())
+			{
+				StateCentricEditorData->PerStateTreeUserSettings.OnInTransitionViewLOD_Changed.RemoveAll(this);
+				StateCentricEditorData->PerStateTreeUserSettings.OnOutTransitionViewLOD_Changed.RemoveAll(this);
+			}
+		}
+	}
+}
+
+void SMainStateTransitionExtensionNode::Construct_InitializeNodePersistentPreLOD()
+{
+	UStateTreeEditorData* EditorData = const_cast<UStateTreeEditorData*>(ViewModel->GetStateTreeEditorData());
+	UStateCentricViewEditorDataExtension* StateCentricEditorData = EditorData->GetOrCreateExtension<UStateCentricViewEditorDataExtension>();
+
+	// Remove any prior bindings in case of issues with lifetime
+	StateCentricEditorData->PerStateTreeUserSettings.OnInTransitionViewLOD_Changed.RemoveAll(this);
+	StateCentricEditorData->PerStateTreeUserSettings.OnOutTransitionViewLOD_Changed.RemoveAll(this);
+
+	TOptional<EExtendableNodeLOD> TransitionViewUserSetLOD;
+	if (TransitionDirection == EStateTransitionDirection::In)
+	{
+		StateCentricEditorData->PerStateTreeUserSettings.OnInTransitionViewLOD_Changed.AddSP(this, &SMainStateTransitionExtensionNode::HandleOnTransitionViewLOD_SettingsChanged);
+		TransitionViewUserSetLOD = StateCentricEditorData->PerStateTreeUserSettings.GetInTransitionLOD();
+	}
+	else if (ensure(TransitionDirection == EStateTransitionDirection::Out))
+	{
+		StateCentricEditorData->PerStateTreeUserSettings.OnOutTransitionViewLOD_Changed.AddSP(this, &SMainStateTransitionExtensionNode::HandleOnTransitionViewLOD_SettingsChanged);
+		TransitionViewUserSetLOD = StateCentricEditorData->PerStateTreeUserSettings.GetOutTransitionLOD();
+	}
+
+	// We were given a LOD to use, but based on user config we may want to use a different LOD for now
+	if (TransitionViewUserSetLOD)
+	{
+		CurrentLOD = TransitionViewUserSetLOD.GetValue();
+	}
+}
+
+void SMainStateTransitionExtensionNode::Construct_GenerateBuiltInExtensions(const EExtendableNodeLOD InLOD)
+{
+	// Entire widget state color outline in overlay extension depending on LOD
+	if (InLOD != EExtendableNodeLOD::Hidden && InLOD != EExtendableNodeLOD::Collapsed)
+	{
+		StateOutlineBrush = FStateCentricViewUtils::MakeStateOutline(ViewModel.ToSharedRef(), ViewState.Get());
+
+		TSharedPtr<SWidget> StateColorOutline = SNew(SBorder)
+			.BorderImage(StateOutlineBrush->GetSlateBrush())
+			.Visibility(EVisibility::HitTestInvisible);
+
+		AddExtension(StateColorOutline.ToSharedRef())
+			.SetLocation(EExtensionLocation::Foreground)
+			.SetPadding(FMargin(0.0, 0.0, 0.0, 0.0));
+	}
+
+	// Muted background if hidden
+	if (InLOD == EExtendableNodeLOD::Hidden)
+	{
+		StateBackgroundBrush = FStateCentricViewUtils::MakeStateCollapsedBackground(ViewModel.ToSharedRef(), ViewState.Get());
+
+		TSharedPtr<SWidget> StateColorBackground = SNew(SBorder)
+			.BorderImage(StateBackgroundBrush->GetSlateBrush())
+			.Visibility(EVisibility::HitTestInvisible);
+
+		double PaddingBackgroundSize = 16.0;
+
+		FMargin StateBackgroundPadding;
+		if (TransitionDirection == EStateTransitionDirection::In)
+		{
+			StateBackgroundPadding = FMargin(-2.0 + PaddingBackgroundSize, 4.0, 2.0, 4.0);
+		}
+		else if (ensure(TransitionDirection == EStateTransitionDirection::Out))
+		{
+			StateBackgroundPadding = FMargin(2.0, 4.0, -2.0 + PaddingBackgroundSize, 4.0);
+		}
+
+		AddExtension(StateColorBackground.ToSharedRef())
+			.SetLocation(EExtensionLocation::Background)
+			.SetPadding(StateBackgroundPadding);
+	}
+}
+
+TSharedRef<SWidget> SMainStateTransitionExtensionNode::Construct_GenerateNodeHeader(const EExtendableNodeLOD InLOD)
+{
+	if (ViewState.IsValid())
+	{
+		UStateTreeEditorData* EditorData = const_cast<UStateTreeEditorData*>(ViewModel->GetStateTreeEditorData());
+		UStateCentricViewEditorDataExtension* StateCentricEditorData = EditorData->GetOrCreateExtension<UStateCentricViewEditorDataExtension>();
+
+		UpdateHideButtonAttributes();
+
+		SAssignNew(HideButton, SButton)
+			.OnClicked(this, &SMainStateTransitionExtensionNode::HandleOnCollapseClicked)
+			.ButtonStyle(&FAppStyle::GetWidgetStyle<FButtonStyle>("NoBorder"))
+			.ContentPadding(FMargin())
+			[
+				SNew(SImage)
+				.Image(&FAppStyle::Get().GetWidgetStyle<FExpandableAreaStyle>("ExpandableArea").CollapsedImage)
+				.RenderTransform(this, &SMainStateTransitionExtensionNode::GetHideButtonRenderTransform)
+				.ColorAndOpacity(this, &SMainStateTransitionExtensionNode::GetHideButtonColor)
+			];
+
+		// Header is just a button when transitions hidden
+		if (InLOD == EExtendableNodeLOD::Hidden)
+		{
+			return SNew(SOverlay)
+				+ SOverlay::Slot()
+				[
+					// Setting background brush 32 px as it must be able to contain the child button + padding
+					// Otherwise splitter layout calculation will be uneven. We rely on fixed 32px size to be in sync with parents when hidden.
+					SNew(SImage)
+					.Image(FAppStyle::Get().GetBrush("Brushes.Recessed"))
+					.ColorAndOpacity(FStyleColors::Transparent)
+					.DesiredSizeOverride(FVector2D(32.0f, 28.0f))
+				]
+				+ SOverlay::Slot()
+				.HAlign(HideButtonAlignment)
+				.VAlign(VAlign_Center)
+				.Padding(MakeAttributeSP(this, &SMainStateTransitionExtensionNode::GetHideButtonPadding))
+				[
+					HideButton.ToSharedRef()
+				];
+		}
+
+		TArray<FStateTreeTransition*> StateTransitions;
+		EHorizontalAlignment HeaderAlignment = HAlign_Left;
+		FMargin HeaderPadding;
+		FText HeaderText;
+		if (TransitionDirection == EStateTransitionDirection::In)
+		{
+			StateTransitions = ViewModel->GetStateInTransitions(ViewState.Get());
+			HeaderAlignment = HAlign_Left;
+			HeaderPadding = FMargin(12.0, 0.0, 2.0, 0.0);
+			HeaderText = LOCTEXT("TransitionHeaderIn", "IN ({0})");
+		}
+		else if (ensure(TransitionDirection == EStateTransitionDirection::Out))
+		{
+			StateTransitions = ViewModel->GetStateOutTransitions(ViewState.Get());
+			HeaderAlignment = HAlign_Right;
+			HeaderPadding = FMargin(2.0, 0.0, 12.0, 0.0);
+			HeaderText = LOCTEXT("TransitionHeaderOut", "OUT ({0})");
+		}
+
+		HeaderText = FText::Format(HeaderText, StateTransitions.Num());
+
+		StateHeaderBackgroundBrush = FStateCentricViewUtils::MakeStateHeaderBackground(ViewModel.ToSharedRef(), ViewState.Get(), FVector2f(32.0f, 28.0f));
+
+		TSharedPtr<SHorizontalBox> HeaderContentsBox;
+
+		TSharedRef<SOverlay> HeaderWidget = SNew(SOverlay)
+			+ SOverlay::Slot()
+			[
+				SNew(SImage)
+				.Image(StateHeaderBackgroundBrush->GetSlateBrush())
+			]
+			+ SOverlay::Slot()
+			[
+				SAssignNew(HeaderContentsBox, SHorizontalBox)
+				+ SHorizontalBox::Slot()
+				.HAlign(HeaderAlignment)
+				.VAlign(VAlign_Center)
+				.FillWidth(1.0f)
+				.Padding(HeaderPadding)
+				[
+					SNew(STextBlock)
+					.Text(HeaderText)
+					.TextStyle(&FStateTreeEditorStyle::Get().GetWidgetStyle<FTextBlockStyle>("Normal.Normal"))
+					.ColorAndOpacity(FStyleColors::White)
+				]
+			];
+
+		// In full LOD, we add a button to allow for new transitions to be added
+		if (InLOD == EExtendableNodeLOD::Full)
+		{
+			HeaderContentsBox->AddSlot()
+			.HAlign(HAlign_Right)
+			.VAlign(VAlign_Center)
+			.AutoWidth()
+			.Padding(FMargin(-8.0, 0.0, 8.0, 0.0))
+			[
+				SAssignNew(AddTransitionButton, SButton)
+				.OnClicked(this, &SMainStateTransitionExtensionNode::HandleOnAddTransitionClicked)
+				.ButtonStyle(&FAppStyle::GetWidgetStyle<FButtonStyle>("NoBorder"))
+				.ContentPadding(FMargin())
+				[
+					SNew(SImage)
+					.Image(FAppStyle::GetBrush("Icons.PlusCircle"))
+					.ColorAndOpacity(this, &SMainStateTransitionExtensionNode::GetAddTransitionColor)
+				]
+			];
+		}
+
+		HeaderWidget->AddSlot()
+			.HAlign(HideButtonAlignment)
+			.VAlign(VAlign_Center)
+			.Padding(MakeAttributeSP(this, &SMainStateTransitionExtensionNode::GetHideButtonPadding))
+			[
+				HideButton.ToSharedRef()
+			];
+
+		return HeaderWidget;
+	}
+
+	return SNullWidget::NullWidget;
+}
+
+TSharedRef<SWidget> SMainStateTransitionExtensionNode::Construct_GenerateNodeBody(const EExtendableNodeLOD InLOD)
+{
+	if (ViewState.IsValid())
+	{
+		TSubclassOf<UStateTreeSchema> SchemaClass = ViewModel->GetStateTree()->GetSchema()->GetClass();
+
+		if (UStateCentricViewSettings::GetSchemaSettingsDefaultFallback(SchemaClass).bJustifyCentralViewTransitions)
+		{
+			NodeBodySizeParam = FStretchContent(1.0f);
+		}
+
+		return FStateCentricViewUtils::GenerateTransitionsViewWidget(InLOD, ViewModel.ToSharedRef(), ViewState.Get(), TransitionDirection);
+	}
+
+	return SNullWidget::NullWidget;
+}
+
+FReply SMainStateTransitionExtensionNode::HandleOnCollapseClicked()
+{
+	if (ViewState.IsValid())
+	{
+		UStateTreeEditorData* EditorData = const_cast<UStateTreeEditorData*>(ViewModel->GetStateTreeEditorData());
+		UStateCentricViewEditorDataExtension* StateCentricEditorData = EditorData->GetOrCreateExtension<UStateCentricViewEditorDataExtension>();
+
+		if (TransitionDirection == EStateTransitionDirection::In)
+		{
+			bool bHidden = StateCentricEditorData->PerStateTreeUserSettings.GetInTransitionLOD().Get(EExtendableNodeLOD::Default) == EExtendableNodeLOD::Hidden;
+
+			if (bHidden)
+			{
+				StateCentricEditorData->PerStateTreeUserSettings.SetInTransitionLOD({});
+				StateCentricEditorData->InTransitionSplitterSizeRule = SSplitter::ESizeRule::FractionOfParent;
+			}
+			else
+			{
+				StateCentricEditorData->PerStateTreeUserSettings.SetInTransitionLOD(EExtendableNodeLOD::Hidden);
+				StateCentricEditorData->InTransitionSplitterSizeRule = SSplitter::ESizeRule::SizeToContent;
+			}
+		}
+		else if (ensure(TransitionDirection == EStateTransitionDirection::Out))
+		{
+			bool bHidden = StateCentricEditorData->PerStateTreeUserSettings.GetOutTransitionLOD().Get(EExtendableNodeLOD::Default) == EExtendableNodeLOD::Hidden;
+
+			if (bHidden)
+			{
+				StateCentricEditorData->PerStateTreeUserSettings.SetOutTransitionLOD({});
+				StateCentricEditorData->OutTransitionSplitterSizeRule = SSplitter::ESizeRule::FractionOfParent;
+			}
+			else
+			{
+				StateCentricEditorData->PerStateTreeUserSettings.SetOutTransitionLOD(EExtendableNodeLOD::Hidden);
+				StateCentricEditorData->OutTransitionSplitterSizeRule = SSplitter::ESizeRule::SizeToContent;
+			}
+		}
+	}
+
+	return FReply::Handled();
+}
+
+FReply SMainStateTransitionExtensionNode::HandleOnAddTransitionClicked()
+{
+	return FReply::Handled();
+}
+
+void SMainStateTransitionExtensionNode::HandleOnTransitionViewLOD_SettingsChanged(TOptional<EExtendableNodeLOD> InLOD)
+{
+	UpdateHideButtonAttributes();
+
+	if (ViewState.IsValid())
+	{
+		UStateTreeEditorData* EditorData = const_cast<UStateTreeEditorData*>(ViewModel->GetStateTreeEditorData());
+		UStateCentricViewEditorDataExtension* StateCentricEditorData = EditorData->GetOrCreateExtension<UStateCentricViewEditorDataExtension>();
+
+		TOptional<EExtendableNodeLOD> TransitionViewLOD;
+		if (TransitionDirection == EStateTransitionDirection::In)
+		{
+			TransitionViewLOD = StateCentricEditorData->PerStateTreeUserSettings.GetInTransitionLOD();
+		}
+		else if (ensure(TransitionDirection == EStateTransitionDirection::Out))
+		{
+			TransitionViewLOD = StateCentricEditorData->PerStateTreeUserSettings.GetOutTransitionLOD();
+		}
+
+		if (TransitionViewLOD)
+		{
+			RegenerateWidgetForLOD(TransitionViewLOD.GetValue());
+		}
+		else
+		{
+			// @TODO: May be better to use a PrevLOD. But this can get complicated so for later
+			RegenerateWidgetForLOD(Config.InitialLOD);
+		}
+	}
+}
+
+FSlateColor SMainStateTransitionExtensionNode::GetAddTransitionColor() const
+{
+	return AddTransitionButton->IsHovered() ? FStyleColors::White : FStyleColors::Foreground;
+}
+
+FSlateColor SMainStateTransitionExtensionNode::GetHideButtonColor() const
+{
+	return HideButton->IsHovered() ? FStyleColors::White : FStyleColors::Foreground;
+}
+
+void SMainStateTransitionExtensionNode::UpdateHideButtonAttributes()
+{
+	if (ViewState.IsValid())
+	{
+		UStateTreeEditorData* EditorData = const_cast<UStateTreeEditorData*>(ViewModel->GetStateTreeEditorData());
+		UStateCentricViewEditorDataExtension* StateCentricEditorData = EditorData->GetOrCreateExtension<UStateCentricViewEditorDataExtension>();
+
+		double IconWidth = 16.0;
+		if (TransitionDirection == EStateTransitionDirection::In)
+		{
+			bool bHidden = StateCentricEditorData->PerStateTreeUserSettings.GetInTransitionLOD().Get(EExtendableNodeLOD::Default) == EExtendableNodeLOD::Hidden;
+			HideButtonAlignment = HAlign_Right;
+			HideButtonPadding = bHidden
+				? FMargin(0.0, 0.0, 2.0 - IconWidth, 0.0)
+				: FMargin(0.0, 0.0, 2.0, 0.0);
+			HideButtonRenderTransform = bHidden
+				? FSlateRenderTransform(FScale2D(-1, 1))
+				: FSlateRenderTransform();
+		}
+		else if (ensure(TransitionDirection == EStateTransitionDirection::Out))
+		{
+			bool bHidden = StateCentricEditorData->PerStateTreeUserSettings.GetOutTransitionLOD().Get(EExtendableNodeLOD::Default) == EExtendableNodeLOD::Hidden;
+			HideButtonAlignment = HAlign_Left;
+			HideButtonPadding = bHidden
+				? FMargin(2.0, 0.0, 0.0, 0.0)
+				: FMargin(2.0 + IconWidth, 0.0, 0.0, 0.0);
+			HideButtonRenderTransform = bHidden
+				? FSlateRenderTransform()
+				: FSlateRenderTransform(FScale2D(-1, 1));
+		}
+	}
+}
+
+TOptional<FSlateRenderTransform> SMainStateTransitionExtensionNode::GetHideButtonRenderTransform() const
+{
+	return HideButtonRenderTransform;
+}
+
+FMargin SMainStateTransitionExtensionNode::GetHideButtonPadding() const
+{
+	return HideButtonPadding;
+}
+
+EHorizontalAlignment SMainStateTransitionExtensionNode::GetHideButtonAlignment() const
+{
+	return HideButtonAlignment;
+}
+
+
+//////////////////////////////////////////////////////////////////////////
+// SMainStateInTransitionExtensionNode
+
+const FName SMainStateInTransitionExtensionNode::NodeID(TEXT("StateCentricView.MainState.InTransitions"));
+
+FName SMainStateInTransitionExtensionNode::GetNodeName() const
+{
+	return NodeID;
+}
+
+void SMainStateInTransitionExtensionNode::Construct_InitializeNodePersistentPreLOD()
+{
+	TransitionDirection = EStateTransitionDirection::In;
+	SMainStateTransitionExtensionNode::Construct_InitializeNodePersistentPreLOD();
+}
+
+
+//////////////////////////////////////////////////////////////////////////
+// SMainStateOutTransitionExtensionNode
+
+const FName SMainStateOutTransitionExtensionNode::NodeID(TEXT("StateCentricView.MainState.OutTransitions"));
+
+FName SMainStateOutTransitionExtensionNode::GetNodeName() const
+{
+	return NodeID;
+}
+
+void SMainStateOutTransitionExtensionNode::Construct_InitializeNodePersistentPreLOD()
+{
+	TransitionDirection = EStateTransitionDirection::Out;
+	SMainStateTransitionExtensionNode::Construct_InitializeNodePersistentPreLOD();
+}
+
+
+//////////////////////////////////////////////////////////////////////////
+
+} // namespace UE::StateTree::Editor::StateCentricView
+
+#undef LOCTEXT_NAMESPACE
+
diff --git a/Engine/Plugins/Runtime/StateTree/Source/StateTreeEditorModule/Private/StateCentricView/View/SMainStateTransitionExtensionNode.h b/Engine/Plugins/Runtime/StateTree/Source/StateTreeEditorModule/Private/StateCentricView/View/SMainStateTransitionExtensionNode.h
new file mode 100644
--- /dev/null
+++ b/Engine/Plugins/Runtime/StateTree/Source/StateTreeEditorModule/Private/StateCentricView/View/SMainStateTransitionExtensionNode.h
@@ -0,0 +1,125 @@
+// Copyright Epic Games, Inc. All Rights Reserved.
+
+#pragma once
+
+#include "CoreMinimal.h"
+
+#include "StateCentricView/View/SStateTreeExtendableNode.h"
+
+struct FSlateBrush;
+
+class SButton;
+
+namespace UE::StateTree::Editor::StateCentricView
+{
+
+
+/**
+ * Experimental. Node used as an extension for In / Out transitions on the main state.
+ */
+class SMainStateTransitionExtensionNode : public SStateTreeExtendableNode
+{
+public:
+
+	virtual ~SMainStateTransitionExtensionNode() override;
+
+protected:
+
+	//~ Begin SExtendableNode Interface
+	virtual void Construct_InitializeNodePersistentPreLOD() override;
+	virtual void Construct_GenerateBuiltInExtensions(const EExtendableNodeLOD InLOD) override;
+	virtual TSharedRef<SWidget> Construct_GenerateNodeHeader(const EExtendableNodeLOD InLOD) override;
+	virtual TSharedRef<SWidget> Construct_GenerateNodeBody(const EExtendableNodeLOD InLOD) override;
+	//~ End SExtendableNode Interface
+
+protected:
+
+	/** If we are showing in / out transitions. Init in child classes. */
+	EStateTransitionDirection TransitionDirection = EStateTransitionDirection::Out;
+
+private:
+
+	/** Callback for collapse button*/
+	FReply HandleOnCollapseClicked();
+
+	/** Callback for top right button */
+	FReply HandleOnAddTransitionClicked();
+
+	/** Callback for when LOD settings change */
+	void HandleOnTransitionViewLOD_SettingsChanged(TOptional<EExtendableNodeLOD> InLOD);
+
+	/** Minor UI / UX callbacks */
+	FSlateColor GetAddTransitionColor() const;
+	FSlateColor GetHideButtonColor() const;
+	
+	void UpdateHideButtonAttributes();
+	FMargin GetHideButtonPadding() const;
+	EHorizontalAlignment GetHideButtonAlignment() const;
+	TOptional<FSlateRenderTransform> GetHideButtonRenderTransform() const;
+
+private:
+
+	TSharedPtr<FDeferredCleanupSlateBrush> StateOutlineBrush;
+	TSharedPtr<FDeferredCleanupSlateBrush> StateBackgroundBrush;
+	TSharedPtr<FDeferredCleanupSlateBrush> StateHeaderBackgroundBrush;
+
+	TSharedPtr<SButton> AddTransitionButton;
+	TSharedPtr<SButton> HideButton;
+
+	FMargin HideButtonPadding = FMargin();
+	EHorizontalAlignment HideButtonAlignment = HAlign_Fill;
+	FSlateRenderTransform HideButtonRenderTransform = FSlateRenderTransform();
+};
+
+
+/**
+ * Experimental. Specialization of SMainStateTransitionExtensionNode for 'In' transitions
+ */
+class SMainStateInTransitionExtensionNode : public SMainStateTransitionExtensionNode
+{
+
+public:
+
+	static const FName NodeID;
+
+public:
+
+	//~ Begin SExtendableNode Interface
+	virtual FName GetNodeName() const override;
+	//~ End SExtendableNode Interface
+
+protected:
+
+	//~ Begin SExtendableNode Interface
+	virtual void Construct_InitializeNodePersistentPreLOD() override;
+	//~ End SExtendableNode Interface
+
+};
+
+
+/**
+ * Experimental. Specialization of SMainStateTransitionExtensionNode for 'Out' transitions
+ */
+class SMainStateOutTransitionExtensionNode : public SMainStateTransitionExtensionNode
+{
+
+public:
+
+	static const FName NodeID;
+
+public:
+
+	//~ Begin SExtendableNode Interface
+	virtual FName GetNodeName() const override;
+	//~ End SExtendableNode Interface
+
+protected:
+
+	//~ Begin SExtendableNode Interface
+	virtual void Construct_InitializeNodePersistentPreLOD() override;
+	//~ End SExtendableNode Interface
+
+};
+
+
+} // namespace UE::StateTree::Editor::StateCentricView
\ No newline at end of file
diff --git a/Engine/Plugins/Runtime/StateTree/Source/StateTreeEditorModule/Private/StateCentricView/View/SParentStateExtensionNode.cpp b/Engine/Plugins/Runtime/StateTree/Source/StateTreeEditorModule/Private/StateCentricView/View/SParentStateExtensionNode.cpp
new file mode 100644
--- /dev/null
+++ b/Engine/Plugins/Runtime/StateTree/Source/StateTreeEditorModule/Private/StateCentricView/View/SParentStateExtensionNode.cpp
@@ -0,0 +1,321 @@
+// Copyright Epic Games, Inc. All Rights Reserved.
+
+#include "StateCentricView/View/SParentStateExtensionNode.h"
+
+#include "StateCentricView/ViewModel/StateCentricViewEditorDataExtension.h"
+#include "StateCentricView/View/SParentStateTransitionExtensionNode.h"
+#include "StateCentricView/View/StateCentricViewUtils.h"
+#include "StateTreeEditorStyle.h"
+
+#include "IDetailsView.h"
+#include "DetailsViewArgs.h"
+#include "Layout/Visibility.h"
+#include "Modules/ModuleManager.h"
+#include "ObjectEditorUtils.h"
+#include "PropertyEditorModule.h"
+#include "Styling/SlateBrush.h"
+#include "Styling/StyleColors.h"
+#include "Widgets/Images/SImage.h"
+#include "Widgets/Input/SButton.h"
+#include "Widgets/Layout/SBorder.h"
+#include "Widgets/Text/STextBlock.h"
+
+#define LOCTEXT_NAMESPACE "SStateTreeExtendableNode"
+
+namespace UE::StateTree::Editor::StateCentricView
+{
+
+const FName SParentStateExtensionNode::NodeID(TEXT("StateCentricView.ParentState.Node"));
+
+FName SParentStateExtensionNode::GetNodeName() const
+{
+	return NodeID;
+}
+
+void SParentStateExtensionNode::Construct_GenerateBuiltInExtensions(const EExtendableNodeLOD InLOD)
+{
+	ChildStateBrushes.Reset();
+
+	if (ViewState.IsValid())
+	{
+		// @TODO: Remove splitter references here once those become a feature toggle on cardinal extensions
+		UStateTreeEditorData* EditorData = const_cast<UStateTreeEditorData*>(ViewModel->GetStateTreeEditorData());
+		UStateCentricViewEditorDataExtension* StateCentricEditorData = EditorData->GetOrCreateExtension<UStateCentricViewEditorDataExtension>();
+		
+		// In Transition extension on left
+		{
+			FScopedMainNodeLayoutInfo ScopedMainLayoutInfo = SetScopedMainNodeLayoutInfo()
+				.SetSplitterValue(MakeAttributeUObject(StateCentricEditorData, &UStateCentricViewEditorDataExtension::HandleGetMainNodeSplitterValue))
+				.SetSplitterOnSlotResized(SSplitter::FOnSlotResized::CreateUObject(StateCentricEditorData, &UStateCentricViewEditorDataExtension::HandleOnMainNodeSplitterResized));
+
+			// @TODO: Read Init LOD from project settings for in transitions.
+			EExtendableNodeLOD TransitionsInLOD = InLOD == EExtendableNodeLOD::Collapsed ? EExtendableNodeLOD::Collapsed : EExtendableNodeLOD::Minimal;
+			TSharedRef<SWidget> InTransitionLeftExtension = SNew(SParentStateInTransitionExtensionNode, ViewModel.ToSharedRef(), ViewState.Get())
+				.Config({ .InitialLOD = TransitionsInLOD });
+
+			AddExtension(InTransitionLeftExtension)
+				.SetLocation(EExtensionLocation::SplitLeft)
+				.SetSplitterSizeRule(MakeAttributeUObject(StateCentricEditorData, &UStateCentricViewEditorDataExtension::GetInTransitionSplitterSizeRule))
+				.SetSplitterValue(MakeAttributeUObject(StateCentricEditorData, &UStateCentricViewEditorDataExtension::HandleGetInTransitionSplitterValue))
+				.SetSplitterOnSlotResized(SSplitter::FOnSlotResized::CreateUObject(StateCentricEditorData, &UStateCentricViewEditorDataExtension::HandleOnInTransitionSplitterResized));
+
+			// Add spacer should we hit max size caps, Needed since fill + max size = left align (Slate bug).
+			AddExtension(
+				SNew(SImage)
+					.Image(FAppStyle::Get().GetBrush("Brushes.Recessed"))
+					.DesiredSizeOverride(FVector2D(32.0, 8.0)))
+				.SetLocation(EExtensionLocation::SplitLeft)
+				.SetSplitterValue(MakeAttributeUObject(StateCentricEditorData, &UStateCentricViewEditorDataExtension::HandleGetLeftSpacerSplitterValue))
+				.SetSplitterOnSlotResized(SSplitter::FOnSlotResized::CreateUObject(StateCentricEditorData, &UStateCentricViewEditorDataExtension::HandleOnLeftSpacerSplitterResized));
+		}
+
+		// Out Transition extension on right
+		{
+			// @TODO: Read Init LOD from project settings for out transitions.
+			EExtendableNodeLOD TransitionsOutLOD = InLOD;
+			TSharedRef<SWidget> OutTransitionRightExtension = SNew(SParentStateOutTransitionExtensionNode, ViewModel.ToSharedRef(), ViewState.Get())
+				.Config({ .InitialLOD = TransitionsOutLOD });
+			
+			AddExtension(OutTransitionRightExtension)
+				.SetLocation(EExtensionLocation::SplitRight)
+				.SetSplitterSizeRule(MakeAttributeUObject(StateCentricEditorData, &UStateCentricViewEditorDataExtension::GetOutTransitionSplitterSizeRule))
+				.SetSplitterValue(MakeAttributeUObject(StateCentricEditorData, &UStateCentricViewEditorDataExtension::HandleGetOutTransitionSplitterValue))
+				.SetSplitterOnSlotResized(SSplitter::FOnSlotResized::CreateUObject(StateCentricEditorData, &UStateCentricViewEditorDataExtension::HandleOnOutTransitionSplitterResized));
+
+			// Add spacer should we hit max size caps, Needed since fill + max size = left align (Slate bug).
+			AddExtension(
+				SNew(SImage)
+					.Image(FAppStyle::Get().GetBrush("Brushes.Recessed"))
+					.DesiredSizeOverride(FVector2D(32.0, 8.0)))
+				.SetLocation(EExtensionLocation::SplitRight)
+				.SetSplitterValue(MakeAttributeUObject(StateCentricEditorData, &UStateCentricViewEditorDataExtension::HandleGetRightSpacerSplitterValue))
+				.SetSplitterOnSlotResized(SSplitter::FOnSlotResized::CreateUObject(StateCentricEditorData, &UStateCentricViewEditorDataExtension::HandleOnRightSpacerSplitterResized));
+		}
+
+		// Child states in footnote depending on LOD
+		FStateCentricViewUtils::AddChildStateFootnoteExtensions(InLOD, this, OUT ChildStateBrushes);
+
+		// Also small bottom rounded rectangle if full LOD
+		if (InLOD == EExtendableNodeLOD::Full)
+		{
+			StateFootnoteBackgroundBrush = FStateCentricViewUtils::MakeStateFootnoteBackground(ViewModel.ToSharedRef(), ViewState.Get());
+
+			AddExtension(
+				SNew(SImage)
+					.Image(FAppStyle::Get().GetBrush("Brushes.Recessed"))
+					.DesiredSizeOverride(FVector2D(32.0, 8.0)))
+				.SetLocation(EExtensionLocation::Footnote)
+				.SetPadding(FMargin(0, 0, 0, 0));
+		}
+
+		// Entire widget state color outline in foreground extension, half circle if not full LOD
+		{
+			// Find what depth of a parent we are to find an overlay layer that will overlap correctly
+			int32 ParentDepth = 0;
+			TArray<UStateTreeState*> CentralStateParents = ViewModel->GetStateParents(ViewModel->GetViewState());
+			for (UStateTreeState* Parent : CentralStateParents)
+			{
+				ParentDepth++;
+				if (Parent->ID == ViewState->ID)
+				{
+					break;
+				}
+			}
+
+			int32 OverlayOffset = (CentralStateParents.Num() - ParentDepth) * 5000;
+
+			FScopedMainNodeLayoutInfo ScopedMainLayoutInfo = SetScopedMainNodeLayoutInfo()
+				.SetOverlayOffset(OverlayOffset);
+
+			if (InLOD == EExtendableNodeLOD::Full)
+			{
+				StateOutlineBrush = FStateCentricViewUtils::MakeStateOutlineMuted(ViewModel.ToSharedRef(), ViewState.Get());
+			}
+			else if (InLOD == EExtendableNodeLOD::Minimal)
+			{
+				StateOutlineBrush = FStateCentricViewUtils::MakeStateHalfTopOutlineMuted(ViewModel.ToSharedRef(), ViewState.Get());
+			}
+			else if (ensure(InLOD == EExtendableNodeLOD::Collapsed))
+			{
+				StateOutlineBrush = FStateCentricViewUtils::MakeStateHalfTopOutlineCollapsed(ViewModel.ToSharedRef(), ViewState.Get());
+			}
+
+			TSharedPtr<SWidget> StateColorOutline = SNew(SBorder)
+				.BorderImage(StateOutlineBrush->GetSlateBrush())
+				.Visibility(EVisibility::HitTestInvisible);
+
+			if (InLOD == EExtendableNodeLOD::Full)
+			{
+				AddExtension(StateColorOutline.ToSharedRef())
+					.SetLocation(EExtensionLocation::Foreground);
+			}
+			else if (ensure(InLOD == EExtendableNodeLOD::Collapsed || InLOD == EExtendableNodeLOD::Minimal))
+			{
+				// Negative pad based on size guess so outline edges reach to bottom of widget / beind next widget even with corner culling.
+				EExtendableNodeLOD TransitionsInLOD = InLOD == EExtendableNodeLOD::Collapsed ? EExtendableNodeLOD::Collapsed : EExtendableNodeLOD::Minimal;
+				EExtendableNodeLOD TransitionsOutLOD = InLOD;
+
+				TArray<FStateTreeTransition*> TransitionsIn = ViewModel->GetStateInTransitions(ViewState.Get());
+				TArray<FStateTreeTransition*> TransitionsOut = ViewModel->GetStateOutTransitions(ViewState.Get());
+
+				int32 TransitionsInSizeEstimate = TransitionsInLOD == EExtendableNodeLOD::Collapsed ? -100 : TransitionsIn.Num() * -150;
+				int32 TransitionsOutSizeEstimate = TransitionsOutLOD == EExtendableNodeLOD::Collapsed ? -100 : TransitionsOut.Num() * -150;
+
+				AddExtension(StateColorOutline.ToSharedRef())
+					.SetLocation(EExtensionLocation::Foreground)
+					.SetPadding(FMargin(0.0, 0.0, 0.0, -200.0 + FMath::Min(TransitionsInSizeEstimate, TransitionsOutSizeEstimate)));
+			}
+		}
+
+	}
+}
+
+TSharedRef<SWidget> SParentStateExtensionNode::Construct_GenerateNodeHeader(const EExtendableNodeLOD InLOD)
+{
+	if (ViewState.IsValid())
+	{
+		// @TODO: It should also write to project settings for LOD / only do that & we respond in this widget
+
+		StateHeaderBackgroundBrush = FStateCentricViewUtils::MakeStateHeaderBackground(ViewModel.ToSharedRef(), ViewState.Get());
+
+		const FSlateBrush* ExpandButtonBrush = nullptr;
+		if (InLOD == EExtendableNodeLOD::Collapsed)
+		{
+			ExpandButtonBrush = &FAppStyle::Get().GetWidgetStyle<FExpandableAreaStyle>("ExpandableArea").CollapsedImage;
+			StateHeaderBackgroundBrush = FStateCentricViewUtils::MakeStateHeaderBackgroundCollapsed(ViewModel.ToSharedRef(), ViewState.Get());
+		}
+		else if (InLOD == EExtendableNodeLOD::Minimal)
+		{
+			ExpandButtonBrush = &FAppStyle::Get().GetWidgetStyle<FExpandableAreaStyle>("ExpandableArea").ExpandedImage;
+		}
+		else
+		{
+			ensure(InLOD == EExtendableNodeLOD::Full);
+			ExpandButtonBrush = FAppStyle::Get().GetBrush("Icons.Details");
+		}
+		
+		return SNew(SOverlay)
+			+ SOverlay::Slot()
+			[
+				SNew(SImage)
+				.Image(StateHeaderBackgroundBrush->GetSlateBrush())
+			]
+			+ SOverlay::Slot()
+			.Padding(FMargin(8.0, 2.0, 32.0, 2.0))
+			[
+				FStateCentricViewUtils::GenerateGotoStateButton(InLOD, ViewModel.ToSharedRef(), ViewState.Get())
+			]
+			+ SOverlay::Slot()
+			[
+				SNew(SHorizontalBox)
+				+ SHorizontalBox::Slot()
+				.HAlign(HAlign_Left)
+				.VAlign(VAlign_Center)
+				.Padding(FMargin(12.0, 0.0, 2.0, 0.0))
+				.AutoWidth()
+				[
+					SNew(SImage)
+					.Image(ViewModel->GetSelectorIcon(ViewState.Get()))
+					.ColorAndOpacity(FStyleColors::Foreground)
+					.Visibility(EVisibility::HitTestInvisible)
+				]
+				+ SHorizontalBox::Slot()
+				.HAlign(HAlign_Left)
+				.VAlign(VAlign_Center)
+				.AutoWidth()
+				[
+					SNew(STextBlock)
+					.Text(FText::FromName(ViewState->Name))
+					.TextStyle(&FStateTreeEditorStyle::Get().GetWidgetStyle<FTextBlockStyle>("StateTree.CentricView.StateTitle.Regular"))
+					.ColorAndOpacity(FStyleColors::Foreground)
+					.Visibility(EVisibility::HitTestInvisible)
+				]
+				+ SHorizontalBox::Slot()
+				.HAlign(HAlign_Right)
+				.VAlign(VAlign_Center)
+				.FillWidth(1.0f)
+				.Padding(FMargin(0.0, 0.0, 10.0, 0.0))
+				[
+					SAssignNew(ChangeLOD_Button, SButton)
+					.OnClicked(this, &SParentStateExtensionNode::HandleOnChangeLOD_Clicked)
+					.ButtonStyle(&FAppStyle::GetWidgetStyle<FButtonStyle>("NoBorder"))
+					.DesiredSizeScale(FVector2D(1.2f, 1.2f))
+					.ContentPadding(FMargin())
+					[
+						SNew(SImage)
+						.Image(ExpandButtonBrush)
+						.ColorAndOpacity(this, &SParentStateExtensionNode::GetExpandButtonColor)
+					]
+				]
+			];
+	}
+
+	return SNullWidget::NullWidget;
+}
+
+TSharedRef<SWidget> SParentStateExtensionNode::Construct_GenerateNodeBody(const EExtendableNodeLOD InLOD)
+{
+	if (ViewState.IsValid())
+	{
+		if (InLOD == EExtendableNodeLOD::Full || InLOD == EExtendableNodeLOD::Minimal)
+		{
+			NodeBodyPadding = FMargin(3.0, 0.0, 3.0, 0.0);
+
+			FPropertyEditorModule& PropertyEditorModule = FModuleManager::GetModuleChecked<FPropertyEditorModule>("PropertyEditor");
+			FDetailsViewArgs DetailsViewArgs;
+			DetailsViewArgs.NameAreaSettings = FDetailsViewArgs::HideNameArea;
+			DetailsViewArgs.InitialCategories = InLOD == EExtendableNodeLOD::Full
+				? TSet<FName>{ "Tasks", "Enter Conditions" }
+				: TSet<FName>{ "Enter Conditions" };
+			DetailsViewArgs.bShowScrollBar = false;
+			DetailsViewArgs.bShowOptions = false;
+			DetailsViewArgs.bAllowSearch = false;
+
+			TSharedRef<IDetailsView> StateDetailsView = PropertyEditorModule.CreateDetailView(DetailsViewArgs);
+			StateDetailsView->SetObject(ViewState.Get());
+
+			return StateDetailsView;
+		}
+		else
+		{
+			// Negative offset so our background is behind the next node
+			NodeBodyPadding = FMargin(0.0, 0.0, 0.0, -8.0);
+
+			return SNew(SImage)
+				.Image(FAppStyle::Get().GetBrush("Brushes.AccentBlack"))
+				.DesiredSizeOverride(FVector2D(32.0, 12.0));
+		}
+	}
+
+	return SNullWidget::NullWidget;
+}
+
+FReply SParentStateExtensionNode::HandleOnChangeLOD_Clicked()
+{
+	if (CurrentLOD == EExtendableNodeLOD::Collapsed)
+	{
+		CurrentLOD = EExtendableNodeLOD::Minimal;
+	}
+	else if (CurrentLOD == EExtendableNodeLOD::Minimal)
+	{
+		CurrentLOD = EExtendableNodeLOD::Full;
+	}
+	else if (ensure(CurrentLOD == EExtendableNodeLOD::Full))
+	{
+		CurrentLOD = EExtendableNodeLOD::Collapsed;
+	}
+
+	RegenerateWidgetForLOD(CurrentLOD);
+
+	return FReply::Handled();
+}
+
+FSlateColor SParentStateExtensionNode::GetExpandButtonColor() const
+{
+	return ChangeLOD_Button->IsHovered() ? FStyleColors::White : FStyleColors::Foreground;
+}
+
+} // namespace UE::StateTree::Editor::StateCentricView
+
+#undef LOCTEXT_NAMESPACE
+
diff --git a/Engine/Plugins/Runtime/StateTree/Source/StateTreeEditorModule/Private/StateCentricView/View/SParentStateExtensionNode.h b/Engine/Plugins/Runtime/StateTree/Source/StateTreeEditorModule/Private/StateCentricView/View/SParentStateExtensionNode.h
new file mode 100644
--- /dev/null
+++ b/Engine/Plugins/Runtime/StateTree/Source/StateTreeEditorModule/Private/StateCentricView/View/SParentStateExtensionNode.h
@@ -0,0 +1,60 @@
+// Copyright Epic Games, Inc. All Rights Reserved.
+
+#pragma once
+
+#include "CoreMinimal.h"
+
+#include "StateCentricView/View/SStateTreeExtendableNode.h"
+
+struct FSlateBrush;
+
+class SButton;
+
+namespace UE::StateTree::Editor::StateCentricView
+{
+
+
+/**
+ * Experimental. Node used an an extension to show parent states
+ */
+class SParentStateExtensionNode : public SStateTreeExtendableNode
+{
+
+public:
+
+	static const FName NodeID;
+
+public:
+
+	//~ Begin SExtendableNode Interface
+	virtual FName GetNodeName() const override;
+	//~ End SExtendableNode Interface
+
+protected:
+
+	//~ Begin SExtendableNode Interface
+	virtual void Construct_GenerateBuiltInExtensions(const EExtendableNodeLOD InLOD) override;
+	virtual TSharedRef<SWidget> Construct_GenerateNodeHeader(const EExtendableNodeLOD InLOD) override;
+	virtual TSharedRef<SWidget> Construct_GenerateNodeBody(const EExtendableNodeLOD InLOD) override;
+	//~ End SExtendableNode Interface
+
+private:
+
+	/** Callback for top right button */
+	FReply HandleOnChangeLOD_Clicked();
+
+	/** Minor UI / UX callbacks */
+	FSlateColor GetExpandButtonColor() const;
+
+private:
+
+	TArray<TSharedPtr<FDeferredCleanupSlateBrush>> ChildStateBrushes;
+	TSharedPtr<FDeferredCleanupSlateBrush> StateOutlineBrush;
+	TSharedPtr<FDeferredCleanupSlateBrush> StateHeaderBackgroundBrush;
+	TSharedPtr<FDeferredCleanupSlateBrush> StateFootnoteBackgroundBrush;
+
+	TSharedPtr<SButton> ChangeLOD_Button;
+};
+
+
+} // namespace UE::StateTree::Editor::StateCentricView
\ No newline at end of file
diff --git a/Engine/Plugins/Runtime/StateTree/Source/StateTreeEditorModule/Private/StateCentricView/View/SParentStateTransitionExtensionNode.cpp b/Engine/Plugins/Runtime/StateTree/Source/StateTreeEditorModule/Private/StateCentricView/View/SParentStateTransitionExtensionNode.cpp
new file mode 100644
--- /dev/null
+++ b/Engine/Plugins/Runtime/StateTree/Source/StateTreeEditorModule/Private/StateCentricView/View/SParentStateTransitionExtensionNode.cpp
@@ -0,0 +1,267 @@
+// Copyright Epic Games, Inc. All Rights Reserved.
+
+#include "StateCentricView/View/SParentStateTransitionExtensionNode.h"
+
+#include "StateCentricView/ViewModel/StateCentricViewEditorDataExtension.h"
+#include "StateCentricView/View/StateCentricViewUtils.h"
+#include "StateCentricView/View/STransitionNode.h"
+#include "StateTreeEditorStyle.h"
+
+#include "Styling/SlateBrush.h"
+#include "Styling/StyleColors.h"
+#include "Widgets/Images/SImage.h"
+#include "Widgets/Input/SButton.h"
+#include "Widgets/Layout/SBorder.h"
+#include "Widgets/Text/STextBlock.h"
+
+#define LOCTEXT_NAMESPACE "SParentStateTransitionExtensionNode"
+
+namespace UE::StateTree::Editor::StateCentricView
+{
+
+
+//////////////////////////////////////////////////////////////////////////
+// SParentStateTransitionExtensionNode
+
+
+SParentStateTransitionExtensionNode::~SParentStateTransitionExtensionNode()
+{
+	if (ViewModel)
+	{
+		if (UStateTreeEditorData* EditorData = const_cast<UStateTreeEditorData*>(ViewModel->GetStateTreeEditorData()))
+		{
+			if (UStateCentricViewEditorDataExtension* StateCentricEditorData = EditorData->GetExtension<UStateCentricViewEditorDataExtension>())
+			{
+				StateCentricEditorData->PerStateTreeUserSettings.OnInTransitionViewLOD_Changed.RemoveAll(this);
+				StateCentricEditorData->PerStateTreeUserSettings.OnOutTransitionViewLOD_Changed.RemoveAll(this);
+			}
+		}
+	}
+}
+
+void SParentStateTransitionExtensionNode::Construct_InitializeNodePersistentPreLOD()
+{
+	UStateTreeEditorData* EditorData = const_cast<UStateTreeEditorData*>(ViewModel->GetStateTreeEditorData());
+	UStateCentricViewEditorDataExtension* StateCentricEditorData = EditorData->GetOrCreateExtension<UStateCentricViewEditorDataExtension>();
+
+	// Remove any prior bindings in case of issues with lifetime
+	StateCentricEditorData->PerStateTreeUserSettings.OnInTransitionViewLOD_Changed.RemoveAll(this);
+	StateCentricEditorData->PerStateTreeUserSettings.OnOutTransitionViewLOD_Changed.RemoveAll(this);
+
+	TOptional<EExtendableNodeLOD> TransitionViewUserSetLOD;
+	if (TransitionDirection == EStateTransitionDirection::In)
+	{
+		StateCentricEditorData->PerStateTreeUserSettings.OnInTransitionViewLOD_Changed.AddSP(this, &SParentStateTransitionExtensionNode::HandleOnTransitionViewLOD_SettingsChanged);
+		TransitionViewUserSetLOD = StateCentricEditorData->PerStateTreeUserSettings.GetInTransitionLOD();
+	}
+	else if (ensure(TransitionDirection == EStateTransitionDirection::Out))
+	{
+		StateCentricEditorData->PerStateTreeUserSettings.OnOutTransitionViewLOD_Changed.AddSP(this, &SParentStateTransitionExtensionNode::HandleOnTransitionViewLOD_SettingsChanged);
+		TransitionViewUserSetLOD = StateCentricEditorData->PerStateTreeUserSettings.GetOutTransitionLOD();
+	}
+
+	// We were given a LOD to use, but based on user config we may want to use a different LOD for now
+	if (TransitionViewUserSetLOD)
+	{
+		CurrentLOD = TransitionViewUserSetLOD.GetValue();
+	}
+}
+
+void SParentStateTransitionExtensionNode::Construct_GenerateBuiltInExtensions(const EExtendableNodeLOD InLOD)
+{
+	// Entire widget state color outline in overlay extension depending on LOD
+	if (InLOD != EExtendableNodeLOD::Hidden && InLOD != EExtendableNodeLOD::Collapsed)
+	{
+		StateOutlineBrush = FStateCentricViewUtils::MakeStateOutlineMuted(ViewModel.ToSharedRef(), ViewState.Get());
+
+		TSharedPtr<SWidget> StateColorOutline = SNew(SBorder)
+			.BorderImage(StateOutlineBrush->GetSlateBrush())
+			.Visibility(EVisibility::HitTestInvisible);
+
+		AddExtension(StateColorOutline.ToSharedRef())
+			.SetLocation(EExtensionLocation::Foreground)
+			.SetPadding(FMargin(0.0, 0.0, 0.0, 0.0));
+	}
+}
+
+TSharedRef<SWidget> SParentStateTransitionExtensionNode::Construct_GenerateNodeHeader(const EExtendableNodeLOD InLOD)
+{
+	// @TODO: Don't show anything if no transitions for parent, maybe
+
+	if (ViewState.IsValid())
+	{
+		if (InLOD == EExtendableNodeLOD::Hidden)
+		{
+			return SNew(SImage)
+				.Image(FAppStyle::Get().GetBrush("Brushes.Recessed"))
+				.DesiredSizeOverride(FVector2D(32.0, 8.0));
+		}
+
+		TArray<FStateTreeTransition*> StateTransitions;
+		EHorizontalAlignment HeaderAlignment = HAlign_Left;
+		FMargin HeaderPadding;
+		FText HeaderText;
+		if (TransitionDirection == EStateTransitionDirection::In)
+		{
+			StateTransitions = ViewModel->GetStateInTransitions(ViewState.Get());
+			HeaderAlignment = HAlign_Left;
+			HeaderPadding = FMargin(12.0, 0.0, 2.0, 0.0);
+			HeaderText = LOCTEXT("TransitionHeaderIn", "IN ({0})");
+		}
+		else if (ensure(TransitionDirection == EStateTransitionDirection::Out))
+		{
+			StateTransitions = ViewModel->GetStateOutTransitions(ViewState.Get());
+			HeaderAlignment = HAlign_Right;
+			HeaderPadding = FMargin(2.0, 0.0, 12.0, 0.0);
+			HeaderText = LOCTEXT("TransitionHeaderOut", "OUT ({0})");
+		}
+
+		HeaderText = FText::Format(HeaderText, StateTransitions.Num());
+
+		StateHeaderBackgroundMutedBrush = FStateCentricViewUtils::MakeStateHeaderBackgroundMuted(ViewModel.ToSharedRef(), ViewState.Get());
+
+		// At Lowest LOD we actually use the header background for the outline
+		if (InLOD == EExtendableNodeLOD::Collapsed)
+		{
+			StateHeaderBackgroundMutedBrush = FStateCentricViewUtils::MakeStateBackgroundCollapsedOutlined(ViewModel.ToSharedRef(), ViewState.Get());
+		}
+
+		TSharedPtr<SHorizontalBox> HeaderContentsBox;
+
+		TSharedRef<SWidget> HeaderWidget = SNew(SOverlay)
+			+ SOverlay::Slot()
+			[
+				SNew(SImage)
+				.Image(StateHeaderBackgroundMutedBrush->GetSlateBrush())
+			]
+			+ SOverlay::Slot()
+			[
+				SAssignNew(HeaderContentsBox, SHorizontalBox)
+				+ SHorizontalBox::Slot()
+				.HAlign(HeaderAlignment)
+				.VAlign(VAlign_Center)
+				.FillWidth(1.0f)
+				.Padding(HeaderPadding)
+				[
+					SNew(STextBlock)
+					.Text(HeaderText)
+					.TextStyle(&FStateTreeEditorStyle::Get().GetWidgetStyle<FTextBlockStyle>("Normal.Normal"))
+				]
+			];
+
+		// In full LOD, we add a button to allow for new transitions to be added
+		if (InLOD == EExtendableNodeLOD::Full)
+		{
+			HeaderContentsBox->AddSlot()
+			.HAlign(HAlign_Right)
+			.VAlign(VAlign_Center)
+			.AutoWidth()
+			.Padding(FMargin(-8.0, 0.0, 8.0, 0.0))
+			[
+				SAssignNew(AddTransitionButton, SButton)
+				.OnClicked(this, &SParentStateTransitionExtensionNode::HandleOnAddTransitionClicked)
+				.ButtonStyle(&FAppStyle::GetWidgetStyle<FButtonStyle>("NoBorder"))
+				.ContentPadding(FMargin())
+				[
+					SNew(SImage)
+					.Image(FAppStyle::GetBrush("Icons.PlusCircle"))
+					.ColorAndOpacity(this, &SParentStateTransitionExtensionNode::GetAddTransitionColor)
+				]
+			];
+		}
+
+		return HeaderWidget;
+	}
+
+	return SNullWidget::NullWidget;
+}
+
+TSharedRef<SWidget> SParentStateTransitionExtensionNode::Construct_GenerateNodeBody(const EExtendableNodeLOD InLOD)
+{
+	if (ViewState.IsValid())
+	{
+		return FStateCentricViewUtils::GenerateTransitionsViewWidget(InLOD, ViewModel.ToSharedRef(), ViewState.Get(), TransitionDirection);
+	}
+
+	return SNullWidget::NullWidget;
+}
+
+FReply SParentStateTransitionExtensionNode::HandleOnAddTransitionClicked()
+{
+	return FReply::Handled();
+}
+
+void SParentStateTransitionExtensionNode::HandleOnTransitionViewLOD_SettingsChanged(TOptional<EExtendableNodeLOD> InLOD)
+{
+	if (ViewState.IsValid())
+	{
+		UStateTreeEditorData* EditorData = const_cast<UStateTreeEditorData*>(ViewModel->GetStateTreeEditorData());
+		UStateCentricViewEditorDataExtension* StateCentricEditorData = EditorData->GetOrCreateExtension<UStateCentricViewEditorDataExtension>();
+
+		TOptional<EExtendableNodeLOD> TransitionViewLOD;
+		if (TransitionDirection == EStateTransitionDirection::In)
+		{
+			TransitionViewLOD = StateCentricEditorData->PerStateTreeUserSettings.GetInTransitionLOD();
+		}
+		else if (ensure(TransitionDirection == EStateTransitionDirection::Out))
+		{
+			TransitionViewLOD = StateCentricEditorData->PerStateTreeUserSettings.GetOutTransitionLOD();
+		}
+
+		if (TransitionViewLOD)
+		{
+			RegenerateWidgetForLOD(TransitionViewLOD.GetValue());
+		}
+		else
+		{
+			RegenerateWidgetForLOD(Config.InitialLOD);
+		}
+	}
+}
+
+FSlateColor SParentStateTransitionExtensionNode::GetAddTransitionColor() const
+{
+	return AddTransitionButton->IsHovered() ? FStyleColors::White : FStyleColors::Foreground;
+}
+
+
+//////////////////////////////////////////////////////////////////////////
+// SParentStateInTransitionExtensionNode
+
+const FName SParentStateInTransitionExtensionNode::NodeID(TEXT("StateCentricView.ParentState.InTransitions"));
+
+FName SParentStateInTransitionExtensionNode::GetNodeName() const
+{
+	return NodeID;
+}
+
+void SParentStateInTransitionExtensionNode::Construct_InitializeNodePersistentPreLOD()
+{
+	TransitionDirection = EStateTransitionDirection::In;
+	SParentStateTransitionExtensionNode::Construct_InitializeNodePersistentPreLOD();
+}
+
+
+//////////////////////////////////////////////////////////////////////////
+// SParentStateOutTransitionExtensionNode
+
+const FName SParentStateOutTransitionExtensionNode::NodeID(TEXT("StateCentricView.ParentState.OutTransitions"));
+
+FName SParentStateOutTransitionExtensionNode::GetNodeName() const
+{
+	return NodeID;
+}
+
+void SParentStateOutTransitionExtensionNode::Construct_InitializeNodePersistentPreLOD()
+{
+	TransitionDirection = EStateTransitionDirection::Out;
+	SParentStateTransitionExtensionNode::Construct_InitializeNodePersistentPreLOD();
+}
+
+
+//////////////////////////////////////////////////////////////////////////
+
+} // namespace UE::StateTree::Editor::StateCentricView
+
+#undef LOCTEXT_NAMESPACE
+
diff --git a/Engine/Plugins/Runtime/StateTree/Source/StateTreeEditorModule/Private/StateCentricView/View/SParentStateTransitionExtensionNode.h b/Engine/Plugins/Runtime/StateTree/Source/StateTreeEditorModule/Private/StateCentricView/View/SParentStateTransitionExtensionNode.h
new file mode 100644
--- /dev/null
+++ b/Engine/Plugins/Runtime/StateTree/Source/StateTreeEditorModule/Private/StateCentricView/View/SParentStateTransitionExtensionNode.h
@@ -0,0 +1,112 @@
+// Copyright Epic Games, Inc. All Rights Reserved.
+
+#pragma once
+
+#include "CoreMinimal.h"
+
+#include "StateCentricView/View/SStateTreeExtendableNode.h"
+
+struct FSlateBrush;
+struct FSlateColor;
+
+class SButton;
+
+namespace UE::StateTree::Editor::StateCentricView
+{
+
+
+/**
+ * Experimental. Node used as an extension for In / Out transitions on the main state.
+ */
+class SParentStateTransitionExtensionNode : public SStateTreeExtendableNode
+{
+
+public:
+
+	virtual ~SParentStateTransitionExtensionNode() override;
+
+protected:
+
+	//~ Begin SExtendableNode Interface
+	virtual void Construct_InitializeNodePersistentPreLOD() override;
+	virtual void Construct_GenerateBuiltInExtensions(const EExtendableNodeLOD InLOD) override;
+	virtual TSharedRef<SWidget> Construct_GenerateNodeHeader(const EExtendableNodeLOD InLOD) override;
+	virtual TSharedRef<SWidget> Construct_GenerateNodeBody(const EExtendableNodeLOD InLOD) override;
+	//~ End SExtendableNode Interface
+
+protected:
+
+	/** If we are showing in / out transitions. Init in child classes. */
+	EStateTransitionDirection TransitionDirection = EStateTransitionDirection::Out;
+
+private:
+
+	/** Callback for top right button */
+	FReply HandleOnAddTransitionClicked();
+
+	/** Callback for when LOD settings change */
+	void HandleOnTransitionViewLOD_SettingsChanged(TOptional<EExtendableNodeLOD> InLOD);
+
+	/** Minor UI / UX callbacks */
+	FSlateColor GetAddTransitionColor() const;
+
+private:
+
+	TSharedPtr<FDeferredCleanupSlateBrush> StateOutlineBrush;
+	TSharedPtr<FDeferredCleanupSlateBrush> StateHeaderBackgroundMutedBrush;
+
+	TSharedPtr<SButton> AddTransitionButton;
+};
+
+
+/**
+ * Experimental. Specialization of SParentStateTransitionExtensionNode for 'In' transitions
+ */
+class SParentStateInTransitionExtensionNode : public SParentStateTransitionExtensionNode
+{
+
+public:
+
+	static const FName NodeID;
+
+public:
+
+	//~ Begin SExtendableNode Interface
+	virtual FName GetNodeName() const override;
+	//~ End SExtendableNode Interface
+
+protected:
+
+	//~ Begin SExtendableNode Interface
+	virtual void Construct_InitializeNodePersistentPreLOD() override;
+	//~ End SExtendableNode Interface
+
+};
+
+
+/**
+ * Experimental. Specialization of SParentStateTransitionExtensionNode for 'Out' transitions
+ */
+class SParentStateOutTransitionExtensionNode : public SParentStateTransitionExtensionNode
+{
+
+public:
+
+	static const FName NodeID;
+
+public:
+
+	//~ Begin SExtendableNode Interface
+	virtual FName GetNodeName() const override;
+	//~ End SExtendableNode Interface
+
+protected:
+
+	//~ Begin SExtendableNode Interface
+	virtual void Construct_InitializeNodePersistentPreLOD() override;
+	//~ End SExtendableNode Interface
+
+};
+
+
+} // namespace UE::StateTree::Editor::StateCentricView
\ No newline at end of file
diff --git a/Engine/Plugins/Runtime/StateTree/Source/StateTreeEditorModule/Private/StateCentricView/View/SStateCentricViewMainStateNode.cpp b/Engine/Plugins/Runtime/StateTree/Source/StateTreeEditorModule/Private/StateCentricView/View/SStateCentricViewMainStateNode.cpp
new file mode 100644
--- /dev/null
+++ b/Engine/Plugins/Runtime/StateTree/Source/StateTreeEditorModule/Private/StateCentricView/View/SStateCentricViewMainStateNode.cpp
@@ -0,0 +1,262 @@
+// Copyright Epic Games, Inc. All Rights Reserved.
+
+#include "StateCentricView/View/SStateCentricViewMainStateNode.h"
+
+#include "StateCentricView/ViewModel/StateCentricViewEditorDataExtension.h"
+#include "StateCentricView/View/SMainStateTransitionExtensionNode.h"
+#include "StateCentricView/View/SParentStateExtensionNode.h"
+#include "StateCentricView/View/StateCentricViewUtils.h"
+#include "StateTreeEditorStyle.h"
+
+#include "IDetailsView.h"
+#include "DetailsViewArgs.h"
+#include "Layout/Visibility.h"
+#include "Modules/ModuleManager.h"
+#include "ObjectEditorUtils.h"
+#include "PropertyEditorModule.h"
+#include "Styling/SlateBrush.h"
+#include "Styling/StyleColors.h"
+#include "Widgets/Images/SImage.h"
+#include "Widgets/Input/SButton.h"
+#include "Widgets/Layout/SBorder.h"
+#include "Widgets/Text/STextBlock.h"
+
+#define LOCTEXT_NAMESPACE "SStateCentricViewMainStateNode"
+
+namespace UE::StateTree::Editor::StateCentricView
+{
+
+const FName SStateCentricViewMainStateNode::NodeID(TEXT("StateCentricView.MainState.Node"));
+
+SStateCentricViewMainStateNode::~SStateCentricViewMainStateNode()
+{
+	if (ViewModel)
+	{
+		ViewModel->GetOnSelectionChanged().RemoveAll(this);
+	}
+}
+
+FName SStateCentricViewMainStateNode::GetNodeName() const
+{
+	return NodeID;
+}
+
+void SStateCentricViewMainStateNode::Construct_InitializeNodePersistentPreLOD()
+{
+	// Remove any prior bindings in case of issues with lifetime
+	ViewModel->GetOnSelectionChanged().RemoveAll(this);
+
+	ViewModel->GetOnSelectionChanged().AddSP(this, &SStateCentricViewMainStateNode::HandleOnSelectionChanged);
+}
+
+void SStateCentricViewMainStateNode::Construct_GenerateBuiltInExtensions(const EExtendableNodeLOD InLOD)
+{
+	// We don't change per LOD, so we shouldn't be getting LOD changes
+	ensure(InLOD == EExtendableNodeLOD::Unused);
+
+	ChildStateBrushes.Reset();
+
+	if (ViewState.IsValid())
+	{
+		// @TODO: Remove splitter references here once those become a feature toggle on cardinal extensions
+		UStateTreeEditorData* EditorData = const_cast<UStateTreeEditorData*>(ViewModel->GetStateTreeEditorData());
+		UStateCentricViewEditorDataExtension* StateCentricEditorData = EditorData->GetOrCreateExtension<UStateCentricViewEditorDataExtension>();
+		
+		// In Transition extension on left
+		{
+			FScopedMainNodeLayoutInfo ScopedMainLayoutInfo = SetScopedMainNodeLayoutInfo()
+				.SetSplitterValue(MakeAttributeUObject(StateCentricEditorData, &UStateCentricViewEditorDataExtension::HandleGetMainNodeSplitterValue))
+				.SetSplitterOnSlotResized(SSplitter::FOnSlotResized::CreateUObject(StateCentricEditorData, &UStateCentricViewEditorDataExtension::HandleOnMainNodeSplitterResized));
+
+			// @TODO: Read Init LOD from project settings for in transitions.
+			TSharedRef<SWidget> InTransitionLeftExtension = SNew(SMainStateInTransitionExtensionNode, ViewModel.ToSharedRef(), ViewState.Get())
+				.Config({ .InitialLOD = EExtendableNodeLOD::Minimal });
+
+			AddExtension(InTransitionLeftExtension)
+				.SetLocation(EExtensionLocation::SplitLeft)
+				.SetSplitterSizeRule(MakeAttributeUObject(StateCentricEditorData, &UStateCentricViewEditorDataExtension::GetInTransitionSplitterSizeRule))
+				.SetSplitterValue(MakeAttributeUObject(StateCentricEditorData, &UStateCentricViewEditorDataExtension::HandleGetInTransitionSplitterValue))
+				.SetSplitterOnSlotResized(SSplitter::FOnSlotResized::CreateUObject(StateCentricEditorData, &UStateCentricViewEditorDataExtension::HandleOnInTransitionSplitterResized));
+
+			// Add spacer should we hit max size caps, Needed since fill + max size = left align (Slate bug).
+			AddExtension(
+				SNew(SImage)
+					.Image(FAppStyle::Get().GetBrush("Brushes.Recessed"))
+					.DesiredSizeOverride(FVector2D(32.0, 8.0)))
+				.SetLocation(EExtensionLocation::SplitLeft)
+				.SetSplitterValue(MakeAttributeUObject(StateCentricEditorData, &UStateCentricViewEditorDataExtension::HandleGetLeftSpacerSplitterValue))
+				.SetSplitterOnSlotResized(SSplitter::FOnSlotResized::CreateUObject(StateCentricEditorData, &UStateCentricViewEditorDataExtension::HandleOnLeftSpacerSplitterResized));
+		}
+
+		// Out Transition extension on right
+		{
+			// @TODO: Read Init LOD from project settings for out transitions.
+			TSharedRef<SWidget> OutTransitionRightExtension = SNew(SMainStateOutTransitionExtensionNode, ViewModel.ToSharedRef(), ViewState.Get())
+				.Config({ .InitialLOD = EExtendableNodeLOD::Full });
+
+			AddExtension(OutTransitionRightExtension)
+				.SetLocation(EExtensionLocation::SplitRight)
+				.SetSplitterSizeRule(MakeAttributeUObject(StateCentricEditorData, &UStateCentricViewEditorDataExtension::GetOutTransitionSplitterSizeRule))
+				.SetSplitterValue(MakeAttributeUObject(StateCentricEditorData, &UStateCentricViewEditorDataExtension::HandleGetOutTransitionSplitterValue))
+				.SetSplitterOnSlotResized(SSplitter::FOnSlotResized::CreateUObject(StateCentricEditorData, &UStateCentricViewEditorDataExtension::HandleOnOutTransitionSplitterResized));
+
+			// Add spacer should we hit max size caps, Needed since fill + max size = left align (Slate bug).
+			AddExtension(
+				SNew(SImage)
+					.Image(FAppStyle::Get().GetBrush("Brushes.Recessed"))
+					.DesiredSizeOverride(FVector2D(32.0, 8.0)))
+				.SetLocation(EExtensionLocation::SplitRight)
+				.SetSplitterValue(MakeAttributeUObject(StateCentricEditorData, &UStateCentricViewEditorDataExtension::HandleGetRightSpacerSplitterValue))
+				.SetSplitterOnSlotResized(SSplitter::FOnSlotResized::CreateUObject(StateCentricEditorData, &UStateCentricViewEditorDataExtension::HandleOnRightSpacerSplitterResized));
+		}
+
+		// Parent state extensions at top
+		{
+			uint32 NumParentStateShown = 0;
+
+			FScopedMainNodeLayoutInfo ScopedMainLayoutInfo = SetScopedMainNodeLayoutInfo()
+				.SetHAlign(HAlign_Fill)
+				.SetFillContentSize(1.0f, 0.0f);
+
+			TArray<UStateTreeState*> ParentStates = ViewModel->GetStateParents(ViewState.Get());
+
+			for (UStateTreeState* ParentState : ParentStates)
+			{
+				// @TODO: Read project settings and cull after N shown parents
+
+				TSharedRef<SWidget> ParentStateTopExtension = SNew(SParentStateExtensionNode, ViewModel.ToSharedRef(), ParentState)
+					.Config({ .InitialLOD = EExtendableNodeLOD::Minimal });
+
+				AddExtension(ParentStateTopExtension)
+					.SetLocation(EExtensionLocation::Top)
+					.SetAutoSize()
+					.SetPadding(FMargin(0.0, 0.0, 0.0, 2.0));
+
+				NumParentStateShown++;
+			}
+		}
+
+		// Child states in footnote
+		FStateCentricViewUtils::AddChildStateFootnoteExtensions(EExtendableNodeLOD::Full, this, OUT ChildStateBrushes);
+
+		// Entire widget state color outline in overlay extension, half circle if not full LOD
+		{
+			TArray<UStateTreeState*> CentralStateParents = ViewModel->GetStateParents(ViewModel->GetViewState());
+			int32 OverlayOffset = CentralStateParents.Num() * 5000;
+
+			FScopedMainNodeLayoutInfo ScopedMainLayoutInfo = SetScopedMainNodeLayoutInfo()
+				.SetOverlayOffset(OverlayOffset);
+
+			StateOutlineBrush = FStateCentricViewUtils::MakeStateOutline(ViewModel.ToSharedRef(), ViewState.Get());
+
+			TSharedPtr<SWidget> StateColorOutline = SNew(SBorder)
+				.BorderImage(StateOutlineBrush->GetSlateBrush())
+				.Visibility(EVisibility::HitTestInvisible);
+
+			AddExtension(StateColorOutline.ToSharedRef())
+				.SetLocation(EExtensionLocation::Foreground);
+		}
+	}
+}
+
+TSharedRef<SWidget> SStateCentricViewMainStateNode::Construct_GenerateNodeHeader(const EExtendableNodeLOD InLOD)
+{
+	if (ViewState.IsValid())
+	{
+		UStateTreeEditorData* EditorData = const_cast<UStateTreeEditorData*>(ViewModel->GetStateTreeEditorData());
+		UStateCentricViewEditorDataExtension* StateCentricEditorData = EditorData->GetOrCreateExtension<UStateCentricViewEditorDataExtension>();
+
+		StateHeaderBackgroundBrush = FStateCentricViewUtils::MakeStateHeaderBackground(ViewModel.ToSharedRef(), ViewState.Get(), FVector2f(32.0f, 28.0f));
+
+		// @TODO: Temp for now have double clicking the header reset layout. Discuss what we want to do. Feels weird to have it do nothing.
+		TWeakObjectPtr<UStateCentricViewEditorDataExtension> WeakStateCentricEditorData = StateCentricEditorData;
+		auto ResetLayout = [WeakStateCentricEditorData](const FGeometry&, const FPointerEvent& )
+		{
+			if (UStateCentricViewEditorDataExtension* StateCentricEditorDataPinned = WeakStateCentricEditorData.Get())
+			{
+				StateCentricEditorDataPinned->ResetSplitters();
+			}
+
+			return FReply::Handled();
+		};
+
+		// Generate via utils to match style
+		TSharedRef<SButton> ResetLayoutButton = FStateCentricViewUtils::GenerateGotoStateButton(InLOD, ViewModel.ToSharedRef(), ViewState.Get());
+		ResetLayoutButton->SetOnMouseDoubleClick(FPointerEventHandler::CreateLambda(ResetLayout));
+
+		return SNew(SOverlay)
+			+ SOverlay::Slot()
+			[
+				SNew(SImage)
+					.Image(StateHeaderBackgroundBrush->GetSlateBrush())
+			]
+			+ SOverlay::Slot()
+			.Padding(FMargin(8.0, 2.0, 8.0, 2.0))
+			[
+				ResetLayoutButton
+			]
+			+ SOverlay::Slot()
+			.HAlign(HAlign_Center)
+			.VAlign(VAlign_Center)
+			[
+				SNew(SHorizontalBox)
+				+ SHorizontalBox::Slot()
+				.HAlign(HAlign_Right)
+				.VAlign(VAlign_Center)
+				.Padding(FMargin(-16.0, 0.0, 2.0, 0.0))
+				.AutoWidth()
+				[
+					SNew(SImage)
+					.Image(ViewModel->GetSelectorIcon(ViewState.Get()))
+					.Visibility(EVisibility::HitTestInvisible)
+				]
+				+ SHorizontalBox::Slot()
+				.HAlign(HAlign_Center)
+				.VAlign(VAlign_Center)
+				.AutoWidth()
+				[
+					SNew(STextBlock)
+					.Text(FText::FromName(ViewState->Name))
+					.TextStyle(&FStateTreeEditorStyle::Get().GetWidgetStyle<FTextBlockStyle>("StateTree.CentricView.StateTitle.Bold"))
+					.Visibility(EVisibility::HitTestInvisible)
+				]
+			];
+	}
+
+	return SNullWidget::NullWidget;
+}
+
+TSharedRef<SWidget> SStateCentricViewMainStateNode::Construct_GenerateNodeBody(const EExtendableNodeLOD InLOD)
+{
+	if (ViewState.IsValid())
+	{
+		NodeBodyPadding = FMargin(3.0, 0.0, 3.0, 0.0);
+
+		FPropertyEditorModule& PropertyEditorModule = FModuleManager::GetModuleChecked<FPropertyEditorModule>("PropertyEditor");
+		FDetailsViewArgs DetailsViewArgs;
+		DetailsViewArgs.NameAreaSettings = FDetailsViewArgs::HideNameArea;
+		DetailsViewArgs.InitialCategories = TSet<FName>{ "Tasks", "Enter Conditions" };
+		DetailsViewArgs.bShowScrollBar = false;
+		DetailsViewArgs.bShowOptions = false;
+		DetailsViewArgs.bAllowSearch = false;
+
+		TSharedRef<IDetailsView> StateDetailsView = PropertyEditorModule.CreateDetailView(DetailsViewArgs);
+		StateDetailsView->SetObject(ViewState.Get());
+
+		return StateDetailsView;
+	}
+
+	return SNullWidget::NullWidget;
+}
+
+void SStateCentricViewMainStateNode::HandleOnSelectionChanged(const TArray<TWeakObjectPtr<UStateTreeState>>& SelectedStates)
+{
+	ViewState = ViewModel->GetLastSelectedState();
+
+	RegenerateWidgetForLOD(CurrentLOD);
+}
+
+} // namespace UE::StateTree::Editor::StateCentricView
+
+#undef LOCTEXT_NAMESPACE
+
diff --git a/Engine/Plugins/Runtime/StateTree/Source/StateTreeEditorModule/Private/StateCentricView/View/SStateCentricViewMainStateNode.h b/Engine/Plugins/Runtime/StateTree/Source/StateTreeEditorModule/Private/StateCentricView/View/SStateCentricViewMainStateNode.h
new file mode 100644
--- /dev/null
+++ b/Engine/Plugins/Runtime/StateTree/Source/StateTreeEditorModule/Private/StateCentricView/View/SStateCentricViewMainStateNode.h
@@ -0,0 +1,54 @@
+// Copyright Epic Games, Inc. All Rights Reserved.
+
+#pragma once
+
+#include "CoreMinimal.h"
+
+#include "StateCentricView/View/SStateTreeExtendableNode.h"
+
+struct FSlateBrush;
+
+namespace UE::StateTree::Editor::StateCentricView
+{
+
+
+/**
+ * Experimental. Node for the main selected state in state centric view
+ */
+class SStateCentricViewMainStateNode : public SStateTreeExtendableNode
+{
+
+public:
+
+	static const FName NodeID;
+	virtual ~SStateCentricViewMainStateNode() override;
+
+public:
+
+	//~ Begin SExtendableNode Interface
+	virtual FName GetNodeName() const override;
+	//~ End SExtendableNode Interface
+
+protected:
+
+	//~ Begin SExtendableNode Interface
+	virtual void Construct_InitializeNodePersistentPreLOD() override;
+	virtual void Construct_GenerateBuiltInExtensions(const EExtendableNodeLOD InLOD) override;
+	virtual TSharedRef<SWidget> Construct_GenerateNodeHeader(const EExtendableNodeLOD InLOD) override;
+	virtual TSharedRef<SWidget> Construct_GenerateNodeBody(const EExtendableNodeLOD InLOD) override;
+	//~ End SExtendableNode Interface
+
+protected:
+
+	/** Callback whenever selected state changes. */
+	void HandleOnSelectionChanged(const TArray<TWeakObjectPtr<UStateTreeState>>& SelectedStates);
+
+private:
+
+	TArray<TSharedPtr<FDeferredCleanupSlateBrush>> ChildStateBrushes;
+	TSharedPtr<FDeferredCleanupSlateBrush> StateOutlineBrush;
+	TSharedPtr<FDeferredCleanupSlateBrush> StateHeaderBackgroundBrush;
+};
+
+
+} // namespace UE::StateTree::Editor::StateCentricView
\ No newline at end of file
diff --git a/Engine/Plugins/Runtime/StateTree/Source/StateTreeEditorModule/Private/StateCentricView/View/SStateTreeExtendableNode.cpp b/Engine/Plugins/Runtime/StateTree/Source/StateTreeEditorModule/Private/StateCentricView/View/SStateTreeExtendableNode.cpp
new file mode 100644
--- /dev/null
+++ b/Engine/Plugins/Runtime/StateTree/Source/StateTreeEditorModule/Private/StateCentricView/View/SStateTreeExtendableNode.cpp
@@ -0,0 +1,36 @@
+// Copyright Epic Games, Inc. All Rights Reserved.
+
+#include "StateCentricView/View/SStateTreeExtendableNode.h"
+
+#define LOCTEXT_NAMESPACE "SStateTreeExtendableNode"
+
+namespace UE::StateTree::Editor::StateCentricView
+{
+
+void SStateTreeExtendableNode::Construct(const SExtendableNode::FArguments& InArgs
+	, TSharedRef<FStateTreeViewModel> InViewModel
+	, TNotNull<UStateTreeState*> InViewState)
+{
+	ViewModel = InViewModel;
+	ViewState = InViewState;
+
+	SExtendableNode::Construct(InArgs);
+}
+
+void SStateTreeExtendableNode::Construct_GatherExtensionsFromExternalSubsystems()
+{
+	// @TODO: 
+	// 
+	// 1. Create a StateTree ExtendableNode subsystem that allows extension registration given schema & node name
+	// 2. Using Schema from ViewModel, gather all extensions that for our schema / node name combination
+	// 3. Those extensions will be of the form `FAddStateTreeExtensionDelegate`, which is not compatible w/ `FAddExtensionDelegate`
+	// 4. To address above, we can wrap that ST delegate a `FAddExtensionDelegate` made here, and provide (VM) to the inner ST delegate.
+
+	// Register any schema-agnostic extensions
+	SExtendableNode::Construct_GatherExtensionsFromExternalSubsystems();
+}
+
+} // namespace UE::StateTree::Editor::StateCentricView
+
+#undef LOCTEXT_NAMESPACE
+
diff --git a/Engine/Plugins/Runtime/StateTree/Source/StateTreeEditorModule/Private/StateCentricView/View/SStateTreeExtendableNode.h b/Engine/Plugins/Runtime/StateTree/Source/StateTreeEditorModule/Private/StateCentricView/View/SStateTreeExtendableNode.h
new file mode 100644
--- /dev/null
+++ b/Engine/Plugins/Runtime/StateTree/Source/StateTreeEditorModule/Private/StateCentricView/View/SStateTreeExtendableNode.h
@@ -0,0 +1,93 @@
+// Copyright Epic Games, Inc. All Rights Reserved.
+
+#pragma once
+
+#include "CoreMinimal.h"
+
+#include "Debugger/StateTreeDebugger.h"
+#include "StateTree.h"
+#include "StateTreeEditorData.h"
+#include "StateTreeDelegates.h"
+#include "StateTreeViewModel.h"
+#include "StateCentricView/Layout/SExtendableNode.h"
+
+#include "Editor.h"
+#include "Slate/DeferredCleanupSlateBrush.h"
+
+namespace UE::StateTree::Editor::StateCentricView
+{
+	struct FStateCentricViewUtils;
+}
+
+namespace UE::StateTree::Editor::StateCentricView
+{
+
+
+/**
+ * Experimental. If a state transition is in / out relative to main central state
+ */
+enum class EStateTransitionDirection : uint8
+{
+	In,			// Transition into central state
+	Out,		// Transition out of central state
+};
+
+
+/**
+ * Experimental. ST specific callback for extensions. 
+ * These extensions will only get added if the schema they are registered against matches the schema for our current ViewModel. 
+ * 
+ * Note: Impl detail. We will internally wrap the extension delegate given into one that conforms to the default extension interface.
+ *
+ * @param ExtendableNode	- Node to extend by programatically calling `SExtendableNode::AddExtension` based on provided LOD,
+ * @param LOD				- Convenience param. LOD for the current node. Same as ExtendableNode's LOD.
+ * @param ViewModel			- ViewModel for the current state tree
+ * @param ViewState			- State that this extension should specialize view around, may be different than the central state (IE: If showing a parent state).
+ * 
+ * @TODO: Currently not used.
+ */
+using FAddStateTreeExtensionDelegate = TDelegate<void(TSharedRef<SExtendableNode> /*ExtendableNode*/
+	, const EExtendableNodeLOD /*LOD*/
+	, TSharedRef<FStateTreeViewModel> /*ViewModel*/
+	, TNotNull<UStateTreeState*> /*ViewState*/)>;
+
+
+/** 
+ * Experimental. Extendable node StateTree specialization. 
+ * Guarantees access to a ST ViewModel & gathers external extensions on a per-schema basis.
+ * Also gets relevant view state during construction;
+ */
+class SStateTreeExtendableNode : public SExtendableNode
+{
+
+protected:
+
+	friend struct FStateCentricViewUtils;
+
+public:
+
+	/**
+	 * @param InViewModel	ViewModel for the StateTree
+	 * @param InViewState	State this node should try to provide a view for. 
+	 */
+	void Construct(const SExtendableNode::FArguments& InArgs
+		, TSharedRef<FStateTreeViewModel> InViewModel
+		, TNotNull<UStateTreeState*> InViewState);
+
+protected:
+
+	//~ Begin SExtendableNode Interface
+	virtual void Construct_GatherExtensionsFromExternalSubsystems() override;
+	//~ End SExtendableNode Interface
+
+protected:
+
+	/** ViewModel for the ST this node interacts with */
+	TSharedPtr<FStateTreeViewModel> ViewModel = nullptr;
+
+	/** View State. May be a different state at different parts of state centric view. Ex: Parent state for parent widgets. */
+	TWeakObjectPtr<UStateTreeState> ViewState = nullptr;
+};
+
+
+} // namespace UE::StateTree::Editor::StateCentricView
\ No newline at end of file
diff --git a/Engine/Plugins/Runtime/StateTree/Source/StateTreeEditorModule/Private/StateCentricView/View/STransitionNode.cpp b/Engine/Plugins/Runtime/StateTree/Source/StateTreeEditorModule/Private/StateCentricView/View/STransitionNode.cpp
new file mode 100644
--- /dev/null
+++ b/Engine/Plugins/Runtime/StateTree/Source/StateTreeEditorModule/Private/StateCentricView/View/STransitionNode.cpp
@@ -0,0 +1,406 @@
+// Copyright Epic Games, Inc. All Rights Reserved.
+
+#include "StateCentricView/View/STransitionNode.h"
+
+#include "StateCentricView/View/StateCentricViewUtils.h"
+#include "StateTreeEditorStyle.h"
+
+#include "Customizations/StateTreeTransitionDetails.h"
+#include "DetailsViewArgs.h"
+#include "IDetailsView.h"
+#include "IStructureDetailsView.h"
+#include "Layout/Visibility.h"
+#include "Modules/ModuleManager.h"
+#include "ObjectEditorUtils.h"
+#include "PropertyEditorModule.h"
+#include "PropertyPathHelpers.h"
+#include "Styling/SlateBrush.h"
+#include "Styling/StyleColors.h"
+#include "Widgets/Images/SImage.h"
+#include "Widgets/Input/SButton.h"
+#include "Widgets/Layout/SBorder.h"
+#include "Widgets/Text/STextBlock.h"
+
+// Bad API, need to make not private
+#include "Editor/PropertyEditor/Private/PropertyNode.h"
+
+#define LOCTEXT_NAMESPACE "STransitionNode"
+
+namespace UE::StateTree::Editor::StateCentricView
+{
+
+
+//////////////////////////////////////////////////////////////////////////
+// STransitionNode
+
+void STransitionNode::Construct(const SExtendableNode::FArguments& InArgs
+		, TSharedRef<FStateTreeViewModel> InViewModel
+		, TNotNull<UStateTreeState*> InViewState
+		, TNotNull<FStateTreeTransition*> InTransition)
+{
+	TransitionID = InTransition->ID;
+
+	SStateTreeExtendableNode::Construct(InArgs, InViewModel, InViewState);
+}
+
+void STransitionNode::Construct_GenerateBuiltInExtensions(const EExtendableNodeLOD InLOD)
+{
+	// @TODO: Rather than handle the transition arrow via padding. Maybe just add an inner node for the transition & make the 
+	// transition arrow not part of this node
+
+	FStateTreeTransition* Transition = ViewModel->GetMutableTransitionByID(TransitionID);
+	if (Transition && ViewState.IsValid())
+	{
+		// Parent state name as top extension
+		if (const UStateTreeState* ParentState = TransitionDirection == EStateTransitionDirection::In 
+			? ViewModel->GetSourceParentStateFromTransition(Transition)
+			: ViewModel->GetTargetParentStateFromTransition(Transition))
+		{
+			TransitionParentStateBrush = FStateCentricViewUtils::MakeStateSubnodeBackgroundOutlined(ViewModel.ToSharedRef(), ParentState);
+
+			FScopedMainNodeLayoutInfo ScopedMainLayoutInfo = SetScopedMainNodeLayoutInfo()
+				.SetPadding(FMargin(0.0));
+
+			// @TODO: Make sure to add a warning icon if one of the parents in the nested hierachy has a entry condition
+			TSharedRef<SWidget> ParentStateNameTopExtension = SNew(SOverlay)
+				+ SOverlay::Slot()
+				[
+					SNew(SImage)
+					.Image(TransitionParentStateBrush->GetSlateBrush())
+				]
+				+ SOverlay::Slot()
+				[
+					FStateCentricViewUtils::GenerateGotoStateButton(InLOD, ViewModel.ToSharedRef(), ParentState)
+				]
+				+ SOverlay::Slot()
+				[
+					SNew(SHorizontalBox)
+					+ SHorizontalBox::Slot()
+					.HAlign(HAlign_Left)
+					.VAlign(VAlign_Center)
+					.Padding(FMargin(12.0, 0.0, 2.0, 0.0))
+					.AutoWidth()
+					[
+						SNew(SImage)
+						.Image(ViewModel->GetSelectorIcon(ParentState))
+						.ColorAndOpacity(FStyleColors::Foreground)
+						.Visibility(EVisibility::HitTestInvisible)
+					]
+					+ SHorizontalBox::Slot()
+					.HAlign(HAlign_Left)
+					.VAlign(VAlign_Center)
+					.AutoWidth()
+					[
+						SNew(STextBlock)
+						.Text(FText::FromName(ParentState->Name))
+						.TextStyle(&FStateTreeEditorStyle::Get().GetWidgetStyle<FTextBlockStyle>("Normal.Normal"))
+						.ColorAndOpacity(FStyleColors::White)
+						.Visibility(EVisibility::HitTestInvisible)
+					]
+				];
+
+			AddExtension(ParentStateNameTopExtension)
+				.SetLocation(EExtensionLocation::Top)
+				.SetAutoSize()
+				.SetPadding(FMargin(0.0, 0.0, 0.0, 4.0));
+		}
+
+		// Target state entry conditions in subheader depending on LOD
+		if (InLOD == EExtendableNodeLOD::Full)
+		{
+			if (const UStateTreeState* State = TransitionDirection == EStateTransitionDirection::In
+				? ViewModel->GetSourceStateFromTransition(Transition)
+				: ViewModel->GetTargetStateFromTransition(Transition))
+			{
+				TransitionStateBackgroundBrush = FStateCentricViewUtils::MakeStateSubnodeBackground(ViewModel.ToSharedRef(), State);
+
+				TSharedPtr<SVerticalBox> EnterConditionsBox;
+
+				TSharedRef<SWidget> EnterConditionsExtension = SNew(SOverlay)
+				+SOverlay::Slot()
+				[
+					SNew(SImage)
+					.Image(TransitionStateBackgroundBrush->GetSlateBrush())
+				]
+				+ SOverlay::Slot()
+				[
+					SAssignNew(EnterConditionsBox, SVerticalBox)
+				];
+
+				for (const FStateTreeEditorNode& EnterCondition : State->EnterConditions)
+				{
+					FText EnterConditionDesc = ViewModel->GetNodeDescription(EnterCondition, EStateTreeNodeFormatting::Text);
+
+					EnterConditionsBox->AddSlot()
+					.Padding(FMargin(8.0, 2.0, 8.0, 2.0))
+					.AutoHeight()
+					[
+						SNew(STextBlock)
+						.AutoWrapText(true)
+						.Text(EnterConditionDesc)
+					];
+				}
+
+				if (State->EnterConditions.Num() > 0)
+				{
+					SVerticalBox::FSlot& FirstSlot = EnterConditionsBox->GetSlot(0);
+					FirstSlot.SetPadding(FirstSlot.GetPadding() + FMargin(0.0, 12.0, 0.0, 0.0));
+
+					AddExtension(EnterConditionsExtension)
+						.SetLocation(EExtensionLocation::SubHeader)
+						.SetPadding(FMargin(0.0, -10.0, 0.0, 0.0));
+				}
+			}
+		}
+
+		// Small arrow used to indicate transition on left / right depending on direction
+		{
+			FScopedMainNodeLayoutInfo ScopedMainLayoutInfo = SetScopedMainNodeLayoutInfo()
+				.SetFillSize(1.0f)
+				.SetPadding(FMargin(0.0));
+
+			TSharedRef<SWidget> TransitionArrowExtension =
+				SNew(SButton)
+				.NormalPaddingOverride(FMargin(2.0, 4.0))
+				.PressedPaddingOverride(FMargin(2.0, 4.0))
+				.ContentPadding(FMargin(0.0f))
+				[
+					SNew(SHorizontalBox)
+						+ SHorizontalBox::Slot()
+						[
+							SNew(SImage)
+								.Image(FAppStyle::Get().GetBrush("Icons.ChevronRight"))
+								.DesiredSizeOverride(FVector2D(10, 24))
+								.ColorAndOpacity(FStyleColors::Foreground)
+						]
+						+ SHorizontalBox::Slot()
+						.Padding(FMargin(-8.0, 0.0, 0.0, 0.0))
+						[
+							SNew(SImage)
+								.Image(FAppStyle::Get().GetBrush("Icons.ChevronRight"))
+								.DesiredSizeOverride(FVector2D(10, 24))
+								.ColorAndOpacity(FStyleColors::Foreground)
+						]
+				];
+
+			EExtensionLocation TransitionArrowLocation = TransitionDirection == EStateTransitionDirection::In
+				? EExtensionLocation::Right
+				: EExtensionLocation::Left;
+
+			FMargin TransitionArrowPadding = TransitionDirection == EStateTransitionDirection::In
+				? FMargin(6.0, 0.0, -6.0, 0.0)
+				: FMargin(-6.0, 0.0, 6.0, 0.0);
+				
+			AddExtension(TransitionArrowExtension)
+				.SetLocation(TransitionArrowLocation)
+				.SetVAlign(VAlign_Center)
+				.SetAutoSize()
+				.SetPadding(TransitionArrowPadding);
+		}
+
+		// Generic background extension, padded to not overlap with the small transition arrow
+		{
+			FScopedMainNodeLayoutInfo ScopedMainLayoutInfo = SetScopedMainNodeLayoutInfo()
+				.SetPadding(FMargin(4.0));
+
+			FMargin TransitionArrowPadding = TransitionDirection == EStateTransitionDirection::In
+				? FMargin(0.0, 0.0, 16.0, 0.0)
+				: FMargin(16.0, 0.0, 0.0, 0.0);
+
+			const FSlateBrush* SubNodeBackgroundBrush = FStateCentricViewUtils::GetSubNodeBackground();
+
+			TSharedPtr<SWidget> SubNodeBackground = SNew(SBorder)
+				.BorderImage(SubNodeBackgroundBrush)
+				.Visibility(EVisibility::HitTestInvisible);
+
+			AddExtension(SubNodeBackground.ToSharedRef())
+				.SetLocation(EExtensionLocation::Background)
+				.SetAddLocationBehavior(EExtensionAddLocationBehavior::OverlayEntireNode)
+				.SetPadding(TransitionArrowPadding);
+		}
+	}
+}
+
+TSharedRef<SWidget> STransitionNode::Construct_GenerateNodeHeader(const EExtendableNodeLOD InLOD)
+{
+	NodeBodyPadding = FMargin(0.0, 0.0, 0.0, 0.0);
+
+	FStateTreeTransition* Transition = ViewModel->GetMutableTransitionByID(TransitionID);
+	if (Transition && ViewState.IsValid())
+	{
+		NodeBodyPadding = FMargin(0.0, 2.0, 0.0, 0.0);
+
+		if (const UStateTreeState* State = TransitionDirection == EStateTransitionDirection::In
+			? ViewModel->GetSourceStateFromTransition(Transition)
+			: ViewModel->GetTargetStateFromTransition(Transition))
+		{
+			TransitionStateBrush = FStateCentricViewUtils::MakeStateSubnodeBackgroundOutlined(ViewModel.ToSharedRef(), State);
+
+			// @TODO move this logic below into FStateCentricViewUtils
+
+			return SNew(SOverlay)
+				+ SOverlay::Slot()
+				[
+					SNew(SImage)
+					.Image(TransitionStateBrush->GetSlateBrush())
+				]
+				+ SOverlay::Slot()
+				[
+					// Draw twice to force our background on top of subheader 
+					SNew(SImage)
+					.Image(TransitionStateBrush->GetSlateBrush())
+				]
+				+ SOverlay::Slot()
+				[
+					FStateCentricViewUtils::GenerateGotoStateButton(InLOD, ViewModel.ToSharedRef(), State)
+				]
+				+ SOverlay::Slot()
+				[
+					SNew(SHorizontalBox)
+					+ SHorizontalBox::Slot()
+					.HAlign(HAlign_Left)
+					.VAlign(VAlign_Center)
+					.Padding(FMargin(12.0, 0.0, 2.0, 0.0))
+					.AutoWidth()
+					[
+						SNew(SImage)
+						.Image(ViewModel->GetSelectorIcon(State))
+						.ColorAndOpacity(FStyleColors::Foreground)
+						.Visibility(EVisibility::HitTestInvisible)
+					]
+					+ SHorizontalBox::Slot()
+					.HAlign(HAlign_Left)
+					.VAlign(VAlign_Center)
+					.AutoWidth()
+					[
+						SNew(STextBlock)
+						.Text(FText::FromName(State->Name))
+						.TextStyle(&FStateTreeEditorStyle::Get().GetWidgetStyle<FTextBlockStyle>("Normal.Normal"))
+						.ColorAndOpacity(FStyleColors::White)
+						.Visibility(EVisibility::HitTestInvisible)
+					]
+				];
+		}
+	}
+
+	return SNullWidget::NullWidget;
+}
+
+TSharedRef<SWidget> STransitionNode::Construct_GenerateNodeBody(const EExtendableNodeLOD InLOD)
+{
+	if (InLOD == EExtendableNodeLOD::Full)
+	{
+		FStateTreeTransition* Transition = ViewModel->GetMutableTransitionByID(TransitionID);
+		if (UStateTreeState* StateOwningTransition = TransitionDirection == EStateTransitionDirection::In
+			? const_cast<UStateTreeState*>(ViewModel->GetSourceStateFromTransition(Transition))
+			: ViewState.Get())
+		{
+			// Struct view doesn't work since the customization is only on the object.
+			// So show the object and filter to the transition
+
+			// @TODO: This seems slow, revist & profile
+
+			FPropertyEditorModule& PropertyEditorModule = FModuleManager::GetModuleChecked<FPropertyEditorModule>("PropertyEditor");
+			FDetailsViewArgs DetailsViewArgs;
+			DetailsViewArgs.NameAreaSettings = FDetailsViewArgs::HideNameArea;
+			DetailsViewArgs.InitialCategories = TSet<FName>{ "Transitions" };
+			DetailsViewArgs.bShowScrollBar = false;
+			DetailsViewArgs.bShowOptions = false;
+			DetailsViewArgs.bAllowSearch = false;
+			DetailsViewArgs.bShowObjectLabel = false;
+			DetailsViewArgs.bShowHeaders = false;
+			DetailsViewArgs.ShouldForceHideProperty.BindLambda([ViewModel = ViewModel, TransitionID = TransitionID, StateOwningTransition = StateOwningTransition](const TSharedRef<FPropertyNode>& InNode)
+			{
+				const FStructProperty* TransitionPropertyStruct = nullptr;
+
+				if (const FStructProperty* PropertyStruct = CastField<FStructProperty>(InNode->GetProperty()))
+				{
+					if (PropertyStruct->Struct == FStateTreeTransition::StaticStruct())
+					{
+						TransitionPropertyStruct = PropertyStruct;
+					}
+				}
+
+				const FPropertyNode* TransitionPropertyNode = TransitionPropertyStruct ? &InNode.Get() : InNode->GetParentNode();
+
+				// Search parents for transition property
+				while (TransitionPropertyNode && !TransitionPropertyStruct)
+				{
+					if (const FProperty* ParentProperty = TransitionPropertyNode->GetProperty())
+					{
+						if (const FStructProperty* ParentPropertyStruct = CastField<FStructProperty>(ParentProperty))
+						{
+							if (ParentPropertyStruct->Struct == FStateTreeTransition::StaticStruct())
+							{
+								TransitionPropertyStruct = ParentPropertyStruct;
+							}
+						}
+					}
+
+					if (!TransitionPropertyStruct)
+					{
+						TransitionPropertyNode = TransitionPropertyNode->GetParentNode();
+					}
+				}
+
+				if (TransitionPropertyStruct && TransitionPropertyNode)
+				{
+					FStateTreeTransition TransitionValue; 
+					if (PropertyPathHelpers::GetPropertyValue(StateOwningTransition, TransitionPropertyNode->GetPropertyPath(), OUT TransitionValue))
+					{
+						return TransitionValue.ID != TransitionID;
+					}
+				}
+
+				return true;
+			});
+
+			TSharedRef<IDetailsView> StateDetailsView = PropertyEditorModule.CreateDetailView(DetailsViewArgs);
+			StateDetailsView->SetObject(StateOwningTransition);
+
+			return StateDetailsView;
+		}
+	}
+
+	return SNullWidget::NullWidget;
+}
+
+
+//////////////////////////////////////////////////////////////////////////
+// SInTransitionNode
+
+const FName SInTransitionNode::NodeID(TEXT("StateCentricView.Transition.InNode"));
+
+FName SInTransitionNode::GetNodeName() const
+{
+	return NodeID;
+}
+
+void SInTransitionNode::Construct_InitializeNodePersistentPreLOD()
+{
+	TransitionDirection = EStateTransitionDirection::In;
+}
+
+
+//////////////////////////////////////////////////////////////////////////
+// SMainStateOutTransitionExtensionNode
+
+const FName SOutTransitionNode::NodeID(TEXT("StateCentricView.Transition.OutNode"));
+
+FName SOutTransitionNode::GetNodeName() const
+{
+	return NodeID;
+}
+
+void SOutTransitionNode::Construct_InitializeNodePersistentPreLOD()
+{
+	TransitionDirection = EStateTransitionDirection::Out;
+}
+
+
+//////////////////////////////////////////////////////////////////////////
+
+} // namespace UE::StateTree::Editor::StateCentricView
+
+#undef LOCTEXT_NAMESPACE
+
diff --git a/Engine/Plugins/Runtime/StateTree/Source/StateTreeEditorModule/Private/StateCentricView/View/STransitionNode.h b/Engine/Plugins/Runtime/StateTree/Source/StateTreeEditorModule/Private/StateCentricView/View/STransitionNode.h
new file mode 100644
--- /dev/null
+++ b/Engine/Plugins/Runtime/StateTree/Source/StateTreeEditorModule/Private/StateCentricView/View/STransitionNode.h
@@ -0,0 +1,116 @@
+// Copyright Epic Games, Inc. All Rights Reserved.
+
+#pragma once
+
+#include "CoreMinimal.h"
+
+#include "StateCentricView/View/SStateTreeExtendableNode.h"
+#include "StateTreeState.h"
+
+struct FSlateBrush;
+struct FStateTreeTransition;
+
+class IStructureDetailsView;
+class FStateTreeViewModel;
+class SButton;
+class UStateTreeState;
+
+enum class EStateTransitionDirection : uint8;
+
+namespace UE::StateTree::Editor::StateCentricView
+{
+
+
+/**
+ * Experimental. Node used to show a single state transition
+ */
+class STransitionNode : public SStateTreeExtendableNode
+{
+
+public:
+
+	void Construct(const SExtendableNode::FArguments& InArgs
+		, TSharedRef<FStateTreeViewModel> InViewModel
+		, TNotNull<UStateTreeState*> InViewState
+		, TNotNull<FStateTreeTransition*> InTransition);
+
+protected:
+
+	//~ Begin SExtendableNode Interface
+	virtual void Construct_GenerateBuiltInExtensions(const EExtendableNodeLOD InLOD) override;
+	virtual TSharedRef<SWidget> Construct_GenerateNodeHeader(const EExtendableNodeLOD InLOD) override;
+	virtual TSharedRef<SWidget> Construct_GenerateNodeBody(const EExtendableNodeLOD InLOD) override;
+	//~ End SExtendableNode Interface
+
+protected:
+
+	/** 
+	 * Transition we represent, guaranteed to exist as long as this node does via value change callbacks
+	 * @TODO: Make that guarantee true.
+	 */
+	FGuid TransitionID = FGuid();
+
+	/** If we are showing in / out transitions. Init in child classes. */
+	EStateTransitionDirection TransitionDirection = EStateTransitionDirection::Out;
+
+private:
+
+	TSharedPtr<FDeferredCleanupSlateBrush> TransitionParentStateBrush;
+	TSharedPtr<FDeferredCleanupSlateBrush> TransitionStateBrush;
+	TSharedPtr<FDeferredCleanupSlateBrush> TransitionStateBackgroundBrush;
+
+	TSharedPtr<IStructureDetailsView> TransitionDetailsView;
+};
+
+
+/**
+ * Experimental. Specialization of STransitionNode for 'In' transitions
+ */
+class SInTransitionNode : public STransitionNode
+{
+
+public:
+
+	static const FName NodeID;
+
+public:
+
+	//~ Begin SExtendableNode Interface
+	virtual FName GetNodeName() const override;
+	//~ End SExtendableNode Interface
+
+protected:
+
+	//~ Begin SExtendableNode Interface
+	virtual void Construct_InitializeNodePersistentPreLOD() override;
+	//~ End SExtendableNode Interface
+
+};
+
+
+/**
+ * Experimental. Specialization of STransitionNode for 'Out' transitions
+ */
+class SOutTransitionNode : public STransitionNode
+{
+
+public:
+
+	static const FName NodeID;
+
+public:
+
+	//~ Begin SExtendableNode Interface
+	virtual FName GetNodeName() const override;
+	//~ End SExtendableNode Interface
+
+protected:
+
+	//~ Begin SExtendableNode Interface
+	virtual void Construct_InitializeNodePersistentPreLOD() override;
+	//~ End SExtendableNode Interface
+
+};
+
+
+} // namespace UE::StateTree::Editor::StateCentricView
\ No newline at end of file
diff --git a/Engine/Plugins/Runtime/StateTree/Source/StateTreeEditorModule/Private/StateCentricView/View/StateCentricViewUtils.cpp b/Engine/Plugins/Runtime/StateTree/Source/StateTreeEditorModule/Private/StateCentricView/View/StateCentricViewUtils.cpp
new file mode 100644
--- /dev/null
+++ b/Engine/Plugins/Runtime/StateTree/Source/StateTreeEditorModule/Private/StateCentricView/View/StateCentricViewUtils.cpp
@@ -0,0 +1,618 @@
+// Copyright Epic Games, Inc. All Rights Reserved.
+
+#include "StateCentricView/View/StateCentricViewUtils.h"
+
+#include "StateCentricView/StateCentricViewSettings.h"
+#include "StateCentricView/View/SStateTreeExtendableNode.h"
+#include "StateCentricView/View/STransitionNode.h"
+#include "StateTreeEditorStyle.h"
+
+#include "Styling/SlateBrush.h"
+#include "Styling/StyleColors.h"
+#include "Widgets/Images/SImage.h"
+#include "Widgets/Input/SButton.h"
+#include "Widgets/Layout/SExpandableArea.h"
+#include "Widgets/Text/STextBlock.h"
+
+#define LOCTEXT_NAMESPACE "StateCentricViewUtils"
+
+// These CVars exist as a discussion point. Should have a very short lifespan. Probably week or so.
+
+bool bEnableDepthBasedOutline = true;
+FAutoConsoleVariableRef CVarEnableDepthBasedOutline(
+	TEXT("StateTree.Editor.Experimental.EnableDepthBasedOutline"),
+	bEnableDepthBasedOutline,
+	TEXT("Set true to change outline style in state centric view when using hierarchy colors")
+);
+
+bool bDisableDarkenHeader = false;
+FAutoConsoleVariableRef CVarDisableDarkenHeader(
+	TEXT("StateTree.Editor.Experimental.DisableDarkenHeader"),
+	bDisableDarkenHeader,
+	TEXT("Set true to darken the headers in state centric")
+);
+
+namespace UE::StateTree::Editor::StateCentricView
+{
+
+TSharedRef<SWidget> FStateCentricViewUtils::GenerateTransitionsViewWidget(const EExtendableNodeLOD LOD
+	, TSharedRef<FStateTreeViewModel> ViewModel
+	, TNotNull<UStateTreeState*> ViewState
+	, const EStateTransitionDirection TransitionDirection)
+{
+	UStateTreeState* CentralState = ViewModel->GetViewState();
+
+	if (LOD == EExtendableNodeLOD::Hidden || LOD == EExtendableNodeLOD::Collapsed)
+	{
+		return SNullWidget::NullWidget;
+	}
+	else if (ensure(LOD == EExtendableNodeLOD::Minimal || LOD == EExtendableNodeLOD::Full))
+	{
+		TArray<FStateTreeTransition*> Transitions;
+		if (TransitionDirection == EStateTransitionDirection::In)
+		{
+			Transitions = ViewModel->GetStateInTransitions(ViewState);
+		}
+		else if (ensure(TransitionDirection == EStateTransitionDirection::Out))
+		{
+			Transitions = ViewModel->GetStateOutTransitions(ViewState);
+		}
+
+		TSharedRef<SVerticalBox> TransitionListBody = SNew(SVerticalBox);
+		for (FStateTreeTransition* Transition : Transitions)
+		{
+			TSharedPtr<SWidget> TransitionNodeWidget;
+			if (TransitionDirection == EStateTransitionDirection::In)
+			{
+				TransitionNodeWidget = SNew(SInTransitionNode, ViewModel, ViewState, Transition)
+					.Config({ .InitialLOD = LOD });
+			}
+			else if (ensure(TransitionDirection == EStateTransitionDirection::Out))
+			{
+				TransitionNodeWidget = SNew(SOutTransitionNode, ViewModel, ViewState, Transition)
+					.Config({ .InitialLOD = LOD });
+			}
+
+			// @TODO, CVar / settings here
+			if (ViewState == CentralState && UStateCentricViewSettings::GetSchemaSettingsDefaultFallback(&ViewModel.Get()).bJustifyCentralViewTransitions)
+			{
+				TransitionListBody->AddSlot()
+				.FillContentHeight(1.0f)
+				.VAlign(VAlign_Center)
+				.Padding(FMargin(8.0, 6.0, 8.0, 6.0))
+				[
+					TransitionNodeWidget.ToSharedRef()
+				];
+			}
+			else
+			{
+				TransitionListBody->AddSlot()
+				.AutoHeight()
+				.Padding(FMargin(8.0, 6.0, 8.0, 6.0))
+				[
+					TransitionNodeWidget.ToSharedRef()
+				];
+			}
+		}
+		
+		if (Transitions.Num() > 0)
+		{
+			SVerticalBox::FSlot& FirstSlot = TransitionListBody->GetSlot(0);
+			SVerticalBox::FSlot& LastSlot = TransitionListBody->GetSlot(TransitionListBody->NumSlots() - 1);
+
+			FirstSlot.SetPadding(FirstSlot.GetPadding() + FMargin(0.0, 6.0, 0.0, 0.0));
+			LastSlot.SetPadding(LastSlot.GetPadding() + FMargin(0.0, 0.0, 0.0, 6.0));
+		}
+
+		return TransitionListBody;
+	}
+
+	return SNullWidget::NullWidget;
+}
+
+TSharedRef<SButton> FStateCentricViewUtils::GenerateGotoStateButton(const EExtendableNodeLOD LOD, TSharedRef<FStateTreeViewModel> ViewModel, TNotNull<const UStateTreeState*> ViewState)
+{
+	TWeakPtr<FStateTreeViewModel> WeakViewModel = ViewModel;
+	auto SelectStateCallback = [WeakViewModel, ViewState](const FGeometry&, const FPointerEvent& )
+	{
+		if (TSharedPtr<FStateTreeViewModel> ViewModelPinned = WeakViewModel.Pin())
+		{
+			ViewModelPinned->SetSelection(const_cast<UStateTreeState*>(NotNullGet(ViewState)));
+		}
+
+		return FReply::Handled();
+	};
+
+	TSharedRef<SButton> Result = SNew(SButton)
+		.ButtonStyle(&FStateTreeEditorStyle::Get().GetWidgetStyle<FButtonStyle>("StateTree.CentricView.HoverHighlight"));
+
+	Result->SetOnMouseDoubleClick(FPointerEventHandler::CreateLambda(SelectStateCallback));
+
+	return Result;
+}
+
+void FStateCentricViewUtils::AddChildStateFootnoteExtensions(const EExtendableNodeLOD InLOD
+	, TNotNull<SStateTreeExtendableNode*> InSourceNode
+	, TArray<TSharedPtr<FDeferredCleanupSlateBrush>>& OutChildStateBrushes)
+{
+	TSharedPtr<FStateTreeViewModel> ViewModel = InSourceNode->ViewModel;
+	TWeakObjectPtr<UStateTreeState> ViewState = InSourceNode->ViewState;
+
+	OutChildStateBrushes.Reset();
+
+	if (InLOD == EExtendableNodeLOD::Full)
+	{
+		// @TODO: This should be be a details entry maybe?
+
+		TArray<UStateTreeState*> ChildStates = ViewModel->GetStateChildren(ViewState.Get());
+
+		if (ChildStates.IsEmpty())
+		{
+			return;
+		}
+
+		TSharedRef<SWidget> ChildStateAreaHeader =
+			SNew(SHorizontalBox)
+			+ SHorizontalBox::Slot()
+			.HAlign(HAlign_Left)
+			.VAlign(VAlign_Center)
+			.AutoWidth()
+			.Padding(FMargin(2.0, 5.0, 0.0, 5.0))
+			[
+				SNew(SImage)
+					.Image(FAppStyle::GetBrush("LevelEditor.Tabs.Outliner"))
+					.ColorAndOpacity(FStyleColors::Foreground)
+			]
+			+ SHorizontalBox::Slot()
+			.HAlign(HAlign_Left)
+			.VAlign(VAlign_Center)
+			.AutoWidth()
+			.Padding(FMargin(4.0, 0.0, 0.0, 0.0))
+			[
+				SNew(STextBlock)
+					.Text(LOCTEXT("StateCentricChildHeader", "Children"))
+					.Font(FAppStyle::GetFontStyle("DetailsView.CategoryFontStyle"))
+					.TextStyle(FAppStyle::Get(), "DetailsView.CategoryTextStyle")
+					.ColorAndOpacity(FStyleColors::Foreground)
+			];
+
+		TSharedRef<SVerticalBox> ChildStateAreaBody = SNew(SVerticalBox);
+
+		for (UStateTreeState* ChildState : ChildStates)
+		{	
+			TSharedPtr<FDeferredCleanupSlateBrush> ChildStateBrush = FStateCentricViewUtils::MakeStateSubnodeBackgroundOutlined(ViewModel.ToSharedRef(), ChildState);
+
+			TSharedRef<SWidget> ChildStateBottomExtension = SNew(SOverlay)
+				+ SOverlay::Slot()
+				[
+					SNew(SImage)
+					.Image(ChildStateBrush->GetSlateBrush())
+				]
+				+ SOverlay::Slot()
+				[
+					FStateCentricViewUtils::GenerateGotoStateButton(InLOD, ViewModel.ToSharedRef(), ChildState)
+				]
+				+ SOverlay::Slot()
+				[
+					SNew(SHorizontalBox)
+					+ SHorizontalBox::Slot()
+					.HAlign(HAlign_Left)
+					.VAlign(VAlign_Center)
+					.Padding(FMargin(12.0, 0.0, 2.0, 0.0))
+					.AutoWidth()
+					[
+						SNew(SImage)
+						.Image(ViewModel->GetSelectorIcon(ChildState))
+						.ColorAndOpacity(FStyleColors::Foreground)
+						.Visibility(EVisibility::HitTestInvisible)
+					]
+					+ SHorizontalBox::Slot()
+					.HAlign(HAlign_Left)
+					.VAlign(VAlign_Center)
+					.AutoWidth()
+					[
+						SNew(STextBlock)
+						.Text(FText::FromName(ChildState->Name))
+						.TextStyle(&FStateTreeEditorStyle::Get().GetWidgetStyle<FTextBlockStyle>("Normal.Normal"))
+						.ColorAndOpacity(FStyleColors::White)
+						.Visibility(EVisibility::HitTestInvisible)
+					]
+				];
+
+			OutChildStateBrushes.Add(ChildStateBrush);
+
+			ChildStateAreaBody->AddSlot()
+			.Padding(FMargin(16.0, 8.0, 16.0, 0.0))
+			.AutoHeight()
+			[
+				ChildStateBottomExtension
+			];
+		}
+	
+
+		TSharedRef<SExpandableArea> ChildStateArea = SNew(SExpandableArea)
+			.InitiallyCollapsed(false)
+			.BorderImage(FAppStyle::Get().GetBrush("Brushes.Header"))
+			.BodyBorderImage(FAppStyle::Get().GetBrush("Brushes.Recessed"))
+			.HeaderContent()
+			[
+				ChildStateAreaHeader
+			]
+			.BodyContent()
+			[
+				ChildStateAreaBody
+			];
+
+		InSourceNode->AddExtension(ChildStateArea)
+			.SetLocation(EExtensionLocation::Footnote)
+			.SetPadding(FMargin(0.0, 0.0, 0.0, 4.0))
+			.SetAutoSize();
+	}
+}
+
+TSharedPtr<FDeferredCleanupSlateBrush> FStateCentricViewUtils::MakeStateHalfTopOutline(TSharedRef<FStateTreeViewModel> ViewModel, TNotNull<const UStateTreeState*> State)
+{
+	const FSlateBrush* BaseBrush = FStateTreeEditorStyle::Get().GetBrush("StateTree.CentricView.Parent.HalfTopOutline");
+
+	FSlateBrush Result = *BaseBrush;
+
+	// Use tint with 0 alpha due to slate rendering bug where alpha still matters.
+	Result.TintColor = ViewModel->GetStateColor(State, EStateTreeViewModelColorAdjustment::Bright).CopyWithNewOpacity(0.0f);
+	Result.OutlineSettings.Color = ViewModel->GetStateColor(State, EStateTreeViewModelColorAdjustment::Bright);
+
+	if (UStateCentricViewSettings::GetSchemaSettingsDefaultFallback(&ViewModel.Get()).bEnableHierarchyDrivenColors)
+	{
+		Result.TintColor = ViewModel->GetHierarchyBasedHighlightColor(State, EStateTreeViewModelColorAdjustment::Bright).CopyWithNewOpacity(0.0f);
+		Result.OutlineSettings.Color = ViewModel->GetHierarchyBasedHighlightColor(State, EStateTreeViewModelColorAdjustment::Bright);
+	}
+
+	return FDeferredCleanupSlateBrush::CreateBrush(Result);
+}
+
+TSharedPtr<FDeferredCleanupSlateBrush> FStateCentricViewUtils::MakeStateHalfTopOutlineMuted(TSharedRef<FStateTreeViewModel> ViewModel, TNotNull<const UStateTreeState*> State)
+{
+	const FSlateBrush* BaseBrush = FStateTreeEditorStyle::Get().GetBrush("StateTree.CentricView.Parent.HalfTopOutline");
+
+	FSlateBrush Result = *BaseBrush;
+
+	// Use tint with 0 alpha due to slate rendering bug where alpha still matters.
+	Result.TintColor = ViewModel->GetStateColor(State).CopyWithNewOpacity(0.0f);
+	Result.OutlineSettings.Color = ViewModel->GetStateColor(State);
+
+	if (UStateCentricViewSettings::GetSchemaSettingsDefaultFallback(&ViewModel.Get()).bEnableHierarchyDrivenColors)
+	{
+		Result.TintColor = ViewModel->GetHierarchyBasedHighlightColor(State).CopyWithNewOpacity(0.0f);
+		Result.OutlineSettings.Color = ViewModel->GetHierarchyBasedHighlightColor(State);
+
+		if (bEnableDepthBasedOutline)
+		{
+			Result.OutlineSettings.Color = ViewModel->GetHierarchyBasedBackgroundColor(State, EStateTreeViewModelColorAdjustment::Bright);
+		}
+	}
+
+	return FDeferredCleanupSlateBrush::CreateBrush(Result);
+}
+
+TSharedPtr<FDeferredCleanupSlateBrush> FStateCentricViewUtils::MakeStateHalfTopOutlineCollapsed(TSharedRef<FStateTreeViewModel> ViewModel, TNotNull<const UStateTreeState*> State)
+{
+	const FSlateBrush* BaseBrush = FStateTreeEditorStyle::Get().GetBrush("StateTree.CentricView.Parent.HalfTopOutline");
+
+	FSlateBrush Result = *BaseBrush;
+
+	// Use tint with 0 alpha due to slate rendering bug where alpha still matters.
+	Result.TintColor = ViewModel->GetStateColor(State, EStateTreeViewModelColorAdjustment::Dark).CopyWithNewOpacity(0.0f);
+	Result.OutlineSettings.Color = ViewModel->GetStateColor(State, EStateTreeViewModelColorAdjustment::Dark);
+
+	if (UStateCentricViewSettings::GetSchemaSettingsDefaultFallback(&ViewModel.Get()).bEnableHierarchyDrivenColors)
+	{
+		Result.TintColor = ViewModel->GetHierarchyBasedHighlightColor(State, EStateTreeViewModelColorAdjustment::Dark).CopyWithNewOpacity(0.0f);
+		Result.OutlineSettings.Color = ViewModel->GetHierarchyBasedHighlightColor(State);
+
+		if (bEnableDepthBasedOutline)
+		{
+			Result.OutlineSettings.Color = ViewModel->GetHierarchyBasedBackgroundColor(State, EStateTreeViewModelColorAdjustment::Bright);
+		}
+	}
+
+	return FDeferredCleanupSlateBrush::CreateBrush(Result);
+}
+
+TSharedPtr<FDeferredCleanupSlateBrush> FStateCentricViewUtils::MakeStateOutline(TSharedRef<FStateTreeViewModel> ViewModel, TNotNull<const UStateTreeState*> State)
+{
+	const FSlateBrush* BaseBrush = State->Children.IsEmpty()
+		? FStateTreeEditorStyle::Get().GetBrush("StateTree.CentricView.Leaf.Outline")
+		: FStateTreeEditorStyle::Get().GetBrush("StateTree.CentricView.Parent.Outline");
+
+	FSlateBrush Result = *BaseBrush;
+
+	// Use tint with 0 alpha due to slate rendering bug where alpha still matters.
+	Result.TintColor = ViewModel->GetStateColor(State, EStateTreeViewModelColorAdjustment::Bright).CopyWithNewOpacity(0.0f);
+	Result.OutlineSettings.Color = ViewModel->GetStateColor(State, EStateTreeViewModelColorAdjustment::Bright);
+
+	if (UStateCentricViewSettings::GetSchemaSettingsDefaultFallback(&ViewModel.Get()).bEnableHierarchyDrivenColors)
+	{
+		Result.TintColor = ViewModel->GetHierarchyBasedHighlightColor(State, EStateTreeViewModelColorAdjustment::Bright).CopyWithNewOpacity(0.0f);
+		Result.OutlineSettings.Color = ViewModel->GetHierarchyBasedHighlightColor(State);
+
+		if (bEnableDepthBasedOutline)
+		{
+			Result.OutlineSettings.Color = ViewModel->GetHierarchyBasedBackgroundColor(State, EStateTreeViewModelColorAdjustment::Bright);
+		}
+	}
+
+	return FDeferredCleanupSlateBrush::CreateBrush(Result);
+}
+
+TSharedPtr<FDeferredCleanupSlateBrush> FStateCentricViewUtils::MakeStateOutlineMuted(TSharedRef<FStateTreeViewModel> ViewModel, TNotNull<const UStateTreeState*> State)
+{
+	const FSlateBrush* BaseBrush = State->Children.IsEmpty()
+		? FStateTreeEditorStyle::Get().GetBrush("StateTree.CentricView.Leaf.Outline")
+		: FStateTreeEditorStyle::Get().GetBrush("StateTree.CentricView.Parent.Outline");
+
+	FSlateBrush Result = *BaseBrush;
+
+	// Use tint with 0 alpha due to slate rendering bug where alpha still matters.
+	Result.TintColor = ViewModel->GetStateColor(State).CopyWithNewOpacity(0.0f);
+	Result.OutlineSettings.Color = ViewModel->GetStateColor(State);
+
+	if (UStateCentricViewSettings::GetSchemaSettingsDefaultFallback(&ViewModel.Get()).bEnableHierarchyDrivenColors)
+	{
+		Result.TintColor = ViewModel->GetHierarchyBasedHighlightColor(State).CopyWithNewOpacity(0.0f);
+		Result.OutlineSettings.Color = ViewModel->GetHierarchyBasedHighlightColor(State);
+
+		if (bEnableDepthBasedOutline)
+		{
+			Result.OutlineSettings.Color = ViewModel->GetHierarchyBasedBackgroundColor(State, EStateTreeViewModelColorAdjustment::Bright);
+		}
+	}
+
+	return FDeferredCleanupSlateBrush::CreateBrush(Result);
+}
+
+TSharedPtr<FDeferredCleanupSlateBrush> FStateCentricViewUtils::MakeStateHeaderBackground(TSharedRef<FStateTreeViewModel> ViewModel, TNotNull<const UStateTreeState*> State, TOptional<FVector2f> ImageSize)
+{
+	const FSlateBrush* BaseBrush = State->Children.IsEmpty()
+		? FStateTreeEditorStyle::Get().GetBrush("StateTree.CentricView.Leaf.Header.Background")
+		: FStateTreeEditorStyle::Get().GetBrush("StateTree.CentricView.Parent.Header.Background");
+
+	FSlateBrush Result = *BaseBrush;
+
+	Result.TintColor = ViewModel->GetStateColor(State, EStateTreeViewModelColorAdjustment::Dark);
+	Result.OutlineSettings.Color = ViewModel->GetStateColor(State, EStateTreeViewModelColorAdjustment::Dark);
+	Result.ImageSize = ImageSize.Get(Result.ImageSize);
+
+	if (UStateCentricViewSettings::GetSchemaSettingsDefaultFallback(&ViewModel.Get()).bEnableHierarchyDrivenColors)
+	{
+		Result.TintColor = ViewModel->GetHierarchyBasedBackgroundColor(State, EStateTreeViewModelColorAdjustment::Dark);
+		Result.OutlineSettings.Color = ViewModel->GetHierarchyBasedBackgroundColor(State, EStateTreeViewModelColorAdjustment::Dark);
+
+		if (bDisableDarkenHeader)
+		{
+			Result.TintColor = ViewModel->GetHierarchyBasedBackgroundColor(State);
+			Result.OutlineSettings.Color = ViewModel->GetHierarchyBasedBackgroundColor(State);
+		}
+	}
+
+	return FDeferredCleanupSlateBrush::CreateBrush(Result);
+}
+
+TSharedPtr<FDeferredCleanupSlateBrush> FStateCentricViewUtils::MakeStateHeaderBackgroundMuted(TSharedRef<FStateTreeViewModel> ViewModel, TNotNull<const UStateTreeState*> State)
+{
+	const FSlateBrush* BaseBrush = State->Children.IsEmpty()
+		? FStateTreeEditorStyle::Get().GetBrush("StateTree.CentricView.Leaf.Header.Background")
+		: FStateTreeEditorStyle::Get().GetBrush("StateTree.CentricView.Parent.Header.Background.Muted");
+
+	FSlateBrush Result = *BaseBrush;
+
+	Result.TintColor = State->Children.IsEmpty() ? ViewModel->GetStateColor(State, EStateTreeViewModelColorAdjustment::Dark) : Result.TintColor;
+	Result.OutlineSettings.Color = State->Children.IsEmpty() ? ViewModel->GetStateColor(State, EStateTreeViewModelColorAdjustment::Dark) : Result.TintColor;
+
+	if (UStateCentricViewSettings::GetSchemaSettingsDefaultFallback(&ViewModel.Get()).bEnableHierarchyDrivenColors)
+	{
+		Result.TintColor = State->Children.IsEmpty() ? ViewModel->GetHierarchyBasedBackgroundColor(State, EStateTreeViewModelColorAdjustment::Dark) : Result.TintColor;
+		Result.OutlineSettings.Color = State->Children.IsEmpty() ? ViewModel->GetHierarchyBasedBackgroundColor(State, EStateTreeViewModelColorAdjustment::Dark) : Result.TintColor;
+
+		if (bDisableDarkenHeader)
+		{
+			Result.TintColor = State->Children.IsEmpty() ? ViewModel->GetHierarchyBasedBackgroundColor(State) : Result.TintColor;
+			Result.OutlineSettings.Color = State->Children.IsEmpty() ? ViewModel->GetHierarchyBasedBackgroundColor(State) : Result.TintColor;
+		}
+	}
+
+	return FDeferredCleanupSlateBrush::CreateBrush(Result);
+}
+
+TSharedPtr<FDeferredCleanupSlateBrush> FStateCentricViewUtils::MakeStateHeaderBackgroundCollapsed(TSharedRef<FStateTreeViewModel> ViewModel, TNotNull<const UStateTreeState*> State)
+{
+	const FSlateBrush* BaseBrush = State->Children.IsEmpty()
+		? FStateTreeEditorStyle::Get().GetBrush("StateTree.CentricView.Leaf.Header.Background")
+		: FStateTreeEditorStyle::Get().GetBrush("StateTree.CentricView.Parent.Header.Background.Muted");
+
+	FSlateBrush Result = *BaseBrush;
+
+	Result.TintColor = ViewModel->GetStateColor(State, EStateTreeViewModelColorAdjustment::Darker);
+	Result.OutlineSettings.Color = ViewModel->GetStateColor(State, EStateTreeViewModelColorAdjustment::Darker);
+
+	if (UStateCentricViewSettings::GetSchemaSettingsDefaultFallback(&ViewModel.Get()).bEnableHierarchyDrivenColors)
+	{
+		Result.TintColor = ViewModel->GetHierarchyBasedBackgroundColor(State, EStateTreeViewModelColorAdjustment::Dark);
+		Result.OutlineSettings.Color = ViewModel->GetHierarchyBasedBackgroundColor(State, EStateTreeViewModelColorAdjustment::Dark);
+
+		if (bDisableDarkenHeader)
+		{
+			Result.TintColor = ViewModel->GetHierarchyBasedBackgroundColor(State);
+			Result.OutlineSettings.Color = ViewModel->GetHierarchyBasedBackgroundColor(State);
+		}
+	}
+
+	return FDeferredCleanupSlateBrush::CreateBrush(Result);
+}
+
+TSharedPtr<FDeferredCleanupSlateBrush> FStateCentricViewUtils::MakeStateBackgroundMutedOutlined(TSharedRef<FStateTreeViewModel> ViewModel, TNotNull<const UStateTreeState*> State)
+{
+	const FSlateBrush* BaseBrush = State->Children.IsEmpty()
+		? FStateTreeEditorStyle::Get().GetBrush("StateTree.CentricView.Leaf.Background.Outlined")
+		: FStateTreeEditorStyle::Get().GetBrush("StateTree.CentricView.Parent.Background.Muted.Outlined");
+
+	FSlateBrush Result = *BaseBrush;
+
+	Result.TintColor = State->Children.IsEmpty() ? ViewModel->GetStateColor(State, EStateTreeViewModelColorAdjustment::Dark) : Result.TintColor;
+	Result.OutlineSettings.Color = ViewModel->GetStateColor(State, EStateTreeViewModelColorAdjustment::Bright);
+
+	if (UStateCentricViewSettings::GetSchemaSettingsDefaultFallback(&ViewModel.Get()).bEnableHierarchyDrivenColors)
+	{
+		Result.TintColor = State->Children.IsEmpty() ? ViewModel->GetHierarchyBasedBackgroundColor(State, EStateTreeViewModelColorAdjustment::Dark) : Result.TintColor;
+		Result.OutlineSettings.Color = ViewModel->GetHierarchyBasedHighlightColor(State, EStateTreeViewModelColorAdjustment::Bright);
+
+		if (bDisableDarkenHeader)
+		{
+			Result.TintColor = State->Children.IsEmpty() ? ViewModel->GetHierarchyBasedBackgroundColor(State) : Result.TintColor;
+		}
+	}
+
+	return FDeferredCleanupSlateBrush::CreateBrush(Result);
+}
+
+TSharedPtr<FDeferredCleanupSlateBrush> FStateCentricViewUtils::MakeStateBackgroundCollapsedOutlined(TSharedRef<FStateTreeViewModel> ViewModel, TNotNull<const UStateTreeState*> State)
+{
+	const FSlateBrush* BaseBrush = State->Children.IsEmpty()
+		? FStateTreeEditorStyle::Get().GetBrush("StateTree.CentricView.Leaf.Background.Outlined")
+		: FStateTreeEditorStyle::Get().GetBrush("StateTree.CentricView.Parent.Background.Muted.Outlined");
+
+	FSlateBrush Result = *BaseBrush;
+
+	Result.TintColor = State->Children.IsEmpty() ? ViewModel->GetStateColor(State, EStateTreeViewModelColorAdjustment::Dark) : Result.TintColor;
+	Result.OutlineSettings.Color = ViewModel->GetStateColor(State, EStateTreeViewModelColorAdjustment::Dark);
+
+	if (UStateCentricViewSettings::GetSchemaSettingsDefaultFallback(&ViewModel.Get()).bEnableHierarchyDrivenColors)
+	{
+		Result.TintColor = State->Children.IsEmpty() ? ViewModel->GetHierarchyBasedBackgroundColor(State, EStateTreeViewModelColorAdjustment::Dark) : Result.TintColor;
+		Result.OutlineSettings.Color = ViewModel->GetHierarchyBasedHighlightColor(State);
+
+		if (bEnableDepthBasedOutline)
+		{
+			Result.OutlineSettings.Color = ViewModel->GetHierarchyBasedBackgroundColor(State, EStateTreeViewModelColorAdjustment::Bright);
+		}
+
+		if (bDisableDarkenHeader)
+		{
+			Result.TintColor = State->Children.IsEmpty() ? ViewModel->GetHierarchyBasedBackgroundColor(State) : Result.TintColor;
+		}
+	}
+
+	return FDeferredCleanupSlateBrush::CreateBrush(Result);
+}
+
+TSharedPtr<FDeferredCleanupSlateBrush> FStateCentricViewUtils::MakeStateFootnoteBackground(TSharedRef<FStateTreeViewModel> ViewModel, TNotNull<const UStateTreeState*> State)
+{
+	const FSlateBrush* BaseBrush = State->Children.IsEmpty()
+		? FStateTreeEditorStyle::Get().GetBrush("StateTree.CentricView.Leaf.Footnote.Background")
+		: FStateTreeEditorStyle::Get().GetBrush("StateTree.CentricView.Parent.Footnote.Background");
+
+	FSlateBrush Result = *BaseBrush;
+	
+	return FDeferredCleanupSlateBrush::CreateBrush(Result);
+}
+
+TSharedPtr<FDeferredCleanupSlateBrush> FStateCentricViewUtils::MakeStateSubnodeBackgroundOutlined(TSharedRef<FStateTreeViewModel> ViewModel, TNotNull<const UStateTreeState*> State)
+{
+	const FSlateBrush* BaseBrush = State->Children.IsEmpty()
+		? FStateTreeEditorStyle::Get().GetBrush("StateTree.CentricView.Leaf.Subnode.Background.Outlined")
+		: FStateTreeEditorStyle::Get().GetBrush("StateTree.CentricView.Parent.Subnode.Background.Outlined");
+
+	FSlateBrush Result = *BaseBrush;
+
+	Result.TintColor = ViewModel->GetStateColor(State, EStateTreeViewModelColorAdjustment::Dark);
+	Result.OutlineSettings.Color = ViewModel->GetStateColor(State);
+
+	if (UStateCentricViewSettings::GetSchemaSettingsDefaultFallback(&ViewModel.Get()).bEnableHierarchyDrivenColors)
+	{
+		Result.TintColor = ViewModel->GetHierarchyBasedBackgroundColor(State, EStateTreeViewModelColorAdjustment::Dark);
+		Result.OutlineSettings.Color = ViewModel->GetHierarchyBasedHighlightColor(State);
+
+		if (bEnableDepthBasedOutline)
+		{
+			Result.OutlineSettings.Color = ViewModel->GetHierarchyBasedBackgroundColor(State, EStateTreeViewModelColorAdjustment::Bright);
+		}
+
+		if (bDisableDarkenHeader)
+		{
+			Result.TintColor = ViewModel->GetHierarchyBasedBackgroundColor(State);
+		}
+	}
+
+	return FDeferredCleanupSlateBrush::CreateBrush(Result);
+}
+
+TSharedPtr<FDeferredCleanupSlateBrush> FStateCentricViewUtils::MakeStateSubnodeBackground(TSharedRef<FStateTreeViewModel> ViewModel, TNotNull<const UStateTreeState*> State)
+{
+	const FSlateBrush* BaseBrush = State->Children.IsEmpty()
+		? FStateTreeEditorStyle::Get().GetBrush("StateTree.CentricView.Leaf.Subnode.Background")
+		: FStateTreeEditorStyle::Get().GetBrush("StateTree.CentricView.Parent.Subnode.Background");
+
+	FSlateBrush Result = *BaseBrush;
+
+	Result.TintColor = ViewModel->GetStateColor(State, EStateTreeViewModelColorAdjustment::Dark);
+	Result.OutlineSettings.Color = ViewModel->GetStateColor(State, EStateTreeViewModelColorAdjustment::Dark);
+
+	if (UStateCentricViewSettings::GetSchemaSettingsDefaultFallback(&ViewModel.Get()).bEnableHierarchyDrivenColors)
+	{
+		Result.TintColor = ViewModel->GetHierarchyBasedBackgroundColor(State, EStateTreeViewModelColorAdjustment::Dark);
+		Result.OutlineSettings.Color = ViewModel->GetHierarchyBasedBackgroundColor(State, EStateTreeViewModelColorAdjustment::Dark);
+
+		if (bEnableDepthBasedOutline)
+		{
+			Result.OutlineSettings.Color = ViewModel->GetHierarchyBasedBackgroundColor(State, EStateTreeViewModelColorAdjustment::Bright);
+		}
+
+		if (bDisableDarkenHeader)
+		{
+			Result.TintColor = ViewModel->GetHierarchyBasedBackgroundColor(State);
+		}
+	}
+
+	return FDeferredCleanupSlateBrush::CreateBrush(Result);
+}
+
+TSharedPtr<FDeferredCleanupSlateBrush> FStateCentricViewUtils::MakeStateCollapsedBackground(TSharedRef<FStateTreeViewModel> ViewModel, TNotNull<const UStateTreeState*> State)
+{
+	const FSlateBrush* BaseBrush =  FStateTreeEditorStyle::Get().GetBrush("StateTree.CentricView.Parent.Subnode.Background");
+
+	FSlateBrush Result = *BaseBrush;
+
+	Result.TintColor = ViewModel->GetStateColor(State, EStateTreeViewModelColorAdjustment::Dark);
+	Result.OutlineSettings.Color = ViewModel->GetStateColor(State, EStateTreeViewModelColorAdjustment::Dark);
+
+	if (UStateCentricViewSettings::GetSchemaSettingsDefaultFallback(&ViewModel.Get()).bEnableHierarchyDrivenColors)
+	{
+		Result.TintColor = ViewModel->GetHierarchyBasedBackgroundColor(State, EStateTreeViewModelColorAdjustment::Dark);
+		Result.OutlineSettings.Color = ViewModel->GetHierarchyBasedBackgroundColor(State, EStateTreeViewModelColorAdjustment::Dark);
+
+		if (bEnableDepthBasedOutline)
+		{
+			Result.OutlineSettings.Color = ViewModel->GetHierarchyBasedBackgroundColor(State, EStateTreeViewModelColorAdjustment::Bright);
+		}
+
+		if (bDisableDarkenHeader)
+		{
+			Result.TintColor = ViewModel->GetHierarchyBasedBackgroundColor(State);
+		}
+	}
+
+	return FDeferredCleanupSlateBrush::CreateBrush(Result);
+}
+
+const FSlateBrush* FStateCentricViewUtils::GetSubNodeBackground()
+{
+	return FStateTreeEditorStyle::Get().GetBrush("StateTree.CentricView.SubNode.Background");
+}
+
+const FSlateBrush* FStateCentricViewUtils::GetSubNodeFootnoteBackground()
+{
+	return FStateTreeEditorStyle::Get().GetBrush("StateTree.CentricView.SubNode.Footnote.Background");
+}
+
+} // namespace UE::StateTree::Editor::StateCentricView
+
+#undef LOCTEXT_NAMESPACE
+
diff --git a/Engine/Plugins/Runtime/StateTree/Source/StateTreeEditorModule/Private/StateCentricView/View/StateCentricViewUtils.h b/Engine/Plugins/Runtime/StateTree/Source/StateTreeEditorModule/Private/StateCentricView/View/StateCentricViewUtils.h
new file mode 100644
--- /dev/null
+++ b/Engine/Plugins/Runtime/StateTree/Source/StateTreeEditorModule/Private/StateCentricView/View/StateCentricViewUtils.h
@@ -0,0 +1,72 @@
+// Copyright Epic Games, Inc. All Rights Reserved.
+
+#pragma once
+
+#include "CoreMinimal.h"
+
+#include "StateCentricView/View/SStateTreeExtendableNode.h"
+
+struct FSlateBrush;
+
+class FDeferredCleanupSlateBrush;
+class FStateTreeViewModel;
+class SButton;
+class SWidget;
+class UStateTreeState;
+
+enum class EExtendableNodeLOD : uint8;
+enum class EStateTransitionDirection : uint8;
+
+namespace UE::StateTree::Editor::StateCentricView
+{
+
+/**
+ * Experimental. General Utils for state centric view widget creation.
+ * 
+ * Currrently contains widget creation that is reused at different places. We prefer duplicated logic be here so it
+ * can be re-used in arbitray places if needed. All widget generation related methods must accept a LOD.
+ */
+struct FStateCentricViewUtils
+{
+	/** Build a view of transitions for given LOD / state / direction */
+	static TSharedRef<SWidget> GenerateTransitionsViewWidget(const EExtendableNodeLOD LOD
+		, TSharedRef<FStateTreeViewModel> ViewModel
+		, TNotNull<UStateTreeState*> ViewState
+		, const EStateTransitionDirection TransitionDirection);
+
+	/** Build a transparent button that jumps to given state */
+	static TSharedRef<SButton> GenerateGotoStateButton(const EExtendableNodeLOD LOD
+		, TSharedRef<FStateTreeViewModel> ViewModel
+		, TNotNull<const UStateTreeState*> ViewState);
+
+	/** Adds list of child state extensions at footnote for Full LOD */
+	static void AddChildStateFootnoteExtensions(const EExtendableNodeLOD InLOD
+		, TNotNull<SStateTreeExtendableNode*> InSourceNode
+		, TArray<TSharedPtr<FDeferredCleanupSlateBrush>>& OutSlateBrushResouces);
+
+	/** Helpers to get brushes based on if state is parent / leaf w/ appropriate colors */
+	[[nodiscard]] static TSharedPtr<FDeferredCleanupSlateBrush> MakeStateHalfTopOutline(TSharedRef<FStateTreeViewModel> ViewModel, TNotNull<const UStateTreeState*> State);
+	[[nodiscard]] static TSharedPtr<FDeferredCleanupSlateBrush> MakeStateHalfTopOutlineMuted(TSharedRef<FStateTreeViewModel> ViewModel, TNotNull<const UStateTreeState*> State);
+	[[nodiscard]] static TSharedPtr<FDeferredCleanupSlateBrush> MakeStateHalfTopOutlineCollapsed(TSharedRef<FStateTreeViewModel> ViewModel, TNotNull<const UStateTreeState*> State);
+	[[nodiscard]] static TSharedPtr<FDeferredCleanupSlateBrush> MakeStateOutline(TSharedRef<FStateTreeViewModel> ViewModel, TNotNull<const UStateTreeState*> State);
+	[[nodiscard]] static TSharedPtr<FDeferredCleanupSlateBrush> MakeStateOutlineMuted(TSharedRef<FStateTreeViewModel> ViewModel, TNotNull<const UStateTreeState*> State);
+
+	[[nodiscard]] static TSharedPtr<FDeferredCleanupSlateBrush> MakeStateHeaderBackground(TSharedRef<FStateTreeViewModel> ViewModel, TNotNull<const UStateTreeState*> State, TOptional<FVector2f> ImageSize = {});
+	[[nodiscard]] static TSharedPtr<FDeferredCleanupSlateBrush> MakeStateHeaderBackgroundMuted(TSharedRef<FStateTreeViewModel> ViewModel, TNotNull<const UStateTreeState*> State);
+	[[nodiscard]] static TSharedPtr<FDeferredCleanupSlateBrush> MakeStateHeaderBackgroundCollapsed(TSharedRef<FStateTreeViewModel> ViewModel, TNotNull<const UStateTreeState*> State);
+
+	[[nodiscard]] static TSharedPtr<FDeferredCleanupSlateBrush> MakeStateBackgroundMutedOutlined(TSharedRef<FStateTreeViewModel> ViewModel, TNotNull<const UStateTreeState*> State);
+	[[nodiscard]] static TSharedPtr<FDeferredCleanupSlateBrush> MakeStateBackgroundCollapsedOutlined(TSharedRef<FStateTreeViewModel> ViewModel, TNotNull<const UStateTreeState*> State);
+
+	[[nodiscard]] static TSharedPtr<FDeferredCleanupSlateBrush> MakeStateFootnoteBackground(TSharedRef<FStateTreeViewModel> ViewModel, TNotNull<const UStateTreeState*> State);
+
+	[[nodiscard]] static TSharedPtr<FDeferredCleanupSlateBrush> MakeStateSubnodeBackgroundOutlined(TSharedRef<FStateTreeViewModel> ViewModel, TNotNull<const UStateTreeState*> State);
+	[[nodiscard]] static TSharedPtr<FDeferredCleanupSlateBrush> MakeStateSubnodeBackground(TSharedRef<FStateTreeViewModel> ViewModel, TNotNull<const UStateTreeState*> State);
+	[[nodiscard]] static TSharedPtr<FDeferredCleanupSlateBrush> MakeStateCollapsedBackground(TSharedRef<FStateTreeViewModel> ViewModel, TNotNull<const UStateTreeState*> State);
+
+	/** Helpers to get standard state centric brushes */
+	static const FSlateBrush* GetSubNodeBackground();
+	static const FSlateBrush* GetSubNodeFootnoteBackground();
+};
+
+} // namespace UE::StateTree::Editor::StateCentricView
\ No newline at end of file
diff --git a/Engine/Plugins/Runtime/StateTree/Source/StateTreeEditorModule/Private/StateCentricView/ViewModel/StateCentricViewEditorDataExtension.cpp b/Engine/Plugins/Runtime/StateTree/Source/StateTreeEditorModule/Private/StateCentricView/ViewModel/StateCentricViewEditorDataExtension.cpp
new file mode 100644
--- /dev/null
+++ b/Engine/Plugins/Runtime/StateTree/Source/StateTreeEditorModule/Private/StateCentricView/ViewModel/StateCentricViewEditorDataExtension.cpp
@@ -0,0 +1,80 @@
+// Copyright Epic Games, Inc. All Rights Reserved.
+
+#include "StateCentricView/ViewModel/StateCentricViewEditorDataExtension.h"
+
+#include UE_INLINE_GENERATED_CPP_BY_NAME(StateCentricViewEditorDataExtension)
+
+UStateCentricViewEditorDataExtension::UStateCentricViewEditorDataExtension()
+{
+	InTransitionSplitterSizeRule = SSplitter::ESizeRule::FractionOfParent;
+	OutTransitionSplitterSizeRule = SSplitter::ESizeRule::FractionOfParent;
+}
+
+void UStateCentricViewEditorDataExtension::ResetSplitters()
+{
+	LeftSpacerSplitterValue = 0.025f;
+	InTransitionSplitterValue = 0.1875f;
+	MainNodeSplitterValue = 0.375f;
+	OutTransitionSplitterValue = 0.4375f;
+	RightSpacerSplitterValue = 0.025f;
+}
+
+float UStateCentricViewEditorDataExtension::HandleGetLeftSpacerSplitterValue() const
+{
+	return LeftSpacerSplitterValue;
+}
+
+float UStateCentricViewEditorDataExtension::HandleGetInTransitionSplitterValue() const
+{
+	return InTransitionSplitterValue;
+}
+
+float UStateCentricViewEditorDataExtension::HandleGetMainNodeSplitterValue() const
+{
+	return MainNodeSplitterValue;
+}
+
+float UStateCentricViewEditorDataExtension::HandleGetOutTransitionSplitterValue() const
+{
+	return OutTransitionSplitterValue;
+}
+
+float UStateCentricViewEditorDataExtension::HandleGetRightSpacerSplitterValue() const
+{
+	return RightSpacerSplitterValue;
+}
+
+SSplitter::ESizeRule UStateCentricViewEditorDataExtension::GetInTransitionSplitterSizeRule() const
+{
+	return InTransitionSplitterSizeRule;
+}
+
+SSplitter::ESizeRule UStateCentricViewEditorDataExtension::GetOutTransitionSplitterSizeRule() const
+{
+	return OutTransitionSplitterSizeRule;
+}
+
+void UStateCentricViewEditorDataExtension::HandleOnLeftSpacerSplitterResized(float InSize)
+{
+	LeftSpacerSplitterValue = InSize;
+}
+
+void UStateCentricViewEditorDataExtension::HandleOnInTransitionSplitterResized(float InSize)
+{
+	InTransitionSplitterValue = InSize;
+}
+
+void UStateCentricViewEditorDataExtension::HandleOnMainNodeSplitterResized(float InSize)
+{
+	MainNodeSplitterValue = InSize;
+}
+
+void UStateCentricViewEditorDataExtension::HandleOnOutTransitionSplitterResized(float InSize)
+{
+	OutTransitionSplitterValue = InSize;
+}
+
+void UStateCentricViewEditorDataExtension::HandleOnRightSpacerSplitterResized(float InSize)
+{
+	RightSpacerSplitterValue = InSize;
+}
\ No newline at end of file
diff --git a/Engine/Plugins/Runtime/StateTree/Source/StateTreeEditorModule/Private/StateCentricView/ViewModel/StateCentricViewEditorDataExtension.h b/Engine/Plugins/Runtime/StateTree/Source/StateTreeEditorModule/Private/StateCentricView/ViewModel/StateCentricViewEditorDataExtension.h
new file mode 100644
--- /dev/null
+++ b/Engine/Plugins/Runtime/StateTree/Source/StateTreeEditorModule/Private/StateCentricView/ViewModel/StateCentricViewEditorDataExtension.h
@@ -0,0 +1,82 @@
+// Copyright Epic Games, Inc. All Rights Reserved.
+
+#pragma once
+
+#include "StateCentricView/StateCentricViewPerUserSettings.h"
+#include "StateCentricView/StateCentricViewSettings.h"
+#include "StateTreeEditorDataExtension.h"
+
+#include "Widgets/Layout/SSplitter.h"
+
+#include "StateCentricViewEditorDataExtension.generated.h"
+
+#define UE_API STATETREEEDITORMODULE_API
+
+class UStateTreeState;
+
+/**
+ * Experimental. Extension to store temporary editor data for state centric view & view interactions
+ */
+UCLASS(Experimental, MinimalAPI)
+class UStateCentricViewEditorDataExtension : public UStateTreeEditorDataExtension
+{
+	GENERATED_BODY()
+
+public:
+
+	UStateCentricViewEditorDataExtension();
+
+	/** User settings for this specific state tree. Ensures consistent behavior when opening different states in this state tree. */
+	UPROPERTY(EditAnywhere, Category = "Appearance", meta = (EditCondition = "StateTreeEditorModule.StateCentricViewSettings.IsStateCentricViewEnabled", EditConditionHides))
+	FStateCentricViewPerUserSetting PerStateTreeUserSettings;
+
+	/**
+	 * LOD to show this parent state at. We use this to ensure the view remains consistent
+	 * when clicking between different states with the same parent & when clicking back & forth.
+	 * 
+	 * @TODO: This is a map, but it would be more idiomatic for it to exist as an editor data
+	 * extension on a state. Currently that doesn't exist. See if ST owners are interested in that being added.
+	 * Ex: This map becomes ugly when a state is deleted, with correct ownership this is not an issue.
+	 */
+	UPROPERTY(Transient)
+	TMap<FGuid, EExtendableNodeLOD> LastUserSetParentStateLOD;
+
+public:
+
+	float LeftSpacerSplitterValue = 0.025f;
+	float InTransitionSplitterValue = 0.1875f;
+	float MainNodeSplitterValue = 0.375f;
+	float OutTransitionSplitterValue = 0.4375f;
+	float RightSpacerSplitterValue = 0.025f;
+
+	SSplitter::ESizeRule InTransitionSplitterSizeRule;
+	SSplitter::ESizeRule OutTransitionSplitterSizeRule;
+
+public:
+
+	UE_API void ResetSplitters();
+
+	UE_API float HandleGetLeftSpacerSplitterValue() const;
+	UE_API float HandleGetInTransitionSplitterValue() const;
+	UE_API float HandleGetMainNodeSplitterValue() const;
+	UE_API float HandleGetOutTransitionSplitterValue() const;
+	UE_API float HandleGetRightSpacerSplitterValue() const;
+
+	SSplitter::ESizeRule GetInTransitionSplitterSizeRule() const;
+	SSplitter::ESizeRule GetOutTransitionSplitterSizeRule() const;
+
+	UE_API void HandleOnLeftSpacerSplitterResized(float InSize);
+	UE_API void HandleOnInTransitionSplitterResized(float InSize);
+	UE_API void HandleOnMainNodeSplitterResized(float InSize);
+	UE_API void HandleOnOutTransitionSplitterResized(float InSize);
+	UE_API void HandleOnRightSpacerSplitterResized(float InSize);
+
+protected:
+	
+	UStateTreeEditorData* GetStateTreeEditorData() const
+	{
+		return GetOuterUStateTreeEditorData();
+	}
+};
+
+#undef UE_API
\ No newline at end of file
diff --git a/Engine/Plugins/Runtime/StateTree/Source/StateTreeEditorModule/Private/StateTreeDescriptionHelpers.cpp b/Engine/Plugins/Runtime/StateTree/Source/StateTreeEditorModule/Private/StateTreeDescriptionHelpers.cpp
--- a/Engine/Plugins/Runtime/StateTree/Source/StateTreeEditorModule/Private/StateTreeDescriptionHelpers.cpp
+++ b/Engine/Plugins/Runtime/StateTree/Source/StateTreeEditorModule/Private/StateTreeDescriptionHelpers.cpp
@@ -170,6 +170,11 @@ const FSlateBrush* GetStateLinkIcon(const UStateTreeEditorData* EditorData, cons
 
 FSlateColor GetStateLinkColor(const UStateTreeEditorData* EditorData, const FStateTreeStateLink& Link)
 {
+	if (!EditorData)
+	{
+		return FLinearColor::White;
+	}
+
 	if (Link.LinkType == EStateTreeTransitionType::GotoState)
 	{
 		if (const UStateTreeState* State = EditorData->GetStateByID(Link.ID))
diff --git a/Engine/Plugins/Runtime/StateTree/Source/StateTreeEditorModule/Private/StateTreeEditorPropertyBindings.cpp b/Engine/Plugins/Runtime/StateTree/Source/StateTreeEditorModule/Private/StateTreeEditorPropertyBindings.cpp
--- a/Engine/Plugins/Runtime/StateTree/Source/StateTreeEditorModule/Private/StateTreeEditorPropertyBindings.cpp
+++ b/Engine/Plugins/Runtime/StateTree/Source/StateTreeEditorModule/Private/StateTreeEditorPropertyBindings.cpp
@@ -190,6 +190,7 @@ int32 FStateTreeEditorPropertyBindings::GetNumBindings() const
 FStateTreeBindingLookup::FStateTreeBindingLookup(const IStateTreeEditorPropertyBindingsOwner* InBindingOwner)
 	: BindingOwner(InBindingOwner)
 {
+	check(BindingOwner);
 }
 
 const FPropertyBindingPath* FStateTreeBindingLookup::GetPropertyBindingSource(const FPropertyBindingPath& InTargetPath) const
diff --git a/Engine/Plugins/Runtime/StateTree/Source/StateTreeEditorModule/Private/StateTreeEditorStyle.cpp b/Engine/Plugins/Runtime/StateTree/Source/StateTreeEditorModule/Private/StateTreeEditorStyle.cpp
--- a/Engine/Plugins/Runtime/StateTree/Source/StateTreeEditorModule/Private/StateTreeEditorStyle.cpp
+++ b/Engine/Plugins/Runtime/StateTree/Source/StateTreeEditorModule/Private/StateTreeEditorStyle.cpp
@@ -211,6 +211,157 @@ FStateTreeEditorStyle::FStateTreeEditorStyle()
 			.SetEditableTextBoxStyle(EditableTextBlockStyle));
 	}
 
+	// State Centric Styling
+	{
+		// Note: AccentBlue below used as placeholder for state color, final color will be determined by state. See 'FStateCentricViewUtils'
+
+		// Parent styles
+		{
+			// How do we get a half-rounded rectangle border w/ no bottom in slate? There is no built in way, 
+			// so we set crazy high curve edges to cause the bottom corners to cull (90000.0f).
+			// To use correctly will need to negative pad to extend the half outline (Depends, -100 is safe).
+			Set("StateTree.CentricView.Parent.HalfTopOutline", new FSlateRoundedBoxBrush(
+				FStyleColors::AccentBlue.GetSpecifiedColor().CopyWithNewOpacity(0.0f)
+				, FVector4{ 10.0, 10.0, 90000.0, 90000.0 }
+				, FStyleColors::AccentBlue
+				, 2.0f));
+
+			Set("StateTree.CentricView.Parent.Outline", new FSlateRoundedBoxBrush(
+				FStyleColors::AccentBlue.GetSpecifiedColor().CopyWithNewOpacity(0.0f)
+				, FVector4{ 10.0, 10.0, 10.0, 10.0 }
+				, FStyleColors::AccentBlue
+				, 2.0f));
+
+			Set("StateTree.CentricView.Parent.Header.Background", new FSlateRoundedBoxBrush(
+				FStyleColors::AccentBlue
+				, FVector4{ 10.0, 10.0, 0.0, 0.0 }
+				, FStyleColors::AccentBlue
+				, 0.0f
+				, FVector2f{ 32.0f, 22.0f }));
+
+			Set("StateTree.CentricView.Parent.Header.Background.Muted", new FSlateRoundedBoxBrush(
+				FStyleColors::Header
+				, FVector4{ 10.0, 10.0, 0.0, 0.0 }
+				, FStyleColors::Header
+				, 0.0f
+				, FVector2f{ 32.0f, 22.0f }));
+
+			Set("StateTree.CentricView.Parent.Background.Muted.Outlined", new FSlateRoundedBoxBrush(
+				FStyleColors::AccentBlack
+				, FVector4{ 10.0, 10.0, 10.0, 10.0 }
+				, FStyleColors::AccentBlue
+				, 2.0f
+				, FVector2f{ 32.0f, 22.0f }));
+
+			// Same hack used to get a correctly culling bottom rounded corner (90000.0f). Also needs negative padding for correct spacing (-10)
+			Set("StateTree.CentricView.Parent.Footnote.Background", new FSlateRoundedBoxBrush(
+				FStyleColors::AccentBlack
+				, FVector4{ 90000.0, 90000.0, 10.0, 10.0 }
+				, FStyleColors::AccentBlack
+				, 0.0f
+				, FVector2f{ 32.0f, 20.0f }));
+
+			Set("StateTree.CentricView.Parent.Subnode.Background.Outlined", new FSlateRoundedBoxBrush(
+				FStyleColors::AccentBlue
+				, FVector4{ 5.0, 5.0, 5.0, 5.0 }
+				, FStyleColors::AccentBlue // Note: Will have bright state color outline
+				, 1.0f
+				, FVector2f{ 32.0f, 18.0f }));
+
+			Set("StateTree.CentricView.Parent.SubNode.Background", new FSlateRoundedBoxBrush(
+				FStyleColors::AccentBlue
+				, FVector4{ 5.0, 5.0, 5.0, 5.0 }
+				, FStyleColors::AccentBlue
+				, 0.0f
+				, FVector2f{ 32.0f, 18.0f }));
+		}
+
+		// Leaf styles
+		{
+			Set("StateTree.CentricView.Leaf.Outline", new FSlateRoundedBoxBrush(
+				FStyleColors::AccentBlue.GetSpecifiedColor().CopyWithNewOpacity(0.0f)
+				, FVector4{ 0.0 }
+				, FStyleColors::AccentBlue
+				, 2.0f));
+
+			Set("StateTree.CentricView.Leaf.Header.Background", new FSlateRoundedBoxBrush(
+				FStyleColors::AccentBlue
+				, FVector4{ 0.0 }
+				, FStyleColors::AccentBlue
+				, 0.0f
+				, FVector2f{ 32.0f, 22.0f }));
+
+			Set("StateTree.CentricView.Leaf.Footnote.Background", new FSlateRoundedBoxBrush(
+				FStyleColors::AccentBlack
+				, FVector4{ 0.0 }
+				, FStyleColors::AccentBlack
+				, 0.0f
+				, FVector2f{ 32.0f, 10.0f }));
+
+			Set("StateTree.CentricView.Leaf.Background.Outlined", new FSlateRoundedBoxBrush(
+				FStyleColors::AccentBlack
+				, FVector4{ 10.0, 10.0, 10.0, 10.0 }
+				, FStyleColors::AccentBlue
+				, 0.0f
+				, FVector2f{ 32.0f, 22.0f }));
+
+			Set("StateTree.CentricView.Leaf.Subnode.Background.Outlined", new FSlateRoundedBoxBrush(
+				FStyleColors::AccentBlue
+				, FVector4{ 0.0 }
+				, FStyleColors::AccentBlue // Note: Will have bright state color outline
+				, 1.0f
+				, FVector2f{ 32.0f, 18.0f }));
+
+			Set("StateTree.CentricView.Leaf.SubNode.Background", new FSlateRoundedBoxBrush(
+				FStyleColors::AccentBlue
+				, FVector4{ 0.0 }
+				, FStyleColors::AccentBlue
+				, 0.0f
+				, FVector2f{ 32.0f, 18.0f }));
+		}
+
+		// General node styling
+		{
+			Set("StateTree.CentricView.SubNode.Background", new FSlateRoundedBoxBrush(
+				FStyleColors::AccentBlack
+				, FVector4{ 5.0, 5.0, 5.0, 5.0 }
+				, FStyleColors::AccentBlack
+				, 0.0f));
+
+			// Same hack used to get a correctly culling bottom rounded corner (90000.0f). Also needs negative padding for correct spacing (-5)
+			Set("StateTree.CentricView.SubNode.Footnote.Background", new FSlateRoundedBoxBrush(
+				FStyleColors::AccentBlack
+				, FVector4{ 90000.0, 90000.0, 5.0, 5.0 }
+				, FStyleColors::AccentBlack
+				, 0.0f
+				, FVector2f{ 32.0f, 10.0f }));
+		}
+
+		// State centric text, these values don't exist elsewhere
+		{
+			Set("StateTree.CentricView.StateTitle.Bold",
+				FTextBlockStyle(NormalText)
+				.SetFont(DEFAULT_FONT("Bold", 12))
+				.SetColorAndOpacity(FStyleColors::White));
+
+			Set("StateTree.CentricView.StateTitle.Regular",
+				FTextBlockStyle(NormalText)
+				.SetFont(DEFAULT_FONT("Regular", 12))
+				.SetColorAndOpacity(FStyleColors::White));
+		}
+
+		// State Centric Button Styles
+		{
+			const FButtonStyle& SimpleButtonStyle = FAppStyle::GetWidgetStyle<FButtonStyle>("SimpleButton");
+
+			FButtonStyle HoverHighlight = FButtonStyle(SimpleButtonStyle)
+				.SetHovered(FSlateRoundedBoxBrush(FStyleColors::White.GetSpecifiedColor().CopyWithNewOpacity(0.2f), 5.0f))
+				.SetPressed(FSlateRoundedBoxBrush(FStyleColors::Black.GetSpecifiedColor().CopyWithNewOpacity(0.1f), 5.0f));
+
+			Set("StateTree.CentricView.HoverHighlight", HoverHighlight);
+		}
+	}
+
 	// Command icons
 	{
 		// From generic Engine
diff --git a/Engine/Plugins/Runtime/StateTree/Source/StateTreeEditorModule/Private/StateTreeViewModel.cpp b/Engine/Plugins/Runtime/StateTree/Source/StateTreeEditorModule/Private/StateTreeViewModel.cpp
--- a/Engine/Plugins/Runtime/StateTree/Source/StateTreeEditorModule/Private/StateTreeViewModel.cpp
+++ b/Engine/Plugins/Runtime/StateTree/Source/StateTreeEditorModule/Private/StateTreeViewModel.cpp
@@ -4,6 +4,7 @@
 
 #include "StateTree.h"
 #include "StateTreeEditorData.h"
+#include "StateTreeEditorStyle.h"
 #include "StateTreeDelegates.h"
 #include "Debugger/StateTreeDebugger.h"
 #include "Editor.h"
@@ -17,6 +18,7 @@
 #include "StateTreeConsiderationBase.h"
 #include "StateTreePropertyHelpers.h"
 #include "StateTreeEditorDataClipboardHelpers.h"
+#include "Styling/SlateBrush.h"
 #include "Misc/NotNull.h"
 #include "Misc/StringOutputDevice.h"
 
@@ -24,8 +26,21 @@
 
 #define LOCTEXT_NAMESPACE "StateTreeEditor"
 
+bool bEnableSimpleHierarchyDrivenColors = false;
+FAutoConsoleVariableRef CVarEnableSimpleHierarchyDrivenColors(
+	TEXT("StateTree.Editor.Experimental.EnableSimpleHierarchyDrivenColors"),
+	bEnableSimpleHierarchyDrivenColors,
+	TEXT("Set true to use a straightforward approach of replacing base colors for hierarchy colors. Disable hierarchy colors in settings to use.")
+);
+
 namespace UE::StateTree::Editor
 {
+	static const FLinearColor DefaultStateColor = FLinearColor(FColor(31, 151, 167));
+
+	// These defaults only apply when using hierarchy driven colors
+	static const FLinearColor DefaultRootColor = FLinearColor(0.347f, 0.347f, 0.347f);
+	static const FLinearColor DefaultParentColor = FLinearColor(1.0f, 0.125f, 0.913f);
+
 	class FStateTreeStateTextFactory : public FCustomizableTextObjectFactory
 	{
 	public:
@@ -440,6 +455,37 @@ UStateTreeState* FStateTreeViewModel::GetMutableStateByID(const FGuid StateID) c
 	return nullptr;
 }
 
+const FStateTreeTransition* FStateTreeViewModel::GetTransitionByID(const FGuid TransitionID) const
+{
+	return GetMutableTransitionByID(TransitionID);
+}
+
+FStateTreeTransition* FStateTreeViewModel::GetMutableTransitionByID(const FGuid TransitionID) const
+{
+	FStateTreeTransition* Result = nullptr;
+
+	auto FindMatchingTransition = [&](UStateTreeState& HierarchyState, UStateTreeState* ParentState)
+	{
+		for (FStateTreeTransition& Transition : HierarchyState.Transitions)
+		{
+			if (Transition.ID == TransitionID)
+			{
+				Result = &Transition;
+				return EStateTreeVisitor::Break;
+			}
+		}
+
+		return EStateTreeVisitor::Continue;
+	};
+
+	if (UStateTreeEditorData* TreeData = TreeDataWeak.Get())
+	{
+		TreeData->VisitHierarchy(FindMatchingTransition);
+	}
+
+	return Result;
+}
+
 void FStateTreeViewModel::HandleIdentifierChanged(const UStateTree& StateTree) const
 {
 	if (GetStateTree() == &StateTree)
@@ -858,9 +904,9 @@ void FStateTreeViewModel::GetLinkStates(FGuid StateID, TArray<FGuid>& LinkingIn,
 	}
 }
 
-UStateTreeState* FStateTreeViewModel::GetTransitionToState(TNotNull<const UStateTreeState*> State, const FGuid TransitionID) const
+UStateTreeState* FStateTreeViewModel::GetTransitionToState(TNotNull<const UStateTreeState*> TransitionOwningState, const FGuid TransitionID) const
 {
-	const FStateTreeTransition* Transition = State->GetTransitionByID(TransitionID);
+	const FStateTreeTransition* Transition = TransitionOwningState->GetTransitionByID(TransitionID);
 	ensure(Transition);
 	if (Transition == nullptr)
 	{
@@ -872,13 +918,13 @@ UStateTreeState* FStateTreeViewModel::GetTransitionToState(TNotNull<const UState
 	case EStateTreeTransitionType::GotoState:
 		return GetMutableStateByID(Transition->State.ID);
 	case EStateTreeTransitionType::Parent:
-		return State->Parent ? State->Parent.Get() : nullptr;
+		return TransitionOwningState->Parent ? TransitionOwningState->Parent.Get() : nullptr;
 	case EStateTreeTransitionType::NextState:
 	case EStateTreeTransitionType::NextSelectableState:
-		return State->GetNextSelectableSiblingState();
+		return TransitionOwningState->GetNextSelectableSiblingState();
 	case EStateTreeTransitionType::NextParent:
 	case EStateTreeTransitionType::NextSelectableParent:
-		return State->Parent ? State->Parent->GetNextSelectableSiblingState() : nullptr;
+		return TransitionOwningState->Parent ? TransitionOwningState->Parent->GetNextSelectableSiblingState() : nullptr;
 	case EStateTreeTransitionType::None:
 	case EStateTreeTransitionType::Succeeded:
 	case EStateTreeTransitionType::Failed:
@@ -908,6 +954,7 @@ void FStateTreeViewModel::ClearSelection()
 	SelectedStates.Reset();
 
 	const TArray<TWeakObjectPtr<UStateTreeState>> SelectedStatesArr;
+	LastSelectedState = FGuid();
 	OnSelectionChanged.Broadcast(SelectedStatesArr);
 }
 
@@ -923,10 +970,9 @@ void FStateTreeViewModel::SetSelection(UStateTreeState* Selected)
 	{
 		SelectedStates.Add(Selected);
 	}
-
-	TArray<TWeakObjectPtr<UStateTreeState>> SelectedStatesArr;
-	SelectedStatesArr.Add(Selected);
-	OnSelectionChanged.Broadcast(SelectedStatesArr);
+	
+	LastSelectedState = Selected ? Selected->ID : FGuid();
+	OnSelectionChanged.Broadcast({ Selected });
 }
 
 void FStateTreeViewModel::DeleteNode(TWeakObjectPtr<UStateTreeState> State, const FGuid& ID)
@@ -1318,6 +1364,7 @@ void FStateTreeViewModel::SetSelection(const TArray<TWeakObjectPtr<UStateTreeSta
 		if (State.Get())
 		{
 			SelectedStates.Add(State);
+			LastSelectedState = State->ID;
 		}
 	}
 
@@ -1379,12 +1426,36 @@ bool FStateTreeViewModel::HasSelection() const
 	return SelectedStates.Num() > 0;
 }
 
+UStateTreeState* FStateTreeViewModel::GetLastSelectedState() const
+{
+	return GetMutableStateByID(LastSelectedState);
+}
+
 void FStateTreeViewModel::BringNodeToFocus(UStateTreeState* State, const FGuid NodeID)
 {
 	SetSelection(State);
 	OnBringNodeToFocus.Broadcast(State, NodeID);
 }
 
+UStateTreeState* FStateTreeViewModel::GetViewState() const
+{
+	UStateTreeState* ViewState = GetLastSelectedState();
+
+	// Try to view the first root when we have no selection
+	if (!ViewState)
+	{
+		TArray<TWeakObjectPtr<UStateTreeState>> OutSubtrees;
+		GetSubTrees(OutSubtrees);
+
+		if (!OutSubtrees.IsEmpty())
+		{
+			ViewState = OutSubtrees[0].Get();
+		}
+	}
+
+	return ViewState;
+}
+
 void FStateTreeViewModel::GetPersistentExpandedStates(TSet<TWeakObjectPtr<UStateTreeState>>& OutExpandedStates)
 {
 	OutExpandedStates.Reset();
@@ -1851,6 +1922,319 @@ void FStateTreeViewModel::SetSelectedStatesEnabled(const bool bEnable)
 	}
 }
 
+void FStateTreeViewModel::ForEachParent(TNotNull<UStateTreeState*> State, TFunctionRef<EStateTreeVisitor(UStateTreeState& ParentState)> InFunc) const
+{
+	UStateTreeState* ParentState = State->Parent;
+	while (ParentState)
+	{
+		EStateTreeVisitor ContinueVisitor = InFunc(*ParentState);
+		if (ContinueVisitor != EStateTreeVisitor::Continue)
+		{
+			return;
+		}
+
+		ParentState = ParentState->Parent;
+	}
+}
+
+void FStateTreeViewModel::ForEachInTransition(TNotNull<UStateTreeState*> State, TFunctionRef<EStateTreeVisitor(FStateTreeTransition& TransitionToState)> InFunc) const
+{
+	auto VisitStateTransitionsToState = [&](UStateTreeState& HierarchyState, UStateTreeState* ParentState)
+	{
+		for (FStateTreeTransition& Transition : HierarchyState.Transitions)
+		{
+			if (Transition.State.ID == State->GetLinkToState().ID)
+			{
+				EStateTreeVisitor ContinueVisitor = InFunc(Transition);
+				if (ContinueVisitor != EStateTreeVisitor::Continue)
+				{
+					return EStateTreeVisitor::Break;
+				}
+			}
+		}
+
+		return EStateTreeVisitor::Continue;
+	};
+
+	if (const UStateTreeEditorData* TreeData = TreeDataWeak.Get())
+	{
+		TreeData->VisitHierarchy(VisitStateTransitionsToState);
+	}
+}
+
+void FStateTreeViewModel::ForEachOutTransition(TNotNull<UStateTreeState*> State, TFunctionRef<EStateTreeVisitor(FStateTreeTransition& TransitionFromState)> InFunc) const
+{
+	// Note: Yes this is trivial. But the view doesn't need to know that.
+	for (FStateTreeTransition& Transition : State->Transitions)
+	{
+		EStateTreeVisitor ContinueVisitor = InFunc(Transition);
+		if (ContinueVisitor != EStateTreeVisitor::Continue)
+		{
+			return;
+		}
+	}
+}
+
+void FStateTreeViewModel::ForEachChild(TNotNull<UStateTreeState*> State, TFunctionRef<EStateTreeVisitor(UStateTreeState& ChildState)> InFunc) const
+{
+	// Note: Yes this is trivial. But the view doesn't need to know that.
+	for (UStateTreeState* Child : State->Children)
+	{
+		if (Child)
+		{
+			EStateTreeVisitor ContinueVisitor = InFunc(*Child);
+			if (ContinueVisitor != EStateTreeVisitor::Continue)
+			{
+				return;
+			}
+		}
+	}
+}
+
+TArray<UStateTreeState*> FStateTreeViewModel::GetStateParents(TNotNull<UStateTreeState*> State) const
+{
+	TArray<UStateTreeState*> ParentStates;
+	auto GatherParentStates = [&](UStateTreeState& ParentState)
+	{
+		ParentStates.Add(&ParentState);
+		return EStateTreeVisitor::Continue;
+	};
+
+	ForEachParent(State, GatherParentStates);
+	return ParentStates;
+}
+
+TArray<FStateTreeTransition*> FStateTreeViewModel::GetStateInTransitions(TNotNull<UStateTreeState*> State) const
+{
+	TArray<FStateTreeTransition*> TransitionsToState;
+	auto GatherInTransitions = [&](FStateTreeTransition& Transition)
+	{
+		TransitionsToState.Add(&Transition);
+		return EStateTreeVisitor::Continue;
+	};
+
+	ForEachInTransition(State, GatherInTransitions);
+	return TransitionsToState;
+}
+
+TArray<FStateTreeTransition*> FStateTreeViewModel::GetStateOutTransitions(TNotNull<UStateTreeState*> State) const
+{
+	TArray<FStateTreeTransition*> TransitionsFromState;
+	auto GatherOutTransitions = [&](FStateTreeTransition& Transition)
+	{
+			TransitionsFromState.Add(&Transition);
+		return EStateTreeVisitor::Continue;
+	};
+
+	ForEachOutTransition(State, GatherOutTransitions);
+	return TransitionsFromState;
+}
+
+TArray<UStateTreeState*> FStateTreeViewModel::GetStateChildren(TNotNull<UStateTreeState*> State) const
+{
+	TArray<UStateTreeState*> ChildStates;
+	auto GatherChildStates = [&](UStateTreeState& ChildState)
+	{
+		ChildStates.Add(&ChildState);
+		return EStateTreeVisitor::Continue;
+	};
+
+	ForEachChild(State, GatherChildStates);
+	return ChildStates;
+}
+
+const UStateTreeState* FStateTreeViewModel::GetSourceStateFromTransition(TNotNull<const FStateTreeTransition*> Transition) const
+{
+	const UStateTreeState* SourceState = nullptr;
+
+	auto VisitStateFindTransitionSource = [&](UStateTreeState& HierarchyState, UStateTreeState* ParentState)
+	{
+		auto TransitionCompareID = [&](const FStateTreeTransition& Other)
+		{
+			return Other.ID == Transition->ID;
+		};
+
+		if (HierarchyState.Transitions.ContainsByPredicate(TransitionCompareID))
+		{
+			SourceState = &HierarchyState;
+			return EStateTreeVisitor::Break;
+		}
+
+		return EStateTreeVisitor::Continue;
+	};
+
+	if (const UStateTreeEditorData* TreeData = TreeDataWeak.Get())
+	{
+		TreeData->VisitHierarchy(VisitStateFindTransitionSource);
+	}
+
+	return SourceState;
+}
+
+const UStateTreeState* FStateTreeViewModel::GetSourceParentStateFromTransition(TNotNull<const FStateTreeTransition*> Transition) const
+{
+	if (const UStateTreeState* TransitionState = GetSourceStateFromTransition(Transition))
+	{
+		return TransitionState->Parent;
+	}
+
+	return nullptr;
+}
+
+const UStateTreeState* FStateTreeViewModel::GetTargetStateFromTransition(TNotNull<const FStateTreeTransition*> Transition) const
+{
+	if (const UStateTreeState* TransitionOwningState = GetSourceStateFromTransition(Transition))
+	{
+		return GetTransitionToState(TransitionOwningState, Transition->ID);
+	}
+
+	return nullptr;
+}
+
+const UStateTreeState* FStateTreeViewModel::GetTargetParentStateFromTransition(TNotNull<const FStateTreeTransition*> Transition) const
+{
+	if (const UStateTreeState* TransitionState = GetTargetStateFromTransition(Transition))
+	{
+		return TransitionState->Parent;
+	}
+
+	return nullptr;
+}
+
+FLinearColor FStateTreeViewModel::GetStateColor(TNotNull<const UStateTreeState*> State, EStateTreeViewModelColorAdjustment ColorAdjustment) const
+{
+	if (bEnableSimpleHierarchyDrivenColors)
+	{
+		return GetColorAdjusted(GetHierarchyBasedBackgroundColor(State), ColorAdjustment);
+	}
+
+	FLinearColor Result = UE::StateTree::Editor::DefaultStateColor;
+
+	if (const UStateTreeEditorData* TreeData = TreeDataWeak.Get())
+	{
+		if (const FStateTreeEditorColor* StateColor = TreeData->FindColor(FStateTreeEditorColorRef(State->ColorRef)))
+		{
+			Result = StateColor->Color;
+		}
+	}
+
+	return GetColorAdjusted(Result, ColorAdjustment);
+}
+
+FLinearColor FStateTreeViewModel::GetHierarchyBasedHighlightColor(TNotNull<const UStateTreeState*> State, EStateTreeViewModelColorAdjustment ColorAdjustment) const
+{
+	// If we are a leaf, then just use leaf color
+	if (State->Children.IsEmpty())
+	{
+		return GetColorAdjusted(UE::StateTree::Editor::DefaultStateColor, ColorAdjustment);
+	}
+
+	int8 NumParents = 0;
+	auto CountNumParents = [&](UStateTreeState& ParentState)
+	{
+		NumParents++;
+		return NumParents == 1 ? EStateTreeVisitor::Break : EStateTreeVisitor::Continue;
+	};
+	ForEachParent(const_cast<UStateTreeState*>(NotNullGet(State)), CountNumParents);
+
+	return GetColorAdjusted(NumParents > 0 ? UE::StateTree::Editor::DefaultParentColor : UE::StateTree::Editor::DefaultRootColor, ColorAdjustment);
+}
+
+FLinearColor FStateTreeViewModel::GetHierarchyBasedBackgroundColor(TNotNull<const UStateTreeState*> State, EStateTreeViewModelColorAdjustment ColorAdjustment) const
+{
+	// If we are a leaf, then just use leaf color
+	if (State->Children.IsEmpty())
+	{
+		return GetColorAdjusted(UE::StateTree::Editor::DefaultStateColor, ColorAdjustment);
+	}
+
+	// We have 4 parent value levels and 1 root value level. Find our depth.
+	int8 NumParents = 0;
+	auto CountNumParents = [&](UStateTreeState& ParentState)
+	{
+		NumParents++;
+		return NumParents == 4 ? EStateTreeVisitor::Break : EStateTreeVisitor::Continue;
+	};
+	ForEachParent(const_cast<UStateTreeState*>(NotNullGet(State)), CountNumParents);
+
+	const FLinearColor DefaultParentColorHSV = UE::StateTree::Editor::DefaultParentColor.LinearRGBToHSV();
+	const FLinearColor BrightnessScale = { 0.0f, 0.0f, 1.0f, 0.0f };
+
+	FLinearColor Result = UE::StateTree::Editor::DefaultStateColor;
+	switch (NumParents)
+	{
+	case 0:
+		Result = UE::StateTree::Editor::DefaultRootColor;
+		break;
+	case 1:
+		Result = UE::StateTree::Editor::DefaultParentColor;
+		break;
+	case 2:
+		Result = (DefaultParentColorHSV - DefaultParentColorHSV * BrightnessScale + BrightnessScale * 0.515f).HSVToLinearRGB();
+		break;
+	case 3:
+		Result = (DefaultParentColorHSV - DefaultParentColorHSV * BrightnessScale + BrightnessScale * 0.300f).HSVToLinearRGB();
+		break;
+	case 4:
+		Result = (DefaultParentColorHSV - DefaultParentColorHSV * BrightnessScale + BrightnessScale * 0.150f).HSVToLinearRGB();
+		break;
+	default:
+		Result = (DefaultParentColorHSV - DefaultParentColorHSV * BrightnessScale + BrightnessScale * 0.150f).HSVToLinearRGB();
+		break;
+	}
+
+	return GetColorAdjusted(Result, ColorAdjustment);
+}
+
+FLinearColor FStateTreeViewModel::GetColorAdjusted(const FLinearColor& InColor, EStateTreeViewModelColorAdjustment ColorAdjustment) const
+{
+	FLinearColor ValueShift;
+	FLinearColor ValueScale;
+	FLinearColor InColorHSV = InColor.LinearRGBToHSV();
+
+	switch (ColorAdjustment)
+	{
+	case EStateTreeViewModelColorAdjustment::Normal:
+		return InColor;
+
+	case EStateTreeViewModelColorAdjustment::Bright:
+		ValueShift = { 0.0f, 0.0f, 0.5f, 0.0f };
+		ValueScale = { 1.0f, 1.0f, 0.5f, 1.0f };
+		return (InColorHSV * ValueScale + ValueShift).HSVToLinearRGB();
+
+	case EStateTreeViewModelColorAdjustment::Brighter:
+		ValueShift = { 0.0f, 0.0f, 0.7f, 0.0f };
+		ValueScale = { 1.0f, 1.0f, 0.3f, 1.0f };
+		return (InColorHSV * ValueScale + ValueShift).HSVToLinearRGB();
+
+	case EStateTreeViewModelColorAdjustment::Dark:
+		ValueScale = { 1.0f, 1.0f, 0.3f, 1.0f };
+		return (InColorHSV * ValueScale).HSVToLinearRGB();
+
+	case EStateTreeViewModelColorAdjustment::Darker:
+		ValueScale = { 1.0f, 1.0f, 0.18f, 1.0f };
+		return (InColorHSV * ValueScale).HSVToLinearRGB();
+
+	default:
+		return InColor;
+	}
+}
+
+const FSlateBrush* FStateTreeViewModel::GetSelectorIcon(TNotNull<const UStateTreeState*> State) const
+{
+	return FStateTreeEditorStyle::GetBrushForSelectionBehaviorType(State->SelectionBehavior, !State->Children.IsEmpty(), State->Type);
+}
+
+FText FStateTreeViewModel::GetNodeDescription(const FStateTreeEditorNode& Node, const EStateTreeNodeFormatting Formatting) const
+{
+	if (const UStateTreeEditorData* TreeData = TreeDataWeak.Get())
+	{
+		return TreeData->GetNodeDescription(Node, Formatting);
+	}
+
+	return FText::GetEmpty();
+}
+
 void FStateTreeViewModel::MoveSelectedStates(UStateTreeState* TargetState, const EStateTreeViewModelInsert RelativeLocation)
 {
 	UStateTreeEditorData* TreeData = TreeDataWeak.Get();
diff --git a/Engine/Plugins/Runtime/StateTree/Source/StateTreeEditorModule/Public/StateTreeEditorData.h b/Engine/Plugins/Runtime/StateTree/Source/StateTreeEditorModule/Public/StateTreeEditorData.h
--- a/Engine/Plugins/Runtime/StateTree/Source/StateTreeEditorModule/Public/StateTreeEditorData.h
+++ b/Engine/Plugins/Runtime/StateTree/Source/StateTreeEditorModule/Public/StateTreeEditorData.h
@@ -214,6 +214,23 @@ public:
 		return CastChecked<ExtensionType>(K2_GetExtension(ExtensionType::StaticClass()), ECastCheckedType::NullAllowed);
 	}
 
+	/** Find the first extension of the requested type or create it on this editor data if it doesn't exist. */
+	template<typename ExtensionType>
+	ExtensionType* GetOrCreateExtension()
+	{
+		static_assert(std::is_base_of_v<UStateTreeEditorDataExtension, ExtensionType>, "Using a non state tree editor data class as extension.");
+
+		if (ExtensionType* FoundExtension = CastChecked<ExtensionType>(K2_GetExtension(ExtensionType::StaticClass()), ECastCheckedType::NullAllowed))
+		{
+			return FoundExtension;
+		}
+
+		UStateTreeEditorDataExtension* Result = NewObject<UStateTreeEditorDataExtension>(this, ExtensionType::StaticClass());
+		Extensions.Add(Result);
+
+		return CastChecked<ExtensionType>(Result);
+	}
+
 	/** Find the first extension of the requested type. */
 	UFUNCTION(BlueprintCallable, Category = "StateTree|Extension", Meta = (DisplayName="Get Extension", DeterminesOutputType = "ExtensionType"))
 	UE_API UStateTreeEditorDataExtension* K2_GetExtension(TSubclassOf<UStateTreeEditorDataExtension> ExtensionType);
diff --git a/Engine/Plugins/Runtime/StateTree/Source/StateTreeEditorModule/Public/StateTreeViewModel.h b/Engine/Plugins/Runtime/StateTree/Source/StateTreeEditorModule/Public/StateTreeViewModel.h
--- a/Engine/Plugins/Runtime/StateTree/Source/StateTreeEditorModule/Public/StateTreeViewModel.h
+++ b/Engine/Plugins/Runtime/StateTree/Source/StateTreeEditorModule/Public/StateTreeViewModel.h
@@ -14,6 +14,7 @@ namespace UE::StateTreeEditor
 	struct FClipboardEditorData;
 }
 
+struct FStateTreeEditorNode;
 struct FStateTreeTransition;
 class FMenuBuilder;
 class UStateTree;
@@ -22,7 +23,10 @@ class UStateTreeState;
 
 enum class ECheckBoxState : uint8;
 enum class EStateTreeBreakpointType : uint8;
+enum class EStateTreeNodeFormatting : uint8;
+enum class EStateTreeVisitor : uint8;
 
+struct FSlateBrush;
 struct FPropertyChangedEvent;
 struct FStateTreeDebugger;
 struct FStateTreeDebuggerBreakpoint;
@@ -36,6 +40,15 @@ enum class EStateTreeViewModelInsert : uint8
 	Into,
 };
 
+enum class EStateTreeViewModelColorAdjustment : uint8
+{
+	Normal,
+	Bright,
+	Brighter,
+	Dark,
+	Darker,
+};
+
 enum class UE_DEPRECATED(5.6, "Use the enum with the E prefix") FStateTreeViewModelInsert : uint8
 {
 	Before,
@@ -77,8 +90,12 @@ public:
 	UE_API void GetSelectedStates(TArray<UStateTreeState*>& OutSelectedStates) const;
 	UE_API void GetSelectedStates(TArray<TWeakObjectPtr<UStateTreeState>>& OutSelectedStates) const;
 	UE_API bool HasSelection() const;
+	UE_API UStateTreeState* GetLastSelectedState() const;
 
 	UE_API void BringNodeToFocus(UStateTreeState* State, const FGuid NodeID);
+
+	// The view state is the last selected state, or first root if that is invalid. Nullptr if neither of those exist.
+	UE_API UStateTreeState* GetViewState() const;
 	
 	// Returns associated state tree asset.
 	UE_API const UStateTree* GetStateTree() const;
@@ -87,6 +104,9 @@ public:
 
 	UE_API const UStateTreeState* GetStateByID(const FGuid StateID) const;
 	UE_API UStateTreeState* GetMutableStateByID(const FGuid StateID) const;
+
+	UE_API const FStateTreeTransition* GetTransitionByID(const FGuid TransitionID) const;
+	UE_API FStateTreeTransition* GetMutableTransitionByID(const FGuid TransitionID) const;
 	
 	// Returns array of subtrees to edit.
 	UE_API TArray<TObjectPtr<UStateTreeState>>* GetSubTrees() const;
@@ -97,7 +117,7 @@ public:
 	UE_API void GetLinkStates(FGuid StateID, TArray<FGuid>& LinkingIn, TArray<FGuid>& LinkedOut) const;
 
 	/** Find the states that the transition go to. */
-	UE_API UStateTreeState* GetTransitionToState(TNotNull<const UStateTreeState*> State, FGuid TransitionID) const;
+	UE_API UStateTreeState* GetTransitionToState(TNotNull<const UStateTreeState*> TransitionOwningState, FGuid TransitionID) const;
 
 	// Gets and sets StateTree view expansion state store in the asset.
 	UE_API void SetPersistentExpandedStates(TSet<TWeakObjectPtr<UStateTreeState>>& InExpandedStates);
@@ -121,6 +141,32 @@ public:
 	UE_API bool CanPasteNodesToSelectedStates() const;
 	UE_API void SetSelectedStatesEnabled(bool bEnable);
 
+	// State utility methods
+	UE_API void ForEachParent(TNotNull<UStateTreeState*> State, TFunctionRef<EStateTreeVisitor(UStateTreeState& ParentState)> InFunc) const;
+	UE_API void ForEachInTransition(TNotNull<UStateTreeState*> State, TFunctionRef<EStateTreeVisitor(FStateTreeTransition& TransitionToState)> InFunc) const;
+	UE_API void ForEachOutTransition(TNotNull<UStateTreeState*> State, TFunctionRef<EStateTreeVisitor(FStateTreeTransition& TransitionFromState)> InFunc) const;
+	UE_API void ForEachChild(TNotNull<UStateTreeState*> State, TFunctionRef<EStateTreeVisitor(UStateTreeState& ChildState)> InFunc) const;
+
+	UE_API TArray<UStateTreeState*> GetStateParents(TNotNull<UStateTreeState*> State) const;
+	UE_API TArray<FStateTreeTransition*> GetStateInTransitions(TNotNull<UStateTreeState*> State) const;
+	UE_API TArray<FStateTreeTransition*> GetStateOutTransitions(TNotNull<UStateTreeState*> State) const;
+	UE_API TArray<UStateTreeState*> GetStateChildren(TNotNull<UStateTreeState*> State) const;
+
+	UE_API const UStateTreeState* GetSourceStateFromTransition(TNotNull<const FStateTreeTransition*> Transition) const;
+	UE_API const UStateTreeState* GetSourceParentStateFromTransition(TNotNull<const FStateTreeTransition*> Transition) const;
+	UE_API const UStateTreeState* GetTargetStateFromTransition(TNotNull<const FStateTreeTransition*> Transition) const;
+	UE_API const UStateTreeState* GetTargetParentStateFromTransition(TNotNull<const FStateTreeTransition*> Transition) const;
+
+	// Color manipulation
+	UE_API FLinearColor GetStateColor(TNotNull<const UStateTreeState*> State, EStateTreeViewModelColorAdjustment ColorAdjustment = EStateTreeViewModelColorAdjustment::Normal) const;
+	UE_API FLinearColor GetHierarchyBasedHighlightColor(TNotNull<const UStateTreeState*> State, EStateTreeViewModelColorAdjustment ColorAdjustment = EStateTreeViewModelColorAdjustment::Normal) const;
+	UE_API FLinearColor GetHierarchyBasedBackgroundColor(TNotNull<const UStateTreeState*> State, EStateTreeViewModelColorAdjustment ColorAdjustment = EStateTreeViewModelColorAdjustment::Normal) const;
+
+	UE_API FLinearColor GetColorAdjusted(const FLinearColor& InColor, EStateTreeViewModelColorAdjustment ColorAdjustment) const;
+
+	UE_API const FSlateBrush* GetSelectorIcon(TNotNull<const UStateTreeState*> State) const;
+	UE_API FText GetNodeDescription(const FStateTreeEditorNode& Node, const EStateTreeNodeFormatting Formatting) const;
+
 	// EditorNode and Transition manipulation
 	// @todo: support ReplaceWith and Rename
 	UE_API void DeleteNode(TWeakObjectPtr<UStateTreeState> State, const FGuid& ID);
@@ -222,6 +268,7 @@ protected:
 
 	TWeakObjectPtr<UStateTreeEditorData> TreeDataWeak;
 	TSet<TWeakObjectPtr<UStateTreeState>> SelectedStates;
+	FGuid LastSelectedState;
 
 #if WITH_STATETREE_TRACE_DEBUGGER
 	UE_API void HandleBreakpointsChanged(const UStateTree& StateTree);
