[State Tree] Chained Property Ref from multiple parent frames above can now correctly fetch source data. Now we fetch the associated parent frame that stores GlobalParameter or SubtreeParameter by FrameID.

- Disabled using PropertyRef in ExternalGlobalParameterData
- Deprecated and refactored FindFrame
- Deprecated and refactored FStateTreePropertyRef::GetMutablePtr
- Deprecated StateTreePropertyRefExternalHandle

#Synced-CL 48414399

diff --git a/Engine/Plugins/Runtime/StateTree/Source/StateTreeEditorModule/Private/StateTreePropertyBindingCompiler.cpp b/Engine/Plugins/Runtime/StateTree/Source/StateTreeEditorModule/Private/StateTreePropertyBindingCompiler.cpp
--- a/Engine/Plugins/Runtime/StateTree/Source/StateTreeEditorModule/Private/StateTreePropertyBindingCompiler.cpp
+++ b/Engine/Plugins/Runtime/StateTree/Source/StateTreeEditorModule/Private/StateTreePropertyBindingCompiler.cpp
@@ -374,12 +374,20 @@ bool FStateTreePropertyBindingCompiler::CompileReferences(const FStateTreeBindab
 			}
 		}
 
-		if (!ReferenceIndex.IsValid())
+		if (!ReferenceIndex.IsValid() && UE::StateTree::PropertyRefHelpers::IsPropertyRef(*SourceIndirections.Last().GetProperty()))
 		{
-			// If referencing another non global or subtree parameter PropertyRef, reuse it's index.
-			if (UE::StateTree::PropertyRefHelpers::IsPropertyRef(*SourceIndirections.Last().GetProperty())
-				&& SourceStruct->DataHandle.GetSource() != EStateTreeDataSourceType::GlobalParameterData
-				&& SourceStruct->DataHandle.GetSource() != EStateTreeDataSourceType::ExternalGlobalParameterData
+			if (SourceStruct->DataHandle.GetSource() == EStateTreeDataSourceType::ExternalGlobalParameterData)
+			{
+				Log->Reportf(EMessageSeverity::Error, TargetStruct, 
+					TEXT("ExternalGlobalParameter doesn't support chained Property Reference from %s to %s."), 
+					*UE::StateTree::GetDescAndPathAsString(*SourceStruct, Binding.GetSourcePath()), 
+					*UE::StateTree::GetDescAndPathAsString(TargetStruct, Binding.GetTargetPath()));
+
+				return false;
+			}
+
+			// If referencing another non-global or non-subtree parameter PropertyRef, reuse it's index.
+			if (SourceStruct->DataHandle.GetSource() != EStateTreeDataSourceType::GlobalParameterData
 				&& SourceStruct->DataHandle.GetSource() != EStateTreeDataSourceType::SubtreeParameterData)
 			{
 				const FCompiledReference* ReferencedReference = CompiledReferences.FindByPredicate([&Binding](const FCompiledReference& CompiledReference)
diff --git a/Engine/Plugins/Runtime/StateTree/Source/StateTreeEditorModule/Public/StateTreeEditorData.h b/Engine/Plugins/Runtime/StateTree/Source/StateTreeEditorModule/Public/StateTreeEditorData.h
--- a/Engine/Plugins/Runtime/StateTree/Source/StateTreeEditorModule/Public/StateTreeEditorData.h
+++ b/Engine/Plugins/Runtime/StateTree/Source/StateTreeEditorModule/Public/StateTreeEditorData.h
@@ -248,11 +248,12 @@ public:
 	 * Adds new Subtree with specified name.
 	 * @return Pointer to the new Subtree.
 	 */
-	UStateTreeState& AddSubTree(const FName Name)
+	UStateTreeState& AddSubTree(const FName Name, EStateTreeStateType StateType = EStateTreeStateType::State)
 	{
 		UStateTreeState* SubTreeState = NewObject<UStateTreeState>(this, FName(), RF_Transactional);
 		check(SubTreeState);
 		SubTreeState->Name = Name;
+		SubTreeState->Type = StateType;
 		SubTrees.Add(SubTreeState);
 		return *SubTreeState;
 	}
diff --git a/Engine/Plugins/Runtime/StateTree/Source/StateTreeModule/Private/Blueprint/StateTreeNodeBlueprintBase.cpp b/Engine/Plugins/Runtime/StateTree/Source/StateTreeModule/Private/Blueprint/StateTreeNodeBlueprintBase.cpp
--- a/Engine/Plugins/Runtime/StateTree/Source/StateTreeModule/Private/Blueprint/StateTreeNodeBlueprintBase.cpp
+++ b/Engine/Plugins/Runtime/StateTree/Source/StateTreeModule/Private/Blueprint/StateTreeNodeBlueprintBase.cpp
@@ -60,21 +60,27 @@ void* UStateTreeNodeBlueprintBase::GetMutablePtrToProperty(const FStateTreeBluep
 {
 	using namespace UE::StateTree;
 	const FStateTreeStrongExecutionContext StrongContext = WeakExecutionContext.MakeStrongExecutionContext();
-
-	const Async::FActivePathInfo ActivePath = StrongContext.GetActivePathInfo();
-	if (!ActivePath.IsValid())
+	if (StrongContext.IsValid())
 	{
-		UE_VLOG_UELOG(this, LogStateTree, Error, TEXT("Trying to GetMutablePtrToProperty while node is not active."));
-		return nullptr;
-	}
+		const Async::FActivePathInfo ActivePath = StrongContext.GetActivePathInfo();
+		if (!ActivePath.IsValid())
+		{
+			UE_VLOG_UELOG(this, LogStateTree, Verbose, TEXT("Trying to GetMutablePtrToProperty while node is not active."));
+			return nullptr;
+		}
 
-	check(StrongContext.IsValid());
-	const FProperty* SourceProperty = nullptr;
-	void* PropertyAddress = PropertyRefHelpers::GetMutablePtrToProperty<void>(PropertyRef, *const_cast<FStateTreeInstanceStorage*>(StrongContext.GetStorage().Get()), *ActivePath.Frame, ActivePath.ParentFrame, &SourceProperty);
-	if (PropertyAddress && PropertyRefHelpers::IsBlueprintPropertyRefCompatibleWithProperty(*SourceProperty, &PropertyRef))
-	{
-		OutSourceProperty = const_cast<FProperty*>(SourceProperty);
-		return PropertyAddress;
+		const FProperty* SourceProperty = nullptr;
+		void* PropertyAddress = PropertyRefHelpers::GetMutablePtrToProperty<void>(PropertyRef,
+			*const_cast<FStateTreeInstanceStorage*>(StrongContext.GetStorage().Get()),
+			StrongContext.GetTemporaryStorage().Get(),
+			*ActivePath.Frame,
+			&SourceProperty);
+
+		if (PropertyAddress && PropertyRefHelpers::IsBlueprintPropertyRefCompatibleWithProperty(*SourceProperty, &PropertyRef))
+		{
+			OutSourceProperty = const_cast<FProperty*>(SourceProperty);
+			return PropertyAddress;
+		}
 	}
 
 	OutSourceProperty = nullptr;
diff --git a/Engine/Plugins/Runtime/StateTree/Source/StateTreeModule/Private/StateTreeExecutionContext.cpp b/Engine/Plugins/Runtime/StateTree/Source/StateTreeModule/Private/StateTreeExecutionContext.cpp
--- a/Engine/Plugins/Runtime/StateTree/Source/StateTreeModule/Private/StateTreeExecutionContext.cpp
+++ b/Engine/Plugins/Runtime/StateTree/Source/StateTreeModule/Private/StateTreeExecutionContext.cpp
@@ -109,27 +109,6 @@ namespace UE::StateTree::ExecutionContext::Private
 			return Number;
 	}
 
-	// @todo: should use const TConstArrayView instead
-	template<typename FrameType>
-	FrameType* FindExecutionFrame(FActiveFrameID FrameID, const TArrayView<FrameType> Frames, const TArrayView<FrameType> TemporaryFrames)
-	{
-		FrameType* Frame = Frames.FindByPredicate(
-			[FrameID](const FStateTreeExecutionFrame& Frame)
-			{
-				return Frame.FrameID == FrameID;
-			});
-		if (Frame)
-		{
-			return Frame;
-		}
-
-		return TemporaryFrames.FindByPredicate(
-			[FrameID](const FStateTreeExecutionFrame& Frame)
-			{
-				return Frame.FrameID == FrameID;
-			});
-	}
-
 	const FCompactStateTreeFrame* FindStateTreeFrame(const FExecutionFrameHandle& FrameHandle)
 	{
 		check(FrameHandle.IsValid());
@@ -1005,6 +984,11 @@ UE::StateTree::FActiveState FStateTreeExecutionContext::FSelectStateResult::GetS
 	return {};
 }
 
+TConstArrayView<FStateTreeExecutionFrame> FStateTreeExecutionContext::FSelectStateResult::GetTemporaryFrames() const
+{
+	return TemporaryFrames;
+}
+
 /**
  * FStateTreeExecutionContext::FCurrentlyProcessedFrameScope implementation
  */
@@ -2368,7 +2352,7 @@ void FStateTreeExecutionContext::UpdateInstanceData(const TSharedPtr<FSelectStat
 			// Global
 			if (SelectedState.GetFrameID() != CurrentFrameID)
 			{
-				ParentFrame = FindExecutionFrame(CurrentFrameID, MakeArrayView(Exec.ActiveFrames), MakeArrayView(Args->TemporaryFrames));
+				ParentFrame = FindExecutionFrame(CurrentFrameID, MakeConstArrayView(Exec.ActiveFrames), MakeConstArrayView(Args->TemporaryFrames));
 				CurrentFrameID = SelectedState.GetFrameID();
 				CurrentFrame = Exec.FindActiveFrame(CurrentFrameID);
 
@@ -2510,7 +2494,7 @@ void FStateTreeExecutionContext::UpdateInstanceData(const TSharedPtr<FSelectStat
 									if (!bIsLastFrame)
 									{
 										const FActiveState& FollowingSelectedState = Args->SelectedStates[SelectedStateIndex + 1];
-										const FStateTreeExecutionFrame* FollowingNextFrame = FindExecutionFrame(FollowingSelectedState.GetFrameID(), MakeArrayView(Exec.ActiveFrames), MakeArrayView(Args->TemporaryFrames));
+										const FStateTreeExecutionFrame* FollowingNextFrame = UE::StateTree::ExecutionContext::FindExecutionFrame(FollowingSelectedState.GetFrameID(), MakeConstArrayView(Exec.ActiveFrames), MakeConstArrayView(Args->TemporaryFrames));
 										if (ensure(FollowingNextFrame))
 										{
 											ParamsInstanceData = FConstStructView::Make(TempParams.Emplace_GetRef(FollowingNextFrame->StateTree->GetDefaultParameters()));
@@ -3005,6 +2989,11 @@ const FStateTreeExecutionFrame* FStateTreeExecutionContext::FindFrame(const USta
 	return &Frames[FrameIndex];
 }
 
+FStateTreeExecutionFrame* FStateTreeExecutionContext::FindMutableExecutionFrame(UE::StateTree::FActiveFrameID InFrameID, TArrayView<FStateTreeExecutionFrame> InActiveFrames, TArrayView<FStateTreeExecutionFrame> InTemporaryFrames)
+{
+	return const_cast<FStateTreeExecutionFrame*>(UE::StateTree::ExecutionContext::FindExecutionFrame(InFrameID, InActiveFrames, InTemporaryFrames));
+}
+
 bool FStateTreeExecutionContext::IsHandleSourceValid(const FStateTreeExecutionFrame* ParentFrame, const FStateTreeExecutionFrame& CurrentFrame, const FStateTreeDataHandle Handle) const
 {
 	switch (Handle.GetSource())
@@ -3457,7 +3446,7 @@ void FStateTreeExecutionContext::CaptureNewStateEvents(const TSharedRef<const FS
 		if (SelectedState.GetFrameID() != CurrentFrameID)
 		{
 			CurrentFrameID = SelectedState.GetFrameID();
-			CurrentFrame = FindExecutionFrame(SelectedState.GetFrameID(), MakeConstArrayView(Exec.ActiveFrames), MakeConstArrayView(Args->TemporaryFrames));
+			CurrentFrame = UE::StateTree::ExecutionContext::FindExecutionFrame(SelectedState.GetFrameID(), MakeConstArrayView(Exec.ActiveFrames), MakeConstArrayView(Args->TemporaryFrames));
 		}
 
 		check(CurrentFrame);
@@ -5436,7 +5425,7 @@ bool FStateTreeExecutionContext::RequestTransitionInternal(
 					}
 
 					FStateTreeExecutionState& Exec = Storage.GetMutableExecutionState();
-					const FStateTreeExecutionFrame* ParentFrame = FindExecutionFrame(TemporaryFrameAndParent.ParentFrameID, MakeConstArrayView(Exec.ActiveFrames), MakeConstArrayView(PrevSelectStateResult->TemporaryFrames));
+					const FStateTreeExecutionFrame* ParentFrame = ExecutionContext::FindExecutionFrame(TemporaryFrameAndParent.ParentFrameID, MakeConstArrayView(Exec.ActiveFrames), MakeConstArrayView(PrevSelectStateResult->TemporaryFrames));
 
 					if (TemporaryFrame->bIsGlobalFrame && TemporaryFrame->bHaveEntered)
 					{
@@ -5469,7 +5458,7 @@ PRAGMA_DISABLE_DEPRECATION_WARNINGS
 				if (SelectedState.GetFrameID() != PreviousFrameID)
 				{
 					PreviousFrameID = SelectedState.GetFrameID();
-					const FStateTreeExecutionFrame* Frame = FindExecutionFrame(SelectedState.GetFrameID(), MakeConstArrayView(Exec.ActiveFrames), MakeConstArrayView(RequestedTransition->Selection->TemporaryFrames));
+					const FStateTreeExecutionFrame* Frame = UE::StateTree::ExecutionContext::FindExecutionFrame(SelectedState.GetFrameID(), MakeConstArrayView(Exec.ActiveFrames), MakeConstArrayView(RequestedTransition->Selection->TemporaryFrames));
 					if (ensure(Frame))
 					{
 						RequestedTransition->Transition.NextActiveFrames.Add(*Frame);
@@ -5503,7 +5492,7 @@ PRAGMA_ENABLE_DEPRECATION_WARNINGS
 		// Consume events from states, if required. The transition might also want to consume the event. (Not a bug.)
 		for (FSelectionEventWithID& Event : RequestedTransition->Selection->SelectionEvents)
 		{
-			const FStateTreeExecutionFrame* Frame = FindExecutionFrame(Event.State.GetFrameID(), MakeConstArrayView(Exec.ActiveFrames), MakeConstArrayView(RequestedTransition->Selection->TemporaryFrames));
+			const FStateTreeExecutionFrame* Frame = UE::StateTree::ExecutionContext::FindExecutionFrame(Event.State.GetFrameID(), MakeConstArrayView(Exec.ActiveFrames), MakeConstArrayView(RequestedTransition->Selection->TemporaryFrames));
 			if (ensure(Frame))
 			{
 				const FCompactStateTreeState* State = Frame->StateTree->GetStateFromHandle(Event.State.GetStateHandle());
@@ -5518,7 +5507,7 @@ PRAGMA_ENABLE_DEPRECATION_WARNINGS
 		{
 			check(RequestedTransition->Selection->SelectedStates.Num() > 0);
 			const FActiveState& LastSelectedState = RequestedTransition->Selection->SelectedStates.Last();
-			const FStateTreeExecutionFrame* SelectedFrame = FindExecutionFrame(LastSelectedState.GetFrameID(), MakeConstArrayView(Exec.ActiveFrames), MakeConstArrayView(RequestedTransition->Selection->TemporaryFrames));
+			const FStateTreeExecutionFrame* SelectedFrame = UE::StateTree::ExecutionContext::FindExecutionFrame(LastSelectedState.GetFrameID(), MakeConstArrayView(Exec.ActiveFrames), MakeConstArrayView(RequestedTransition->Selection->TemporaryFrames));
 			check(SelectedFrame);
 			STATETREE_LOG(Verbose, TEXT("Transition Request. Source:'%s'. Target:'%s'. Selected:'%s'"),
 				*GetSafeStateName(SourceFrame, SourceStateHandle),
@@ -6383,7 +6372,7 @@ FRecordedStateTreeTransitionResult FStateTreeExecutionContext::MakeRecordedTrans
 	Recorded.States.Reserve(SelectStateResult->SelectedStates.Num());
 	for (const UE::StateTree::FActiveState& SelectedState : SelectStateResult->SelectedStates)
 	{
-		const FStateTreeExecutionFrame* ProcessFrame = FindExecutionFrame(SelectedState.GetFrameID(), MakeConstArrayView(Exec.ActiveFrames), MakeConstArrayView(SelectStateResult->TemporaryFrames));
+		const FStateTreeExecutionFrame* ProcessFrame = UE::StateTree::ExecutionContext::FindExecutionFrame(SelectedState.GetFrameID(), MakeConstArrayView(Exec.ActiveFrames), MakeConstArrayView(SelectStateResult->TemporaryFrames));
 		if (ProcessFrame == nullptr)
 		{
 			return FRecordedStateTreeTransitionResult();
@@ -6578,7 +6567,7 @@ bool FStateTreeExecutionContext::SelectState(const FSelectStateArguments& Args,
 		}
 	}
 	checkf(SourceStates.Num() == 0 || FActiveStatePath::StartsWith(Args.ActiveStates, SourceStates), TEXT("Source is part of the active path."));
-	FStateTreeExecutionFrame* SourceExecFrame = FindExecutionFrame(Args.SourceState.GetFrameID(), MakeArrayView(Exec.ActiveFrames), MakeArrayView(OutSelectionResult->TemporaryFrames));
+	const FStateTreeExecutionFrame* SourceExecFrame = UE::StateTree::ExecutionContext::FindExecutionFrame(Args.SourceState.GetFrameID(), MakeConstArrayView(Exec.ActiveFrames), MakeConstArrayView(OutSelectionResult->TemporaryFrames));
 	if (!ensure(SourceExecFrame))
 	{
 		return false;
@@ -6625,7 +6614,7 @@ bool FStateTreeExecutionContext::SelectState(const FSelectStateArguments& Args,
 					const FActiveFrameID CurrentActiveFrameID = SourceStates[FoundActiveStateIndex].GetFrameID();
 					if (CurrentActiveFrameID != CurrentFrameID)
 					{
-						FStateTreeExecutionFrame* ProcessFrame = FindExecutionFrame(CurrentActiveFrameID, MakeArrayView(Exec.ActiveFrames), MakeArrayView(OutSelectionResult->TemporaryFrames));
+						const FStateTreeExecutionFrame* ProcessFrame = UE::StateTree::ExecutionContext::FindExecutionFrame(CurrentActiveFrameID, MakeConstArrayView(Exec.ActiveFrames), MakeConstArrayView(OutSelectionResult->TemporaryFrames));
 						if (!ensure(ProcessFrame))
 						{
 							return false;
@@ -6834,7 +6823,7 @@ bool FStateTreeExecutionContext::SelectStateInternal(
 			}
 
 			const FActiveState PreviousSelectedState = OutSelectionResult->SelectedStates.Last();
-			FStateTreeExecutionFrame* PreviousFramePtr = FindExecutionFrame(PreviousSelectedState.GetFrameID(), MakeArrayView(GetExecState().ActiveFrames), MakeArrayView(OutSelectionResult->TemporaryFrames));
+			FStateTreeExecutionFrame* PreviousFramePtr = FindMutableExecutionFrame(PreviousSelectedState.GetFrameID(), MakeArrayView(GetExecState().ActiveFrames), MakeArrayView(OutSelectionResult->TemporaryFrames));
 			check(PreviousFramePtr);
 			check(PreviousFramePtr->StateTree == NextStateTree);
 
@@ -6852,7 +6841,7 @@ bool FStateTreeExecutionContext::SelectStateInternal(
 		{
 			// We are building the path toward the desired target.
 			const FActiveFrameID MissingFrameID = InternalArgs.MissingSourceStates.Num() != 0 ? InternalArgs.MissingSourceStates[0].GetFrameID() : InternalArgs.MissingSourceFrameID;
-			FStateTreeExecutionFrame* MissingFramePtr = FindExecutionFrame(MissingFrameID, MakeArrayView(GetExecState().ActiveFrames), MakeArrayView(OutSelectionResult->TemporaryFrames));
+			FStateTreeExecutionFrame* MissingFramePtr = FindMutableExecutionFrame(MissingFrameID, MakeArrayView(GetExecState().ActiveFrames), MakeArrayView(OutSelectionResult->TemporaryFrames));
 			check(MissingFramePtr);
 			if (!ensure(MissingFramePtr->StateTree == NextStateTree))
 			{
@@ -6912,12 +6901,12 @@ bool FStateTreeExecutionContext::SelectStateInternal(
 	const FStateTreeExecutionFrame* NextParentFrame = nullptr;
 	auto CacheNextFrame = [&NextFrame, &NextParentFrame, &NextFrameID, &Exec, &OutSelectionResult]()
 		{
-			NextFrame = FindExecutionFrame(NextFrameID, MakeConstArrayView(Exec.ActiveFrames), MakeConstArrayView(OutSelectionResult->TemporaryFrames));
+			NextFrame = UE::StateTree::ExecutionContext::FindExecutionFrame(NextFrameID, MakeConstArrayView(Exec.ActiveFrames), MakeConstArrayView(OutSelectionResult->TemporaryFrames));
 			check(NextFrame);
 
 	if (OutSelectionResult->SelectedFrames.Num() > 1)
 	{
-		NextParentFrame = FindExecutionFrame(OutSelectionResult->SelectedFrames.Last(1), MakeConstArrayView(Exec.ActiveFrames), MakeConstArrayView(OutSelectionResult->TemporaryFrames));
+		NextParentFrame = UE::StateTree::ExecutionContext::FindExecutionFrame(OutSelectionResult->SelectedFrames.Last(1), MakeConstArrayView(Exec.ActiveFrames), MakeConstArrayView(OutSelectionResult->TemporaryFrames));
 		check(NextParentFrame);
 	}
 		};
@@ -7325,22 +7314,22 @@ TOptional<bool> FStateTreeExecutionContext::SelectStateFromSourceInternal(
 namespace UE::StateTree::ExecutionContext::Private
 {
 	bool PreventRecursionCheck(const FExecutionFrameHandle& LinkStateFrameHandle,
-		TArrayView<const FActiveFrameID> SelectedFrames,
-		TArrayView<const FStateTreeExecutionFrame> ActiveFrames,
-		TArrayView<const FStateTreeExecutionFrame> TemporaryFrames)
+		TConstArrayView<const FActiveFrameID> SelectedFrames,
+		TConstArrayView<const FStateTreeExecutionFrame> ActiveFrames,
+		TConstArrayView<const FStateTreeExecutionFrame> TemporaryFrames)
 	{
 		// Check and prevent recursion.
 		return SelectedFrames.ContainsByPredicate(
 			[&LinkStateFrameHandle, &ActiveFrames, &TemporaryFrames](const FActiveFrameID& FrameID)
 			{
-				const FStateTreeExecutionFrame* Frame = FindExecutionFrame(FrameID, ActiveFrames, TemporaryFrames);
+				const FStateTreeExecutionFrame* Frame = UE::StateTree::ExecutionContext::FindExecutionFrame(FrameID, ActiveFrames, TemporaryFrames);
 				return Frame && Frame->HasRoot(LinkStateFrameHandle);
 			});
 	}
 
 	FStateTreeExecutionFrame* SelectedFrameLinkedFrame(FStateTreeExecutionState& Exec,
 		const bool bShouldCreateNewState,
-		TArrayView<const UE::StateTree::FActiveState> MatchingActiveStates,
+		TConstArrayView<const UE::StateTree::FActiveState> MatchingActiveStates,
 		const FExecutionFrameHandle& LinkStateFrameHandle,
 		const EStateTreeStateSelectionRules StateSelectionRules
 	)
@@ -7468,7 +7457,7 @@ bool FStateTreeExecutionContext::SelectStateInternal_Linked(
 		FStateTreeExecutionFrame& NewFrame = OutSelectionResult->MakeAndAddTemporaryFrame(FActiveFrameID(Storage.GenerateUniqueId()), LinkStateFrameHandle, bIsGlobalFrame);
 
 		// Note. Adding to TemporaryFrame can invalidate ParentFrame.
-		const FStateTreeExecutionFrame* NextParentFrame = FindExecutionFrame(OutSelectionResult->SelectedFrames.Last(), MakeConstArrayView(Exec.ActiveFrames), MakeConstArrayView(OutSelectionResult->TemporaryFrames));
+		const FStateTreeExecutionFrame* NextParentFrame = UE::StateTree::ExecutionContext::FindExecutionFrame(OutSelectionResult->SelectedFrames.Last(), MakeConstArrayView(Exec.ActiveFrames), MakeConstArrayView(OutSelectionResult->TemporaryFrames));
 		check(NextParentFrame);
 
 		NewFrame.ExternalDataBaseIndex = NextParentFrame->ExternalDataBaseIndex;
@@ -7650,7 +7639,7 @@ bool FStateTreeExecutionContext::SelectStateInternal_LinkedAsset(
 		FStateTreeExecutionFrame& NewFrame = OutSelectionResult->MakeAndAddTemporaryFrame(FActiveFrameID(Storage.GenerateUniqueId()), LinkStateFrameHandle, bIsGlobalFrame);
 
 		// Adding to TemporaryFrame can invalidate ParentFrame.
-		const FStateTreeExecutionFrame* NextParentFrame = FindExecutionFrame(OutSelectionResult->SelectedFrames.Last(), MakeConstArrayView(Exec.ActiveFrames), MakeConstArrayView(OutSelectionResult->TemporaryFrames));
+		const FStateTreeExecutionFrame* NextParentFrame = UE::StateTree::ExecutionContext::FindExecutionFrame(OutSelectionResult->SelectedFrames.Last(), MakeConstArrayView(Exec.ActiveFrames), MakeConstArrayView(OutSelectionResult->TemporaryFrames));
 		check(NextParentFrame);
 
 		NewFrame.ExternalDataBaseIndex = ExternalDataBaseIndex;
@@ -7675,10 +7664,10 @@ bool FStateTreeExecutionContext::SelectStateInternal_LinkedAsset(
 	if (bIsNewFrame)
 	{
 		// Adding to TemporaryFrame can invalidate ParentFrame.
-		const FStateTreeExecutionFrame* NextParentFrame = FindExecutionFrame(OutSelectionResult->SelectedFrames.Last(1), MakeConstArrayView(Exec.ActiveFrames), MakeConstArrayView(OutSelectionResult->TemporaryFrames));
+		const FStateTreeExecutionFrame* NextParentFrame = UE::StateTree::ExecutionContext::FindExecutionFrame(OutSelectionResult->SelectedFrames.Last(1), MakeConstArrayView(Exec.ActiveFrames), MakeConstArrayView(OutSelectionResult->TemporaryFrames));
 		check(NextParentFrame);
 
-		FStateTreeExecutionFrame* NextFrame = FindExecutionFrame(SelectedFrame->FrameID, MakeArrayView(Exec.ActiveFrames), MakeArrayView(OutSelectionResult->TemporaryFrames));
+		FStateTreeExecutionFrame* NextFrame = FindMutableExecutionFrame(SelectedFrame->FrameID, MakeArrayView(Exec.ActiveFrames), MakeArrayView(OutSelectionResult->TemporaryFrames));
 		check(NextFrame && SelectedFrame == NextFrame);
 
 		// Start global tasks and evaluators temporarily, so that their data is available already during select.
@@ -7718,8 +7707,8 @@ bool FStateTreeExecutionContext::SelectStateInternal_LinkedAsset(
 		if (bIsNewFrame)
 		{
 			// Adding to TemporaryFrame can invalidate ParentFrame
-			const FStateTreeExecutionFrame* NextParentFrame = FindExecutionFrame(OutSelectionResult->SelectedFrames.Last(1), MakeArrayView(Exec.ActiveFrames), MakeArrayView(OutSelectionResult->TemporaryFrames));
-			FStateTreeExecutionFrame* NextFrame = FindExecutionFrame(OutSelectionResult->SelectedFrames.Last(), MakeArrayView(Exec.ActiveFrames), MakeArrayView(OutSelectionResult->TemporaryFrames));
+			const FStateTreeExecutionFrame* NextParentFrame = UE::StateTree::ExecutionContext::FindExecutionFrame(OutSelectionResult->SelectedFrames.Last(1), MakeArrayView(Exec.ActiveFrames), MakeArrayView(OutSelectionResult->TemporaryFrames));
+			FStateTreeExecutionFrame* NextFrame = FindMutableExecutionFrame(OutSelectionResult->SelectedFrames.Last(), MakeArrayView(Exec.ActiveFrames), MakeArrayView(OutSelectionResult->TemporaryFrames));
 
 			constexpr EStateTreeRunStatus CompletionStatus = EStateTreeRunStatus::Stopped;
 			StopTemporaryEvaluatorsAndGlobalTasks(NextParentFrame, *NextFrame, CompletionStatus);
@@ -7838,9 +7827,9 @@ bool FStateTreeExecutionContext::SelectStateInternal_TrySelectChildrenWithHighes
 	{
 		using namespace UE::StateTree::ExecutionContext::Private;
 		const FStateTreeExecutionFrame* NextParentFrame = OutSelectionResult->SelectedFrames.Num() > 1
-			? FindExecutionFrame(OutSelectionResult->SelectedFrames.Last(1), MakeConstArrayView(GetExecState().ActiveFrames), MakeConstArrayView(OutSelectionResult->TemporaryFrames))
+			? UE::StateTree::ExecutionContext::FindExecutionFrame(OutSelectionResult->SelectedFrames.Last(1), MakeConstArrayView(GetExecState().ActiveFrames), MakeConstArrayView(OutSelectionResult->TemporaryFrames))
 			: nullptr;
-		FStateTreeExecutionFrame* NextFrame = FindExecutionFrame(OutSelectionResult->SelectedFrames.Last(), MakeArrayView(GetExecState().ActiveFrames), MakeArrayView(OutSelectionResult->TemporaryFrames));
+		const FStateTreeExecutionFrame* NextFrame = UE::StateTree::ExecutionContext::FindExecutionFrame(OutSelectionResult->SelectedFrames.Last(), MakeConstArrayView(GetExecState().ActiveFrames), MakeConstArrayView(OutSelectionResult->TemporaryFrames));
 		check(NextFrame);
 
 		for (uint16 ChildState = NextState.ChildrenBegin; ChildState < NextState.ChildrenEnd; ChildState = NextStateTree->States[ChildState].GetNextSibling())
@@ -7913,9 +7902,9 @@ bool FStateTreeExecutionContext::SelectStateInternal_TrySelectChildrenAtRandomWe
 	{
 		using namespace UE::StateTree::ExecutionContext::Private;
 		const FStateTreeExecutionFrame* NextParentFrame = OutSelectionResult->SelectedFrames.Num() > 1
-			? FindExecutionFrame(OutSelectionResult->SelectedFrames.Last(1), MakeConstArrayView(GetExecState().ActiveFrames), MakeConstArrayView(OutSelectionResult->TemporaryFrames))
+			? UE::StateTree::ExecutionContext::FindExecutionFrame(OutSelectionResult->SelectedFrames.Last(1), MakeConstArrayView(GetExecState().ActiveFrames), MakeConstArrayView(OutSelectionResult->TemporaryFrames))
 			: nullptr;
-		FStateTreeExecutionFrame* NextFrame = FindExecutionFrame(OutSelectionResult->SelectedFrames.Last(), MakeArrayView(GetExecState().ActiveFrames), MakeArrayView(OutSelectionResult->TemporaryFrames));
+		const FStateTreeExecutionFrame* NextFrame = UE::StateTree::ExecutionContext::FindExecutionFrame(OutSelectionResult->SelectedFrames.Last(), MakeConstArrayView(GetExecState().ActiveFrames), MakeConstArrayView(OutSelectionResult->TemporaryFrames));
 		check(NextFrame);
 
 		for (uint16 ChildState = NextState.ChildrenBegin; ChildState < NextState.ChildrenEnd; ChildState = NextStateTree->States[ChildState].GetNextSibling())
@@ -8042,9 +8031,9 @@ bool FStateTreeExecutionContext::SelectStateInternal_TryFollowTransitions(
 			{
 				using namespace UE::StateTree::ExecutionContext::Private;
 				const FStateTreeExecutionFrame* NextParentFrame = OutSelectionResult->SelectedFrames.Num() > 1
-					? FindExecutionFrame(OutSelectionResult->SelectedFrames.Last(1), MakeConstArrayView(GetExecState().ActiveFrames), MakeConstArrayView(OutSelectionResult->TemporaryFrames))
+					? UE::StateTree::ExecutionContext::FindExecutionFrame(OutSelectionResult->SelectedFrames.Last(1), MakeConstArrayView(GetExecState().ActiveFrames), MakeConstArrayView(OutSelectionResult->TemporaryFrames))
 					: nullptr;
-				FStateTreeExecutionFrame* NextFrame = FindExecutionFrame(OutSelectionResult->SelectedFrames.Last(), MakeArrayView(GetExecState().ActiveFrames), MakeArrayView(OutSelectionResult->TemporaryFrames));
+				const FStateTreeExecutionFrame* NextFrame = UE::StateTree::ExecutionContext::FindExecutionFrame(OutSelectionResult->SelectedFrames.Last(), MakeConstArrayView(GetExecState().ActiveFrames), MakeConstArrayView(OutSelectionResult->TemporaryFrames));
 				check(NextFrame);
 
 				FCurrentlyProcessedTransitionEventScope TransitionEventScope(*this, SelectedStateTransitionEvent.IsValid() ? SelectedStateTransitionEvent.Get() : nullptr);
diff --git a/Engine/Plugins/Runtime/StateTree/Source/StateTreeModule/Private/StateTreeExecutionContextTypes.cpp b/Engine/Plugins/Runtime/StateTree/Source/StateTreeModule/Private/StateTreeExecutionContextTypes.cpp
new file mode 100644
--- /dev/null
+++ b/Engine/Plugins/Runtime/StateTree/Source/StateTreeModule/Private/StateTreeExecutionContextTypes.cpp
@@ -0,0 +1,27 @@
+// Copyright Epic Games, Inc. All Rights Reserved.
+
+#include "StateTreeExecutionContextTypes.h"
+
+namespace UE::StateTree::ExecutionContext
+{
+	const FStateTreeExecutionFrame* FindExecutionFrame(FActiveFrameID InFrameID,
+		TConstArrayView<FStateTreeExecutionFrame> InActiveFrames, TConstArrayView<FStateTreeExecutionFrame> InTemporaryFrames)
+	{
+		const FStateTreeExecutionFrame* Frame = InActiveFrames.FindByPredicate(
+			[InFrameID](const FStateTreeExecutionFrame& Frame)
+			{
+				return Frame.FrameID == InFrameID;
+			});
+		if (Frame)
+		{
+			return Frame;
+		}
+	
+		return InTemporaryFrames.FindByPredicate(
+			[InFrameID](const FStateTreeExecutionFrame& Frame)
+			{
+				return Frame.FrameID == InFrameID;
+			});
+	
+	}
+}
diff --git a/Engine/Plugins/Runtime/StateTree/Source/StateTreeModule/Private/StateTreePropertyRef.cpp b/Engine/Plugins/Runtime/StateTree/Source/StateTreeModule/Private/StateTreePropertyRef.cpp
new file mode 100644
--- /dev/null
+++ b/Engine/Plugins/Runtime/StateTree/Source/StateTreeModule/Private/StateTreePropertyRef.cpp
@@ -0,0 +1,61 @@
+// Copyright Epic Games, Inc. All Rights Reserved.
+
+#include "StateTreePropertyRef.h"
+
+namespace UE::StateTree::PropertyRefHelpers::Private
+{
+	TOptional<FResolvePropertyReferenceIndirectionsResult> ResolvePropertyReferenceIndirections(
+		const FStateTreePropertyRef& InPropertyRef,
+		FStateTreeInstanceStorage& InstanceDataStorage,
+		const ExecutionContext::ITemporaryStorage* InTemporaryStorage,
+		const FStateTreeExecutionFrame& InExecutionFrame)
+	{
+		const FStateTreePropertyBindings& PropertyBindings = InExecutionFrame.StateTree->GetPropertyBindings();
+		if (const FStateTreePropertyAccess* PropertyAccess = PropertyBindings.GetPropertyAccess(InPropertyRef))
+		{
+			// Passing empty ContextAndExternalDataViews, as PropertyRef is not allowed to point to context or external data.
+			const FStateTreeDataView SourceView = InstanceData::GetDataViewOrTemporary(InstanceDataStorage, nullptr, InExecutionFrame, PropertyAccess->SourceDataHandle);
+
+			if (IsPropertyRef(*PropertyAccess->SourceLeafProperty))
+			{
+				// The only possibility when PropertyRef references another PropertyRef is when source one is a global or subtree parameter, i.e lives in parent execution frame.
+				// If that's the case, referenced PropertyRef is obtained and we recursively take the address where it points to.
+				FActiveFrameID ParentFrameID;
+				switch (PropertyAccess->SourceDataHandle.GetSource())
+				{
+				case EStateTreeDataSourceType::GlobalParameterData:
+					ParentFrameID = InExecutionFrame.GlobalParameterDataFrameID;
+					break;
+				case EStateTreeDataSourceType::SubtreeParameterData:
+					ParentFrameID = InExecutionFrame.StateParameterDataFrameID;
+					break;
+				default:
+					checkf(false, TEXT("Property Reference chaining should happen in Linked Subtree or Linked Asset!"));
+					return {};
+				}
+
+				const TConstArrayView<FStateTreeExecutionFrame> ActiveFrames = InstanceDataStorage.GetExecutionState().ActiveFrames;
+				const TConstArrayView<FStateTreeExecutionFrame> TemporaryFrames = InTemporaryStorage ? InTemporaryStorage->GetTemporaryFrames() : TConstArrayView<FStateTreeExecutionFrame>();
+				const FStateTreeExecutionFrame* ParentFrame = ExecutionContext::FindExecutionFrame(ParentFrameID, ActiveFrames, TemporaryFrames);
+
+				if (ParentFrame == nullptr)
+				{
+					return {};
+				}
+
+				const FStateTreePropertyRef* ReferencedPropertyRef = PropertyBindings.GetMutablePropertyPtr<FStateTreePropertyRef>(SourceView, *PropertyAccess);
+				if (ReferencedPropertyRef == nullptr)
+				{
+					return {};
+				}
+
+				return ResolvePropertyReferenceIndirections(*ReferencedPropertyRef, InstanceDataStorage, InTemporaryStorage, *ParentFrame);
+			}
+
+			return TOptional<FResolvePropertyReferenceIndirectionsResult>(InPlace, SourceView, PropertyAccess, &PropertyBindings);
+		}
+
+		return {};
+	}
+
+}
diff --git a/Engine/Plugins/Runtime/StateTree/Source/StateTreeModule/Public/StateTreeAsyncExecutionContext.h b/Engine/Plugins/Runtime/StateTree/Source/StateTreeModule/Public/StateTreeAsyncExecutionContext.h
--- a/Engine/Plugins/Runtime/StateTree/Source/StateTreeModule/Public/StateTreeAsyncExecutionContext.h
+++ b/Engine/Plugins/Runtime/StateTree/Source/StateTreeModule/Public/StateTreeAsyncExecutionContext.h
@@ -84,6 +84,12 @@ struct TStateTreeStrongExecutionContext
 		return Storage;
 	}
 
+	/** @return the Temporary Storage. */
+	TSharedPtr<const UE::StateTree::ExecutionContext::ITemporaryStorage> GetTemporaryStorage() const
+	{
+		return TemporaryStorage;
+	}
+
 	/**
 	 * Sends event for the StateTree.
 	 * @return false if the context is not valid or the event could not be sent.
diff --git a/Engine/Plugins/Runtime/StateTree/Source/StateTreeModule/Public/StateTreeExecutionContext.h b/Engine/Plugins/Runtime/StateTree/Source/StateTreeModule/Public/StateTreeExecutionContext.h
--- a/Engine/Plugins/Runtime/StateTree/Source/StateTreeModule/Public/StateTreeExecutionContext.h
+++ b/Engine/Plugins/Runtime/StateTree/Source/StateTreeModule/Public/StateTreeExecutionContext.h
@@ -789,10 +789,15 @@ PRAGMA_ENABLE_DEPRECATION_WARNINGS
 		return UE::StateTree::InstanceData::GetDataView(InstanceDataStorage, CurrentlyProcessedSharedInstanceStorage, CurrentFrame, Handle);
 	}
 
-	//@todo deprecate in favor of a find with frameID
 	/** Looks for a frame in provided list of frames. */
+	UE_DEPRECATED(5.8, "Use UE::StateTree::ExecutionContext::FindExecutionFrame instead.")
 	static UE_API const FStateTreeExecutionFrame* FindFrame(const UStateTree* StateTree, FStateTreeStateHandle RootState, TConstArrayView<FStateTreeExecutionFrame> Frames, const FStateTreeExecutionFrame*& OutParentFrame);
 
+protected:
+	static UE_API FStateTreeExecutionFrame* FindMutableExecutionFrame(UE::StateTree::FActiveFrameID InFrameID, TArrayView<FStateTreeExecutionFrame> ActiveFrames, TArrayView<FStateTreeExecutionFrame> InTemporaryFrames);
+
+public:
+
 	/**
 	 * Forces transition to a state. It will skip all conditions.
 	 * Primarily used for replication purposes so that a client state tree stays in sync with its server counterpart.
@@ -857,6 +862,7 @@ protected:
 		//~ITemporaryStorage
 		virtual FFrameAndParent GetExecutionFrame(UE::StateTree::FActiveFrameID ID) override;
 		virtual UE::StateTree::FActiveState GetStateHandle(UE::StateTree::FActiveStateID ID) const override;
+		virtual TConstArrayView<FStateTreeExecutionFrame> GetTemporaryFrames() const override;
 	};
 
 PRAGMA_DISABLE_DEPRECATION_WARNINGS
diff --git a/Engine/Plugins/Runtime/StateTree/Source/StateTreeModule/Public/StateTreeExecutionContextTypes.h b/Engine/Plugins/Runtime/StateTree/Source/StateTreeModule/Public/StateTreeExecutionContextTypes.h
--- a/Engine/Plugins/Runtime/StateTree/Source/StateTreeModule/Public/StateTreeExecutionContextTypes.h
+++ b/Engine/Plugins/Runtime/StateTree/Source/StateTreeModule/Public/StateTreeExecutionContextTypes.h
@@ -41,8 +41,10 @@ namespace UE::StateTree::ExecutionContext
 
 		virtual FFrameAndParent GetExecutionFrame(UE::StateTree::FActiveFrameID ID) = 0;
 		virtual UE::StateTree::FActiveState GetStateHandle(UE::StateTree::FActiveStateID ID) const = 0;
+		virtual TConstArrayView<FStateTreeExecutionFrame> GetTemporaryFrames() const = 0;
 	};
 
+	UE_API const FStateTreeExecutionFrame* FindExecutionFrame(FActiveFrameID InFrameID, TConstArrayView<FStateTreeExecutionFrame> InActiveFrames, TConstArrayView<FStateTreeExecutionFrame> InTemporaryFrames);
 } // UE::StateTree::ExecutionContext
 
 #undef UE_API
diff --git a/Engine/Plugins/Runtime/StateTree/Source/StateTreeModule/Public/StateTreePropertyRef.h b/Engine/Plugins/Runtime/StateTree/Source/StateTreeModule/Public/StateTreePropertyRef.h
--- a/Engine/Plugins/Runtime/StateTree/Source/StateTreeModule/Public/StateTreePropertyRef.h
+++ b/Engine/Plugins/Runtime/StateTree/Source/StateTreeModule/Public/StateTreePropertyRef.h
@@ -6,90 +6,91 @@
 #include "StateTreeExecutionContext.h"
 #include "StateTreeIndexTypes.h"
 #include "StateTreeInstanceData.h"
-#include "StateTreePropertyRefHelpers.h"
 #include "Templates/Tuple.h"
 #include "StateTreePropertyRef.generated.h"
 
+#define UE_API STATETREEMODULE_API
+
 struct FStateTreePropertyRef;
 
 namespace UE::StateTree::PropertyRefHelpers
 {
+	namespace Private
+	{
+		struct FResolvePropertyReferenceIndirectionsResult
+		{
+			FStateTreeDataView SourceView;
+			const FStateTreePropertyAccess* PropertyAccess = nullptr;
+			const FStateTreePropertyBindings* PropertyBindings = nullptr;
+		};
+		UE_API TOptional<FResolvePropertyReferenceIndirectionsResult> ResolvePropertyReferenceIndirections(
+			const FStateTreePropertyRef& InPropertyRef,
+			FStateTreeInstanceStorage& InInstanceDataStorage,
+			const ExecutionContext::ITemporaryStorage* InTemporaryStorage,
+			const FStateTreeExecutionFrame& InExecutionFrame);
+	}
+
+	template <class T>
+	UE_DEPRECATED(5.8, "Use the version with TemporaryStorage instead of ParentFrame")
+	T* GetMutablePtrToProperty(const FStateTreePropertyRef& PropertyRef, FStateTreeInstanceStorage& InstanceDataStorage, const FStateTreeExecutionFrame& ExecutionFrame, const FStateTreeExecutionFrame* ParentExecutionFrame, const FProperty** OutSourceProperty = nullptr)
+	{
+		constexpr ExecutionContext::ITemporaryStorage* TemporaryStorage = nullptr;
+		return GetMutablePtrToProperty<T>(PropertyRef, InstanceDataStorage, TemporaryStorage, ExecutionFrame, OutSourceProperty);
+	}
+
 	/**
-	 * @param PropertyRef Property's reference to get pointer to.
-	 * @param InstanceDataStorage Instance Data Storage.
-	 * @param ExecutionFrame Execution frame owning referenced property.
-	 * @param ParentExecutionFrame Parent of execution frame owning referenced property.
+	 * @param InPropertyRef Property's reference to get pointer to.
+	 * @param InInstanceDataStorage Instance Data Storage.
+	 * @param InTemporaryStorage Temporary Data Storage.
+	 * @param InExecutionFrame Execution frame owning referenced property.
 	 * @param OutSourceProperty On success, returns referenced property.
 	 * @return Pointer to referenced property value if succeeded.
 	 */
 	template <class T>
-	static T* GetMutablePtrToProperty(const FStateTreePropertyRef& PropertyRef, FStateTreeInstanceStorage& InstanceDataStorage, const FStateTreeExecutionFrame& ExecutionFrame, const FStateTreeExecutionFrame* ParentExecutionFrame, const FProperty** OutSourceProperty = nullptr)
+	T* GetMutablePtrToProperty(const FStateTreePropertyRef& InPropertyRef, FStateTreeInstanceStorage& InInstanceDataStorage, const ExecutionContext::ITemporaryStorage* InTemporaryStorage, const FStateTreeExecutionFrame& InExecutionFrame, const FProperty** OutSourceProperty = nullptr)
 	{
-		const FStateTreePropertyBindings& PropertyBindings = ExecutionFrame.StateTree->GetPropertyBindings();
-		if (const FStateTreePropertyAccess* PropertyAccess = PropertyBindings.GetPropertyAccess(PropertyRef))
+		using namespace Private;
+		if (TOptional<FResolvePropertyReferenceIndirectionsResult> OptionalResult = ResolvePropertyReferenceIndirections(InPropertyRef, InInstanceDataStorage, InTemporaryStorage, InExecutionFrame))
 		{
-			const FStateTreeDataView SourceView = InstanceData::GetDataViewOrTemporary(InstanceDataStorage, nullptr, ExecutionFrame, PropertyAccess->SourceDataHandle);
-			
-			// The only possibility when PropertyRef references another PropertyRef is when source one is a global or subtree parameter, i.e lives in parent execution frame.
-			// If that's the case, referenced PropertyRef is obtained and we recursively take the address where it points to.
-			if (IsPropertyRef(*PropertyAccess->SourceLeafProperty))
+			const FResolvePropertyReferenceIndirectionsResult& Result = OptionalResult.GetValue();
+			check(Result.PropertyAccess && Result.PropertyBindings);
+
+			if (OutSourceProperty)
 			{
-				check(PropertyAccess->SourceDataHandle.GetSource() == EStateTreeDataSourceType::GlobalParameterData ||
-					PropertyAccess->SourceDataHandle.GetSource() == EStateTreeDataSourceType::ExternalGlobalParameterData ||
-					PropertyAccess->SourceDataHandle.GetSource() == EStateTreeDataSourceType::SubtreeParameterData);
-
-				if (ParentExecutionFrame == nullptr)
-				{
-					return nullptr;
-				}
-
-				const FStateTreePropertyRef* ReferencedPropertyRef = PropertyBindings.GetMutablePropertyPtr<FStateTreePropertyRef>(SourceView, *PropertyAccess);
-				if (ReferencedPropertyRef == nullptr)
-				{
-					return nullptr;
-				}
-
-				const FStateTreeExecutionFrame* ParentFrame = nullptr;		
-				TConstArrayView<FStateTreeExecutionFrame> ActiveFrames = InstanceDataStorage.GetExecutionState().ActiveFrames;
-				const FStateTreeExecutionFrame* Frame = FStateTreeExecutionContext::FindFrame(ParentExecutionFrame->StateTree, ParentExecutionFrame->RootState, ActiveFrames, ParentFrame);
-				
-				if (Frame == nullptr)
-				{
-					return nullptr;
-				}
-
-				return GetMutablePtrToProperty<T>(*ReferencedPropertyRef, InstanceDataStorage, *Frame, ParentFrame, OutSourceProperty);
+				*OutSourceProperty = Result.PropertyAccess->SourceLeafProperty;
 			}
-			else
-			{
-				if (OutSourceProperty)
-				{
-					*OutSourceProperty = PropertyAccess->SourceLeafProperty;
-				}
 
-				return PropertyBindings.GetMutablePropertyPtr<T>(SourceView, *PropertyAccess);
-			}
+			return Result.PropertyBindings->GetMutablePropertyPtr<T>(Result.SourceView, *Result.PropertyAccess);
 		}
 
 		return nullptr;
 	}
 
+	template <class... T>
+	UE_DEPRECATED(5.8, "Use the version with ITemporaryStorage instead of ParentFrame")
+	TTuple<T*...> GetMutablePtrTupleToProperty(const FStateTreePropertyRef& PropertyRef, FStateTreeInstanceStorage& InstanceDataStorage, const FStateTreeExecutionFrame& ExecutionFrame, const FStateTreeExecutionFrame* ParentExecutionFrame)
+	{
+		constexpr ExecutionContext::ITemporaryStorage* TemporaryStorage = nullptr;
+		return GetMutablePtrTupleToProperty<T...>(PropertyRef, InstanceDataStorage, TemporaryStorage, ExecutionFrame);
+	}
+
 	/**
-	 * @param PropertyRef Property's reference to get pointer to.
-	 * @param InstanceDataStorage Instance Data Storage
-	 * @param ExecutionFrame Execution frame owning referenced property
-	 * @param ParentExecutionFrame Parent of execution frame owning referenced property
+	 * @param InPropertyRef Property's reference to get pointer to.
+	 * @param InInstanceDataStorage Instance Data Storage
+	 * @param InTemporaryStorage Temporary Data Storage
+	 * @param InExecutionFrame Execution frame owning referenced property
 	 * @return A tuple of pointer to referenced property if succeeded.
 	 */
 	template <class... T>
-	static TTuple<T*...> GetMutablePtrTupleToProperty(const FStateTreePropertyRef& PropertyRef, FStateTreeInstanceStorage& InstanceDataStorage, const FStateTreeExecutionFrame& ExecutionFrame, const FStateTreeExecutionFrame* ParentExecutionFrame)
+	TTuple<T*...> GetMutablePtrTupleToProperty(const FStateTreePropertyRef& InPropertyRef, FStateTreeInstanceStorage& InInstanceDataStorage, const ExecutionContext::ITemporaryStorage* InTemporaryStorage, const FStateTreeExecutionFrame& InExecutionFrame)
 	{
-		const FStateTreePropertyBindings& PropertyBindings = ExecutionFrame.StateTree->GetPropertyBindings();
-		if (const FStateTreePropertyAccess* PropertyAccess = PropertyBindings.GetPropertyAccess(PropertyRef))
+		using namespace Private;
+		if (TOptional<FResolvePropertyReferenceIndirectionsResult> OptionalResult = ResolvePropertyReferenceIndirections(InPropertyRef, InInstanceDataStorage, InTemporaryStorage, InExecutionFrame))
 		{
-			// Passing empty ContextAndExternalDataViews, as PropertyRef is not allowed to point to context or external data.
-			const FStateTreeDataView SourceView = InstanceData::GetDataViewOrTemporary(InstanceDataStorage, nullptr, ExecutionFrame, PropertyAccess->SourceDataHandle);
-			return TTuple<T*...>(PropertyBindings.GetMutablePropertyPtr<T>(SourceView, *PropertyAccess)...);
+			const FResolvePropertyReferenceIndirectionsResult& Result = OptionalResult.GetValue();
+			check(Result.PropertyAccess && Result.PropertyBindings);
+
+			return TTuple<T*...>(Result.PropertyBindings->GetMutablePropertyPtr<T>(Result.SourceView, *Result.PropertyAccess)...);
 		}
 
 		return TTuple<T*...>{};
@@ -99,7 +100,7 @@ namespace UE::StateTree::PropertyRefHelpers
 /**
  * Property ref allows to get a pointer to selected property in StateTree.
  * The expected type of the reference should be set in "RefType" meta specifier.
- *
+ * PropertyRef is recommended to be used to avoid copying large struct. Otherwise, Output Binding offers a simpler workflow and designers can author the write-back inside Editor.   
  * Meta specifiers for the type:
  *  - RefType = "<type>"
  *		- Specifies a comma separated list of type of property to reference
@@ -145,7 +146,7 @@ struct FStateTreePropertyRef
 		const FStateTreeExecutionFrame* CurrentlyProcessedFrame = Context.GetCurrentlyProcessedFrame();
 		check(CurrentlyProcessedFrame);
 
-		return UE::StateTree::PropertyRefHelpers::GetMutablePtrToProperty<T>(*this, Context.GetMutableInstanceData()->GetMutableStorage(), *CurrentlyProcessedFrame, Context.GetCurrentlyProcessedParentFrame());
+		return UE::StateTree::PropertyRefHelpers::GetMutablePtrToProperty<T>(*this, Context.GetMutableInstanceData()->GetMutableStorage(), Context.GetCurrentlyProcessedTemporaryStorage().Get(), *CurrentlyProcessedFrame);
 	}
 
 	/** @return pointer to the property if possible, nullptr otherwise. */
@@ -155,7 +156,7 @@ struct FStateTreePropertyRef
 		UE::StateTree::Async::FActivePathInfo ActivePath = Context.GetActivePathInfo();
 		if (ActivePath.IsValid())
 		{
-			return UE::StateTree::PropertyRefHelpers::GetMutablePtrToProperty<T>(*this, *Context.Storage, *ActivePath.Frame, ActivePath.ParentFrame);
+			return UE::StateTree::PropertyRefHelpers::GetMutablePtrToProperty<T>(*this, *Context.Storage, Context.TemporaryStorage.Get(), *ActivePath.Frame);
 		}
 
 		return nullptr;
@@ -168,7 +169,7 @@ struct FStateTreePropertyRef
 		const FStateTreeExecutionFrame* CurrentlyProcessedFrame = Context.GetCurrentlyProcessedFrame();
 		check(CurrentlyProcessedFrame);
 
-		return UE::StateTree::PropertyRefHelpers::GetMutablePtrTupleToProperty<T...>(*this, Context.GetMutableInstanceData()->GetMutableStorage(), *CurrentlyProcessedFrame, Context.GetCurrentlyProcessedParentFrame());
+		return UE::StateTree::PropertyRefHelpers::GetMutablePtrTupleToProperty<T...>(*this, Context.GetMutableInstanceData()->GetMutableStorage(), Context.GetCurrentlyProcessedTemporaryStorage().Get(), *CurrentlyProcessedFrame, Context.GetCurrentlyProcessedParentFrame());
 	}
 
 	/** @return a tuple of pointers of the given types to the property if possible, nullptr otherwise. */
@@ -178,7 +179,7 @@ struct FStateTreePropertyRef
 		UE::StateTree::Async::FActivePathInfo ActivePath = Context.GetActivePathInfo();
 		if (ActivePath.IsValid())
 		{
-			return UE::StateTree::PropertyRefHelpers::GetMutablePtrTupleToProperty<T...>(*this, *Context.Storage, *ActivePath.Frame, ActivePath.ParentFrame);
+			return UE::StateTree::PropertyRefHelpers::GetMutablePtrTupleToProperty<T...>(*this, *Context.Storage, Context.TemporaryStorage.Get(), *ActivePath.Frame);
 		}
 
 		return {};
@@ -252,17 +253,31 @@ template <class TRef>
 struct TStateTreePropertyRef
 {
 	/** @return pointer to the property if possible, nullptr otherwise. */
-	TRef* GetMutablePtr(FStateTreeExecutionContext& Context) const
+	TRef* GetMutablePtr(const FStateTreeExecutionContext& Context) const
 	{
 		return PropertyRef.GetMutablePtr<TRef>(Context);
 	}
 
+	/** @return pointer to the property if possible, nullptr otherwise. */
+	template<bool bWithWriteAccess>
+	std::conditional_t<bWithWriteAccess, TRef*, const TRef*> GetPtrFromStrongExecutionContext(const TStateTreeStrongExecutionContext<bWithWriteAccess>& Context)
+	{
+		return PropertyRef.GetPtrFromStrongExecutionContext<TRef, bWithWriteAccess>(Context);
+	}
+
 	/** @return a tuple of pointer to the property of the type or array of type, nullptr otherwise. */
-	TTuple<TRef*, TArray<TRef>*> GetMutablePtrTuple(FStateTreeExecutionContext& Context) const
+	TTuple<TRef*, TArray<TRef>*> GetMutablePtrTuple(const FStateTreeExecutionContext& Context) const
 	{
 		return PropertyRef.GetMutablePtrTuple<TRef, TArray<TRef>>(Context);
 	}
 
+	/** @return a tuple of pointers of the given types to the property if possible, nullptr otherwise. */
+	template <bool bWithWriteAccess>
+	std::conditional_t<bWithWriteAccess, TTuple<TRef*, TArray<TRef>*>, TTuple<const TRef*, const TArray<TRef>*>> GetPtrTupleFromStrongExecutionContext(const TStateTreeStrongExecutionContext<bWithWriteAccess>& Context) const
+	{
+		return PropertyRef.GetPtrTupleFromStrongExecutionContext<TRef, TArray<TRef>>(Context);
+	}
+
 	/**
 	 * Used internally.
 	 * @return internal property ref
@@ -276,10 +291,8 @@ private:
 	FStateTreePropertyRef PropertyRef;
 };
 
-/**
- * External Handle allows to wrap-up property reference to make it accessible without having an access to StateTreeExecutionContext. Useful for capturing property reference in callbacks.
- */
-struct FStateTreePropertyRefExternalHandle
+PRAGMA_DISABLE_DEPRECATION_WARNINGS
+struct UE_DEPRECATED(5.8, "PropertyRefExternalHandle is deprecated. Use PropertyRef with ExecutionContext/StrongExecutionContext instead, which is a general async pattern for data access using Frame/StateID.") FStateTreePropertyRefExternalHandle
 {
 	FStateTreePropertyRefExternalHandle(FStateTreePropertyRef InPropertyRef, FStateTreeExecutionContext& InContext)
 		: WeakInstanceStorage(InContext.GetMutableInstanceData()->GetWeakMutableStorage())
@@ -325,11 +338,8 @@ protected:
 	FStateTreePropertyRef PropertyRef;
 };
 
-/**
- * Single type safe external handle allows to wrap-up property reference to make it accessible without having an access to StateTreeExecutionContext. Useful for capturing property reference in callbacks.
- */
 template <class TRef>
-struct TStateTreePropertyRefExternalHandle : public FStateTreePropertyRefExternalHandle
+struct UE_DEPRECATED(5.8, "PropertyRefExternalHandle is deprecated. Use PropertyRef with ExecutionContext/StrongExecutionContext instead, which is a general async pattern for data access using Frame/StateID.") TStateTreePropertyRefExternalHandle : public FStateTreePropertyRefExternalHandle
 {
 	using FStateTreePropertyRefExternalHandle::FStateTreePropertyRefExternalHandle;
 	TStateTreePropertyRefExternalHandle(TStateTreePropertyRef<TRef> InPropertyRef, FStateTreeExecutionContext& InContext)
@@ -353,6 +363,7 @@ private:
 	using FStateTreePropertyRefExternalHandle::GetMutablePtr;
 	using FStateTreePropertyRefExternalHandle::GetMutablePtrTuple;
 };
+PRAGMA_ENABLE_DEPRECATION_WARNINGS
 
 UENUM()
 enum class EStateTreePropertyRefType : uint8
@@ -415,4 +426,6 @@ private:
 	TObjectPtr<UObject> TypeObject = nullptr;
 
 	friend class FStateTreeBlueprintPropertyRefDetails;
-};
\ No newline at end of file
+};
+
+#undef UE_API
diff --git a/Engine/Plugins/Runtime/StateTree/Source/StateTreeTestSuite/Private/StateTreeBindingTest.cpp b/Engine/Plugins/Runtime/StateTree/Source/StateTreeTestSuite/Private/StateTreeBindingTest.cpp
--- a/Engine/Plugins/Runtime/StateTree/Source/StateTreeTestSuite/Private/StateTreeBindingTest.cpp
+++ b/Engine/Plugins/Runtime/StateTree/Source/StateTreeTestSuite/Private/StateTreeBindingTest.cpp
@@ -1058,6 +1058,195 @@ struct FStateTreeTest_ReferencesOnNode : FStateTreeTestBase
 };
 IMPLEMENT_STATE_TREE_INSTANT_TEST(FStateTreeTest_ReferencesOnNode, "System.StateTree.Binding.ReferencesOnNode");
 
+struct FStateTreeTest_ChainedReferences : FStateTreeTestBase
+{
+	virtual bool InstantTest() override
+	{
+		// TreeA(Global Parameter: 10)
+		//	RootA
+		//    LinkedStateB -> SubtreeB (Parameter : PropertyRef linked to Global Parameter)
+		//  SubtreeB (Parameter : PropertyRef)
+		//    LinkedStateC -> SubtreeC (Parameter : PropertyRef)
+		//  SubtreeC (Parameter : PropertyRef)
+		//    LinkedAssetStateD -> TreeD (Parameter : PropertyRef)
+		// TreeD(Global Parameter: PropertyRef)
+		//  RootE
+		//   LinkedStateF -> SubtreeF
+		//  SubtreeF
+		//    LinkedStateG -> SubtreeG
+		//  SubtreeG(ConditionG with PropertyRef, TaskG with PropertyRef)
+
+		const void* GlobalParamAddressFromCondition = nullptr;
+		const void* GlobalParamAddressFromTask = nullptr;
+
+		auto AddPropertyRefPropertyToPropertyBag = [](FInstancedPropertyBag& PropertyBag, FName VarName)
+			{
+				PropertyBag.AddProperty(VarName, EPropertyBagPropertyType::Struct, FStateTreePropertyRef::StaticStruct());
+				const FPropertyBagPropertyDesc* PropertyDesc = PropertyBag.FindPropertyDescByName(VarName);
+				check(PropertyDesc && PropertyDesc->CachedProperty);
+				const_cast<FProperty*>(PropertyDesc->CachedProperty)->SetMetaData(TEXT("Optional"), nullptr);
+				const_cast<FProperty*>(PropertyDesc->CachedProperty)->SetMetaData(TEXT("RefType"), TEXT("int32"));
+
+				return PropertyDesc->ID;
+			};
+
+		constexpr auto PropertyRefName = TEXT("PropertyRef");
+
+		FGuid StateTreeDPropertyRefParameterID;
+		UStateTree& StateTreeD = NewStateTree();
+		{
+			UStateTreeEditorData& EditorDataD = *Cast<UStateTreeEditorData>(StateTreeD.EditorData);
+			{
+				StateTreeDPropertyRefParameterID = AddPropertyRefPropertyToPropertyBag(const_cast<FInstancedPropertyBag&>(EditorDataD.GetRootParametersPropertyBag()), PropertyRefName);
+
+				auto& RootE = EditorDataD.AddSubTree(TEXT("RootE"));
+
+				auto& SubtreeF = EditorDataD.AddSubTree(TEXT("SubtreeF"), EStateTreeStateType::Subtree);
+
+				auto& SubtreeG = EditorDataD.AddSubTree(TEXT("SubtreeG"), EStateTreeStateType::Subtree);
+				{
+					constexpr auto NodePropertyRefName = TEXT("RefOnNode");
+
+					TStateTreeEditorNode<FTestCondition_PropertyRefOnNodeAndInstance>& ConditionNode = SubtreeG.AddEnterCondition<FTestCondition_PropertyRefOnNodeAndInstance>(TEXT("SubtreeGStateCondition"));
+					{
+						ConditionNode.GetNode().CustomTestConditionFunc = [&GlobalParamAddressFromCondition](FStateTreeExecutionContext& Context, const FTestCondition_PropertyRefOnNodeAndInstance& Node)
+						{
+								GlobalParamAddressFromCondition = Node.RefOnNode.GetMutablePtr<int32>(Context);
+						};
+
+						EditorDataD.AddPropertyBinding(
+							FPropertyBindingPath(EditorDataD.GetRootParametersGuid(), PropertyRefName),
+							FPropertyBindingPath(ConditionNode.GetNodeID(), NodePropertyRefName)
+						);
+					}
+
+					TStateTreeEditorNode<FTestTask_PropertyRefOnNodeAndInstance>& TaskNode = SubtreeG.AddTask<FTestTask_PropertyRefOnNodeAndInstance>(TEXT("SubtreeGStateTask"));
+					{
+						TaskNode.GetNode().CustomEnterStateFunc = [&GlobalParamAddressFromTask](FStateTreeExecutionContext& Context, const FTestTask_PropertyRefOnNodeAndInstance& Node)
+							{
+								GlobalParamAddressFromTask = Node.RefOnNode.GetMutablePtr<int32>(Context);
+							};
+
+						EditorDataD.AddPropertyBinding(
+							FPropertyBindingPath(EditorDataD.GetRootParametersGuid(), PropertyRefName),
+							FPropertyBindingPath(TaskNode.GetNodeID(), NodePropertyRefName)
+						);
+					}
+				}
+
+				auto& LinkedF = RootE.AddChildState(TEXT("LinkedF"), EStateTreeStateType::Linked);
+				LinkedF.SetLinkedState(SubtreeF.GetLinkToState());
+
+				auto& LinkedG = SubtreeF.AddChildState(TEXT("LinkedG"), EStateTreeStateType::Linked);
+				LinkedG.SetLinkedState(SubtreeG.GetLinkToState());
+			}
+		}
+
+		{
+			FStateTreeCompilerLog Log;
+			FStateTreeCompiler Compiler(Log);
+			const bool bResult = Compiler.Compile(StateTreeD);
+			AITEST_TRUE(TEXT("StateTreeD should get compiled"), bResult);
+		}
+
+		constexpr auto GlobalIntParameterName = TEXT("GlobalInt");
+		UStateTree& StateTreeA = NewStateTree();
+		{
+			UStateTreeEditorData& EditorDataA = *Cast<UStateTreeEditorData>(StateTreeA.EditorData);
+			{
+				const_cast<FInstancedPropertyBag&>(EditorDataA.GetRootParametersPropertyBag()).AddProperty(GlobalIntParameterName, EPropertyBagPropertyType::Int32);
+				const_cast<FInstancedPropertyBag&>(EditorDataA.GetRootParametersPropertyBag()).SetValueInt32(GlobalIntParameterName, 10);
+
+				auto& RootA = EditorDataA.AddSubTree(TEXT("RootA"));
+
+				FGuid SubtreeBPropertyRefParameterID;
+				auto& SubtreeB = EditorDataA.AddSubTree(TEXT("SubtreeB"), EStateTreeStateType::Subtree);
+				{
+					SubtreeBPropertyRefParameterID = AddPropertyRefPropertyToPropertyBag(SubtreeB.Parameters.Parameters, PropertyRefName);
+				}
+
+				FGuid SubtreeCPropertyRefParameterID;
+				auto& SubtreeC = EditorDataA.AddSubTree(TEXT("SubtreeC"), EStateTreeStateType::Subtree);
+				{
+					SubtreeCPropertyRefParameterID = AddPropertyRefPropertyToPropertyBag(SubtreeC.Parameters.Parameters, PropertyRefName);
+				}
+
+				auto& LinkedB = RootA.AddChildState(TEXT("LinkedB"), EStateTreeStateType::Linked);
+				{
+					LinkedB.SetLinkedState(SubtreeB.GetLinkToState());
+
+					constexpr bool bIsOverridden = true;
+					LinkedB.SetParametersPropertyOverridden(SubtreeBPropertyRefParameterID, bIsOverridden);
+
+					EditorDataA.AddPropertyBinding(
+						FPropertyBindingPath(EditorDataA.GetRootParametersGuid(), GlobalIntParameterName),
+						FPropertyBindingPath(LinkedB.Parameters.ID, PropertyRefName));
+				}
+
+				auto& LinkedC = SubtreeB.AddChildState(TEXT("LinkedC"), EStateTreeStateType::Linked);
+				{
+					LinkedC.SetLinkedState(SubtreeC.GetLinkToState());
+
+					constexpr bool bIsOverridden = true;
+					LinkedC.SetParametersPropertyOverridden(SubtreeCPropertyRefParameterID, bIsOverridden);
+
+					EditorDataA.AddPropertyBinding(
+						FPropertyBindingPath(SubtreeB.Parameters.ID, PropertyRefName),
+						FPropertyBindingPath(LinkedC.Parameters.ID, PropertyRefName));
+				}
+
+				auto& LinkedAssetD = SubtreeC.AddChildState(TEXT("LinkedAssetD"), EStateTreeStateType::LinkedAsset);
+				{
+					LinkedAssetD.SetLinkedStateAsset(&StateTreeD);
+
+					constexpr bool bIsOverridden = true;
+					LinkedAssetD.SetParametersPropertyOverridden(StateTreeDPropertyRefParameterID, bIsOverridden);
+
+					EditorDataA.AddPropertyBinding(
+						FPropertyBindingPath(SubtreeC.Parameters.ID, PropertyRefName),
+						FPropertyBindingPath(LinkedAssetD.Parameters.ID, PropertyRefName));
+				}
+			}
+		}
+
+		{
+			FStateTreeCompilerLog Log;
+			FStateTreeCompiler Compiler(Log);
+			const bool bResult = Compiler.Compile(StateTreeA);
+			AITEST_TRUE(TEXT("StateTreeA should get compiled"), bResult);
+		}
+
+
+		{
+			FStateTreeInstanceData InstanceData;
+			FTestStateTreeExecutionContext Exec(StateTreeA, StateTreeA, InstanceData);
+
+			const bool bInitSucceeded = Exec.IsValid();
+			AITEST_TRUE(TEXT("StateTree should init"), bInitSucceeded);
+
+			EStateTreeRunStatus Status = EStateTreeRunStatus::Unset;
+
+			Status = Exec.Start();
+
+			const void* GlobalIntParameterActualAddress =
+				InstanceData.GetStorage().GetGlobalParameters().GetMemory() + StateTreeA.GetDefaultParameters().FindPropertyDescByName(GlobalIntParameterName)->CachedProperty->GetOffset_ForInternal();
+
+			AITEST_EQUAL(TEXT("Start should complete with Running"), Status, EStateTreeRunStatus::Running);
+			AITEST_TRUE(TEXT("In correct states"), Exec.ExpectInActiveStates("RootA", "LinkedB", "SubtreeB", "LinkedC", "SubtreeC", "LinkedAssetD", "RootE", "LinkedF", "SubtreeF", "LinkedG", "SubtreeG"));
+
+			AITEST_TRUE(TEXT("Condition Node should correctly fetch global parameter address."), GlobalIntParameterActualAddress == GlobalParamAddressFromCondition);
+			AITEST_TRUE(TEXT("Task Node should correctly fetch global parameter address."), GlobalIntParameterActualAddress == GlobalParamAddressFromTask);
+
+			Exec.LogClear();
+
+			Exec.Stop();
+		}
+
+		return true;
+	}
+};
+IMPLEMENT_STATE_TREE_INSTANT_TEST(FStateTreeTest_ChainedReferences, "System.StateTree.Binding.ChainedReferences");
+
 struct FStateTreeTest_MutableArray : FStateTreeTestBase
 {
 	virtual bool InstantTest() override
diff --git a/Engine/Plugins/Runtime/StateTree/Source/StateTreeTestSuite/Private/StateTreeTestTypes.h b/Engine/Plugins/Runtime/StateTree/Source/StateTreeTestSuite/Private/StateTreeTestTypes.h
--- a/Engine/Plugins/Runtime/StateTree/Source/StateTreeTestSuite/Private/StateTreeTestTypes.h
+++ b/Engine/Plugins/Runtime/StateTree/Source/StateTreeTestSuite/Private/StateTreeTestTypes.h
@@ -766,7 +766,7 @@ struct FTestTask_PropertyRefOnNodeAndInstance_InstanceData
 {
 	GENERATED_BODY()
 
-	UPROPERTY(EditAnywhere, Category = Parameter, meta = (RefType = "int32"))
+	UPROPERTY(EditAnywhere, Category = Parameter, meta = (RefType = "int32", Optional))
 	FStateTreePropertyRef RefOnInstance;
 };
 
@@ -795,12 +795,43 @@ struct FTestTask_PropertyRefOnNodeAndInstance : public FStateTreeTaskBase
 		return EStateTreeRunStatus::Running;
 	}
 
-	UPROPERTY(EditAnywhere, Category = Parameter, meta = (RefType = "int32"))
+	UPROPERTY(EditAnywhere, Category = Parameter, meta = (RefType = "int32", Optional))
 	FStateTreePropertyRef RefOnNode;
 
 	TFunction<void(FStateTreeExecutionContext&, const FTestTask_PropertyRefOnNodeAndInstance&)> CustomEnterStateFunc;
 };
 
+USTRUCT()
+struct FTestCondition_PropertyRefOnNodeAndInstance : public FStateTreeConditionBase
+{
+	GENERATED_BODY()
+
+	using FInstanceDataType = FTestTask_PropertyRefOnNodeAndInstance_InstanceData;
+
+	FTestCondition_PropertyRefOnNodeAndInstance() = default;
+	FTestCondition_PropertyRefOnNodeAndInstance(FName InName)
+	{
+		Name = InName;
+	}
+
+	virtual const UStruct* GetInstanceDataType() const override { return FInstanceDataType::StaticStruct(); }
+
+	virtual bool TestCondition(FStateTreeExecutionContext& Context) const override
+	{
+		if (CustomTestConditionFunc)
+		{
+			CustomTestConditionFunc(Context, *this);
+		}
+
+		return true;
+	}
+
+	UPROPERTY(EditAnywhere, Category = Parameter, meta = (RefType = "int32", Optional))
+	FStateTreePropertyRef RefOnNode;
+
+	TFunction<void(FStateTreeExecutionContext&, const FTestCondition_PropertyRefOnNodeAndInstance&)> CustomTestConditionFunc;
+};
+
 USTRUCT()
 struct FStateTreeTestConditionInstanceData
 {
