Fix state centric crash missing schema

#tests editor
#Synced-CL 48414911

diff --git a/Engine/Plugins/Runtime/StateTree/Source/StateTreeEditorModule/Private/StateCentricView/View/SMainStateTransitionExtensionNode.cpp b/Engine/Plugins/Runtime/StateTree/Source/StateTreeEditorModule/Private/StateCentricView/View/SMainStateTransitionExtensionNode.cpp
--- a/Engine/Plugins/Runtime/StateTree/Source/StateTreeEditorModule/Private/StateCentricView/View/SMainStateTransitionExtensionNode.cpp
+++ b/Engine/Plugins/Runtime/StateTree/Source/StateTreeEditorModule/Private/StateCentricView/View/SMainStateTransitionExtensionNode.cpp
@@ -237,7 +237,13 @@ TSharedRef<SWidget> SMainStateTransitionExtensionNode::Construct_GenerateNodeBod
 {
 	if (ViewState.IsValid())
 	{
-		TSubclassOf<UStateTreeSchema> SchemaClass = ViewModel->GetStateTree()->GetSchema()->GetClass();
+		TSubclassOf<UStateTreeSchema> SchemaClass = nullptr;
+		
+		// If compilation never passed. Then our outer statetree can be missing a schema.
+		if (const UStateTreeSchema* Schema = ViewModel->GetStateTree()->GetSchema())
+		{
+			SchemaClass = Schema->GetClass();
+		}
 
 		if (UStateCentricViewSettings::GetSchemaSettingsDefaultFallback(SchemaClass).bJustifyCentralViewTransitions)
 		{
