State centric view parent expansion controls
- Add a likely-temp toolbar widget for state centric for parent expansion control.
- Maintain Parent state expansion
- Add buttons for auto collapse / expand / preview / clear stored
- Add support for num parents shown control
- Also add initial transition responsiveness

#tests editor
#Synced-CL 48414911

diff --git a/Engine/Plugins/Runtime/StateTree/Source/StateTreeEditorModule/Private/SStateTreeView.cpp b/Engine/Plugins/Runtime/StateTree/Source/StateTreeEditorModule/Private/SStateTreeView.cpp
--- a/Engine/Plugins/Runtime/StateTree/Source/StateTreeEditorModule/Private/SStateTreeView.cpp
+++ b/Engine/Plugins/Runtime/StateTree/Source/StateTreeEditorModule/Private/SStateTreeView.cpp
@@ -25,6 +25,7 @@
 // @TODO: Dependency onto state centric to enable view toggling. 
 // Long term plan is to have this in a separate tab.
 #include "StateCentricView/StateCentricViewSettings.h"
+#include "StateCentricView/View/SStateCentricToolbar.h"
 #include "StateCentricView/View/SStateCentricViewMainStateNode.h"
 
 #define LOCTEXT_NAMESPACE "StateTreeEditor"
@@ -197,16 +198,27 @@ void SStateTreeView::Construct(const FArguments& InArgs, TSharedRef<FStateTreeVi
 			+SSplitter::Slot()
 			.Value(0.85f)
 			[
-				SNew(SScrollBox)
-				.Orientation(Orient_Vertical)
-				+ SScrollBox::Slot()
-				.FillSize(1.0f)
+				SNew(SSplitter)
+				.Orientation(EOrientation::Orient_Vertical)
+				+ SSplitter::Slot()
+				.SizeRule(SSplitter::ESizeRule::SizeToContent)
+				[
+					SNew(UE::StateTree::Editor::StateCentricView::SStateCentricToolbar, InStateTreeViewModel)
+				]
+				+ SSplitter::Slot()
+				.Value(1.0f)
 				[
-					SNew(SBorder)
-					.BorderImage(FAppStyle::Get().GetBrush("Brushes.Recessed"))
-					.Padding(FMargin(2.0, 4.0))
+					SNew(SScrollBox)
+					.Orientation(Orient_Vertical)
+					+ SScrollBox::Slot()
+					.FillSize(1.0f)
 					[
-						SNew(UE::StateTree::Editor::StateCentricView::SStateCentricViewMainStateNode, InStateTreeViewModel, StateCentricViewState)
+						SNew(SBorder)
+						.BorderImage(FAppStyle::Get().GetBrush("Brushes.Recessed"))
+						.Padding(FMargin(2.0, 4.0))
+						[
+							SNew(UE::StateTree::Editor::StateCentricView::SStateCentricViewMainStateNode, InStateTreeViewModel, StateCentricViewState)
+						]
 					]
 				]
 			]
diff --git a/Engine/Plugins/Runtime/StateTree/Source/StateTreeEditorModule/Private/StateCentricView/StateCentricViewPerUserSettings.cpp b/Engine/Plugins/Runtime/StateTree/Source/StateTreeEditorModule/Private/StateCentricView/StateCentricViewPerUserSettings.cpp
--- a/Engine/Plugins/Runtime/StateTree/Source/StateTreeEditorModule/Private/StateCentricView/StateCentricViewPerUserSettings.cpp
+++ b/Engine/Plugins/Runtime/StateTree/Source/StateTreeEditorModule/Private/StateCentricView/StateCentricViewPerUserSettings.cpp
@@ -1,6 +1,8 @@
 // Copyright Epic Games, Inc. All Rights Reserved.
 
 #include "StateCentricView/StateCentricViewPerUserSettings.h"
+#include "StateTreeEditorData.h"
+#include "StateTreeViewModel.h"
 
 #include UE_INLINE_GENERATED_CPP_BY_NAME(StateCentricViewPerUserSettings)
 
@@ -15,6 +17,23 @@ UStateCentricViewPerUserSettings::UStateCentricViewPerUserSettings()
 	PerSchemaSettings.Emplace(UStateTreeSchema::StaticClass(), DefaultPerUserSettings);
 }
 
+const FStateCentricViewPerUserSetting& UStateCentricViewPerUserSettings::GetSchemaSettingsDefaultFallback(const FStateTreeViewModel* ViewModel)
+{
+	if (ViewModel)
+	{
+		if (const UStateTree* StateTree = ViewModel->GetStateTree())
+		{
+			if (const UStateTreeSchema* Schema = StateTree->GetSchema())
+			{
+				const TSubclassOf<UStateTreeSchema> SchemaClass = Schema->GetClass();
+				return UStateCentricViewPerUserSettings::GetSchemaSettingsDefaultFallback(SchemaClass);
+			}
+		}
+	}
+
+	return UStateCentricViewPerUserSettings::Get().PerSchemaSettings[UStateTreeSchema::StaticClass()];
+}
+
 #if WITH_EDITOR
 
 FText UStateCentricViewPerUserSettings::GetSectionText() const
diff --git a/Engine/Plugins/Runtime/StateTree/Source/StateTreeEditorModule/Private/StateCentricView/StateCentricViewPerUserSettings.h b/Engine/Plugins/Runtime/StateTree/Source/StateTreeEditorModule/Private/StateCentricView/StateCentricViewPerUserSettings.h
--- a/Engine/Plugins/Runtime/StateTree/Source/StateTreeEditorModule/Private/StateCentricView/StateCentricViewPerUserSettings.h
+++ b/Engine/Plugins/Runtime/StateTree/Source/StateTreeEditorModule/Private/StateCentricView/StateCentricViewPerUserSettings.h
@@ -12,6 +12,7 @@
 #define UE_API STATETREEEDITORMODULE_API
 
 DECLARE_MULTICAST_DELEGATE_OneParam(FOnStateCentricViewLOD_Changed, TOptional<EExtendableNodeLOD> /*InLOD*/);
+DECLARE_MULTICAST_DELEGATE_OneParam(FOnStateCentricViewIntValue_Changed, TOptional<uint32> /*InValue*/);
 
 /**
  * Experimental. Settings per user. This exists locally per user per schema. But also per state tree.
@@ -30,7 +31,7 @@ public:
 private:
 
 	/** Amount of info to show for every parent state in the hierachy. */
-	UPROPERTY(EditAnywhere, Category="Layout", Config, meta = (ValidEnumValues = "Minimal, Medium, Full"))
+	UPROPERTY(EditAnywhere, Category="Layout", Config, meta = (ValidEnumValues = "Collapsed, Minimal, Full"))
 	TOptional<EExtendableNodeLOD> ParentHireachyLOD = {};
 
 	/** Amount of info to show for incoming transitions to this state. */
@@ -47,6 +48,11 @@ private:
 
 public:
 
+	UE_API TOptional<EExtendableNodeLOD> GetParentHireachyLOD() const
+	{
+		return ParentHireachyLOD;
+	}
+
 	UE_API TOptional<EExtendableNodeLOD> GetInTransitionLOD() const
 	{
 		return InTransitionViewLOD;
@@ -57,25 +63,66 @@ public:
 		return OutTransitionLOD;
 	}
 
+	UE_API TOptional<uint32> GetNumParentsShown() const
+	{
+		return NumParentsShown;
+	}
+
+	UE_API void SetParentHireachyLOD(TOptional<EExtendableNodeLOD> InLOD)
+	{
+		ensure(!InLOD
+			|| InLOD.GetValue() == EExtendableNodeLOD::Collapsed
+			|| InLOD.GetValue() == EExtendableNodeLOD::Minimal
+			|| InLOD.GetValue() == EExtendableNodeLOD::Full);
+
+		ParentHireachyLOD = InLOD;
+		OnParentHireachyLOD_Changed.Broadcast(ParentHireachyLOD);
+	}
+
 	UE_API void SetInTransitionLOD(TOptional<EExtendableNodeLOD> InLOD)
 	{
+		ensure(!InLOD
+			|| InLOD.GetValue() == EExtendableNodeLOD::Hidden
+			|| InLOD.GetValue() == EExtendableNodeLOD::Minimal
+			|| InLOD.GetValue() == EExtendableNodeLOD::Full);
+
 		InTransitionViewLOD = InLOD;
 		OnInTransitionViewLOD_Changed.Broadcast(InTransitionViewLOD);
 	}
 
 	UE_API void SetOutTransitionLOD(TOptional<EExtendableNodeLOD> InLOD)
 	{
+		ensure(!InLOD
+			|| InLOD.GetValue() == EExtendableNodeLOD::Hidden
+			|| InLOD.GetValue() == EExtendableNodeLOD::Minimal
+			|| InLOD.GetValue() == EExtendableNodeLOD::Full);
+
 		OutTransitionLOD = InLOD;
 		OnOutTransitionViewLOD_Changed.Broadcast(OutTransitionLOD);
 	}
 
+	UE_API void SetNumParentsShown(TOptional<uint32> InNumParentShown)
+	{
+		if (NumParentsShown != InNumParentShown)
+		{
+			NumParentsShown = InNumParentShown;
+			OnNumParentsShownChanged.Broadcast(NumParentsShown);
+		}
+	}
+
 public:
 
+	/** Callback whenever default parent hierarchy LOD changed. */
+	FOnStateCentricViewLOD_Changed OnParentHireachyLOD_Changed;
+
 	/** Callback whenever the LOD for 'In' transition view has changed. */
 	FOnStateCentricViewLOD_Changed OnInTransitionViewLOD_Changed;
 
 	/** Callback whenever the LOD for 'Out' transition view has changed. */
 	FOnStateCentricViewLOD_Changed OnOutTransitionViewLOD_Changed;
+
+	/** Callback whenever num parents shown has changed. */
+	FOnStateCentricViewIntValue_Changed OnNumParentsShownChanged;
 };
 
 
@@ -102,6 +149,18 @@ public:
 		return *GetMutableDefault<UStateCentricViewPerUserSettings>();
 	}
 
+	UE_API static const FStateCentricViewPerUserSetting& GetSchemaSettingsDefaultFallback(const FStateTreeViewModel* ViewModel);
+
+	UE_API static const FStateCentricViewPerUserSetting& GetSchemaSettingsDefaultFallback(const TSubclassOf<UStateTreeSchema> InSchema)
+	{
+		if (const FStateCentricViewPerUserSetting* SchemaSettings = UStateCentricViewPerUserSettings::Get().PerSchemaSettings.Find(InSchema))
+		{
+			return *SchemaSettings;
+		}
+
+		return UStateCentricViewPerUserSettings::Get().PerSchemaSettings[UStateTreeSchema::StaticClass()];
+	}
+
 public:
 
 	/** Per schema settings. Defined per schema since we have no need for an editor schema to be defined. */
diff --git a/Engine/Plugins/Runtime/StateTree/Source/StateTreeEditorModule/Private/StateCentricView/View/SMainStateTransitionExtensionNode.cpp b/Engine/Plugins/Runtime/StateTree/Source/StateTreeEditorModule/Private/StateCentricView/View/SMainStateTransitionExtensionNode.cpp
--- a/Engine/Plugins/Runtime/StateTree/Source/StateTreeEditorModule/Private/StateCentricView/View/SMainStateTransitionExtensionNode.cpp
+++ b/Engine/Plugins/Runtime/StateTree/Source/StateTreeEditorModule/Private/StateCentricView/View/SMainStateTransitionExtensionNode.cpp
@@ -238,9 +238,8 @@ TSharedRef<SWidget> SMainStateTransitionExtensionNode::Construct_GenerateNodeBod
 	if (ViewState.IsValid())
 	{
 		TSubclassOf<UStateTreeSchema> SchemaClass = nullptr;
-		
-		// If compilation never passed. Then our outer statetree can be missing a schema.
-		if (const UStateTreeSchema* Schema = ViewModel->GetStateTree()->GetSchema())
+
+		if (const UStateTreeSchema* Schema = ViewModel->GetStateTreeSchema())
 		{
 			SchemaClass = Schema->GetClass();
 		}
@@ -300,6 +299,14 @@ FReply SMainStateTransitionExtensionNode::HandleOnCollapseClicked()
 
 FReply SMainStateTransitionExtensionNode::HandleOnAddTransitionClicked()
 {
+	if (ViewState.IsValid())
+	{
+		FStateTreeTransition DefaultTransition = {};
+		ViewState->AddTransition(DefaultTransition.Trigger, EStateTreeTransitionType::None);
+
+		// @TODO: Temp local LOD regen instead of listening to delegate
+		RegenerateWidgetForLOD(CurrentLOD);
+	}
 	return FReply::Handled();
 }
 
diff --git a/Engine/Plugins/Runtime/StateTree/Source/StateTreeEditorModule/Private/StateCentricView/View/SParentStateExtensionNode.cpp b/Engine/Plugins/Runtime/StateTree/Source/StateTreeEditorModule/Private/StateCentricView/View/SParentStateExtensionNode.cpp
--- a/Engine/Plugins/Runtime/StateTree/Source/StateTreeEditorModule/Private/StateCentricView/View/SParentStateExtensionNode.cpp
+++ b/Engine/Plugins/Runtime/StateTree/Source/StateTreeEditorModule/Private/StateCentricView/View/SParentStateExtensionNode.cpp
@@ -292,20 +292,32 @@ TSharedRef<SWidget> SParentStateExtensionNode::Construct_GenerateNodeBody(const
 
 FReply SParentStateExtensionNode::HandleOnChangeLOD_Clicked()
 {
-	if (CurrentLOD == EExtendableNodeLOD::Collapsed)
-	{
-		CurrentLOD = EExtendableNodeLOD::Minimal;
-	}
-	else if (CurrentLOD == EExtendableNodeLOD::Minimal)
-	{
-		CurrentLOD = EExtendableNodeLOD::Full;
-	}
-	else if (ensure(CurrentLOD == EExtendableNodeLOD::Full))
+	if (ViewState.IsValid())
 	{
-		CurrentLOD = EExtendableNodeLOD::Collapsed;
-	}
+		// @TODO: Remove splitter references here once those become a feature toggle on cardinal extensions
+		UStateTreeEditorData* EditorData = const_cast<UStateTreeEditorData*>(ViewModel->GetStateTreeEditorData());
+		UStateCentricViewEditorDataExtension* StateCentricEditorData = EditorData->GetOrCreateExtension<UStateCentricViewEditorDataExtension>();
+
+		if (CurrentLOD == EExtendableNodeLOD::Collapsed)
+		{
+			CurrentLOD = EExtendableNodeLOD::Minimal;
+		}
+		else if (CurrentLOD == EExtendableNodeLOD::Minimal)
+		{
+			CurrentLOD = EExtendableNodeLOD::Full;
+		}
+		else if (ensure(CurrentLOD == EExtendableNodeLOD::Full))
+		{
+			CurrentLOD = EExtendableNodeLOD::Collapsed;
+		}
 
-	RegenerateWidgetForLOD(CurrentLOD);
+		if (StateCentricEditorData->bShouldStoreParentExpansion)
+		{
+			StateCentricEditorData->LastUserSetParentStateLOD.FindOrAdd(ViewState->ID) = CurrentLOD;
+		}
+
+		RegenerateWidgetForLOD(CurrentLOD);
+	}
 
 	return FReply::Handled();
 }
diff --git a/Engine/Plugins/Runtime/StateTree/Source/StateTreeEditorModule/Private/StateCentricView/View/SParentStateTransitionExtensionNode.cpp b/Engine/Plugins/Runtime/StateTree/Source/StateTreeEditorModule/Private/StateCentricView/View/SParentStateTransitionExtensionNode.cpp
--- a/Engine/Plugins/Runtime/StateTree/Source/StateTreeEditorModule/Private/StateCentricView/View/SParentStateTransitionExtensionNode.cpp
+++ b/Engine/Plugins/Runtime/StateTree/Source/StateTreeEditorModule/Private/StateCentricView/View/SParentStateTransitionExtensionNode.cpp
@@ -188,6 +188,14 @@ TSharedRef<SWidget> SParentStateTransitionExtensionNode::Construct_GenerateNodeB
 
 FReply SParentStateTransitionExtensionNode::HandleOnAddTransitionClicked()
 {
+	if (ViewState.IsValid())
+	{
+		FStateTreeTransition DefaultTransition = {};
+		ViewState->AddTransition(DefaultTransition.Trigger, EStateTreeTransitionType::None);
+
+		// @TODO: Temp local LOD regen instead of listening to delegate
+		RegenerateWidgetForLOD(CurrentLOD);
+	}
 	return FReply::Handled();
 }
 
diff --git a/Engine/Plugins/Runtime/StateTree/Source/StateTreeEditorModule/Private/StateCentricView/View/SStateCentricToolbar.cpp b/Engine/Plugins/Runtime/StateTree/Source/StateTreeEditorModule/Private/StateCentricView/View/SStateCentricToolbar.cpp
new file mode 100644
--- /dev/null
+++ b/Engine/Plugins/Runtime/StateTree/Source/StateTreeEditorModule/Private/StateCentricView/View/SStateCentricToolbar.cpp
@@ -0,0 +1,174 @@
+// Copyright Epic Games, Inc. All Rights Reserved.
+
+#include "StateCentricView/View/SStateCentricToolbar.h"
+
+#include "StateCentricView/StateCentricViewSettings.h"
+#include "StateCentricView/ViewModel/StateCentricViewEditorDataExtension.h"
+#include "StateCentricView/View/SStateTreeExtendableNode.h"
+#include "StateCentricView/View/StateCentricViewUtils.h"
+
+#include "Framework/Application/SlateApplication.h"
+#include "Framework/MultiBox/MultiBoxBuilder.h"
+#include "HAL/IConsoleManager.h"
+#include "SlateGlobals.h"
+#include "Styling/CoreStyle.h"
+#include "Styling/ToolBarStyle.h"
+#include "Widgets/Input/SCheckBox.h"
+#include "Widgets/Input/SSpinBox.h"
+#include "Widgets/Layout/SBorder.h"
+#include "Widgets/Layout/SBox.h"
+#include "Widgets/Text/STextBlock.h"
+#include "Widgets/SBoxPanel.h"
+
+#define LOCTEXT_NAMESPACE "SStateCentricToolbar"
+
+namespace UE::StateTree::Editor::StateCentricView
+{
+
+void SStateCentricToolbar::Construct( const FArguments& InArgs, TSharedRef<FStateTreeViewModel> InViewModel)
+{
+	ViewModel = InViewModel;
+
+	FSlimHorizontalToolBarBuilder ToolbarBuilder(TSharedPtr<const FUICommandList>(), FMultiBoxCustomization::None);
+
+	CachedNumParentsShown = FStateCentricViewUtils::GetNumParentsToShow(ViewModel.ToSharedRef());
+
+	TWeakPtr<FStateTreeViewModel> WeakViewModel = ViewModel;
+
+	auto HandleOnCollapseParentsClicked = [WeakViewModel]()
+	{
+		if (TSharedPtr<FStateTreeViewModel> ViewModelPinned = WeakViewModel.Pin())
+		{
+			UStateTreeEditorData* EditorData = const_cast<UStateTreeEditorData*>(ViewModelPinned->GetStateTreeEditorData());
+			UStateCentricViewEditorDataExtension* StateCentricEditorData = EditorData->GetOrCreateExtension<UStateCentricViewEditorDataExtension>();
+
+			StateCentricEditorData->LastUserSetParentStateLOD.Reset();
+			StateCentricEditorData->PerStateTreeUserSettings.SetParentHireachyLOD(EExtendableNodeLOD::Collapsed);
+		}
+	};
+	
+	auto HandleOnPreviewParentsClicked = [WeakViewModel]()
+	{
+		if (TSharedPtr<FStateTreeViewModel> ViewModelPinned = WeakViewModel.Pin())
+		{
+			UStateTreeEditorData* EditorData = const_cast<UStateTreeEditorData*>(ViewModelPinned->GetStateTreeEditorData());
+			UStateCentricViewEditorDataExtension* StateCentricEditorData = EditorData->GetOrCreateExtension<UStateCentricViewEditorDataExtension>();
+
+			StateCentricEditorData->LastUserSetParentStateLOD.Reset();
+			StateCentricEditorData->PerStateTreeUserSettings.SetParentHireachyLOD(EExtendableNodeLOD::Minimal);
+		}
+	};
+	
+	auto HandleOnExpandParentsClicked = [WeakViewModel]()
+	{
+		if (TSharedPtr<FStateTreeViewModel> ViewModelPinned = WeakViewModel.Pin())
+		{
+			UStateTreeEditorData* EditorData = const_cast<UStateTreeEditorData*>(ViewModelPinned->GetStateTreeEditorData());
+			UStateCentricViewEditorDataExtension* StateCentricEditorData = EditorData->GetOrCreateExtension<UStateCentricViewEditorDataExtension>();
+
+			StateCentricEditorData->LastUserSetParentStateLOD.Reset();
+			StateCentricEditorData->PerStateTreeUserSettings.SetParentHireachyLOD(EExtendableNodeLOD::Full);
+		}
+	};
+	
+	auto HandleOnClearPerStateExpansions = [WeakViewModel]()
+	{
+		if (TSharedPtr<FStateTreeViewModel> ViewModelPinned = WeakViewModel.Pin())
+		{
+			UStateTreeEditorData* EditorData = const_cast<UStateTreeEditorData*>(ViewModelPinned->GetStateTreeEditorData());
+			UStateCentricViewEditorDataExtension* StateCentricEditorData = EditorData->GetOrCreateExtension<UStateCentricViewEditorDataExtension>();
+
+			StateCentricEditorData->LastUserSetParentStateLOD.Reset();
+			StateCentricEditorData->PerStateTreeUserSettings.SetParentHireachyLOD(StateCentricEditorData->PerStateTreeUserSettings.GetParentHireachyLOD());
+		}
+	};
+
+	// @TODO: Toolbar builder doesn't work well if we don't have dedicated icons. Which we don't.
+
+	ToolbarBuilder.AddToolBarButton(
+		FUIAction(FExecuteAction::CreateLambda(HandleOnCollapseParentsClicked)),
+		NAME_None,
+		LOCTEXT("CollapseParents", "Collapse Parents"),
+		LOCTEXT("CollapseParentsTip", "Clears all saves per-state expansions, then collapses parents."),
+		FSlateIcon());
+
+	ToolbarBuilder.AddToolBarButton(
+		FUIAction(FExecuteAction::CreateLambda(HandleOnPreviewParentsClicked)),
+		NAME_None,
+		LOCTEXT("PreviewParents", "Preview Parents"),
+		LOCTEXT("PreviewParentsTip", "Clears all per-state expansions, then previews parents."),
+		FSlateIcon());
+
+	ToolbarBuilder.AddToolBarButton(
+		FUIAction(FExecuteAction::CreateLambda(HandleOnExpandParentsClicked)),
+		NAME_None,
+		LOCTEXT("ExpandParents", "Expand Parents"),
+		LOCTEXT("ExpandParentsTip", "Clears all per-state expansions, then expands parents."),
+		FSlateIcon());
+
+	ToolbarBuilder.AddToolBarButton(
+		FUIAction(FExecuteAction::CreateLambda(HandleOnClearPerStateExpansions)),
+		NAME_None,
+		LOCTEXT("ClearExpansions", "Clear Per-State Expansions"),
+		LOCTEXT("ClearExpansionsTip", "Clears all per-state expansions."),
+		FSlateIcon());
+
+	ChildSlot
+	[
+		SNew(SBorder)
+		.BorderImage(FCoreStyle::Get().GetBrush("ToolPanel.GroupBorder"))
+		[
+			SNew(SHorizontalBox)
+			+ SHorizontalBox::Slot()
+			.AutoWidth()
+			.VAlign(VAlign_Center)
+			.Padding(12.f, -2.0f, 4.f, -2.0f)
+			[
+				SNew(STextBlock)
+				.Text(LOCTEXT("NumParentsShown", "Num Parents Shown: "))
+			]
+			+ SHorizontalBox::Slot()
+			.AutoWidth()
+			[
+				SNew(SBox)
+				.MinDesiredWidth(100.f)
+				.MaxDesiredWidth(250.f)
+				[
+					SNew(SSpinBox<uint32>)
+					.Value(this, &SStateCentricToolbar::HandleNumParentsSliderValue)
+					.MinValue(0)
+					.MaxValue(10)
+					.Delta(1)
+					.OnValueChanged(this, &SStateCentricToolbar::HandleNumParentsSliderChanged)
+				]
+			]
+			+ SHorizontalBox::Slot()
+			.FillWidth(1)
+				.Padding(5.f, -2.0f, 5.f, -2.0f)
+			[
+				ToolbarBuilder.MakeWidget()
+			]
+		]
+	];
+}
+
+void SStateCentricToolbar::HandleNumParentsSliderChanged(uint32 NewValue)
+{
+	if (ViewModel)
+	{
+		UStateTreeEditorData* EditorData = const_cast<UStateTreeEditorData*>(ViewModel->GetStateTreeEditorData());
+		UStateCentricViewEditorDataExtension* StateCentricEditorData = EditorData->GetOrCreateExtension<UStateCentricViewEditorDataExtension>();
+
+		StateCentricEditorData->PerStateTreeUserSettings.SetNumParentsShown(NewValue);
+		CachedNumParentsShown = NewValue;
+	}
+}
+
+uint32 SStateCentricToolbar::HandleNumParentsSliderValue() const
+{
+	return CachedNumParentsShown;
+}
+
+} // namespace UE::StateTree::Editor::StateCentricView
+
+#undef LOCTEXT_NAMESPACE
diff --git a/Engine/Plugins/Runtime/StateTree/Source/StateTreeEditorModule/Private/StateCentricView/View/SStateCentricToolbar.h b/Engine/Plugins/Runtime/StateTree/Source/StateTreeEditorModule/Private/StateCentricView/View/SStateCentricToolbar.h
new file mode 100644
--- /dev/null
+++ b/Engine/Plugins/Runtime/StateTree/Source/StateTreeEditorModule/Private/StateCentricView/View/SStateCentricToolbar.h
@@ -0,0 +1,46 @@
+// Copyright Epic Games, Inc. All Rights Reserved.
+
+#pragma once
+
+#include "CoreMinimal.h"
+#include "SlateFwd.h"
+#include "Styling/SlateTypes.h"
+#include "Widgets/DeclarativeSyntaxSupport.h"
+#include "Widgets/SCompoundWidget.h"
+
+class FStateTreeViewModel;
+
+namespace UE::StateTree::Editor::StateCentricView
+{
+
+/**
+ * Experimental. Toolbar to set options for state centric 
+ * 
+ * Current plan is placeholder to test out workflows. 
+ * Long term state centric will need to inject into different toolbars.
+ */
+class SStateCentricToolbar : public SCompoundWidget
+{
+public:
+	SLATE_USER_ARGS(SStateCentricToolbar) { }
+	SLATE_END_ARGS()
+
+	virtual void Construct(const FArguments& InArgs, TSharedRef<FStateTreeViewModel> InViewModel);
+
+private:
+
+	/** ViewModel for this ST */
+	TSharedPtr<FStateTreeViewModel> ViewModel = nullptr;
+
+	/** Cached number of parents set to be shown for this state tree */
+	uint32 CachedNumParentsShown = 0;
+
+private:
+
+	void HandleNumParentsSliderChanged(uint32 NewValue);
+	uint32 HandleNumParentsSliderValue() const;
+};
+
+
+} // namespace UE::StateTree::Editor::StateCentricView
+
diff --git a/Engine/Plugins/Runtime/StateTree/Source/StateTreeEditorModule/Private/StateCentricView/View/SStateCentricViewMainStateNode.cpp b/Engine/Plugins/Runtime/StateTree/Source/StateTreeEditorModule/Private/StateCentricView/View/SStateCentricViewMainStateNode.cpp
--- a/Engine/Plugins/Runtime/StateTree/Source/StateTreeEditorModule/Private/StateCentricView/View/SStateCentricViewMainStateNode.cpp
+++ b/Engine/Plugins/Runtime/StateTree/Source/StateTreeEditorModule/Private/StateCentricView/View/SStateCentricViewMainStateNode.cpp
@@ -33,6 +33,16 @@ SStateCentricViewMainStateNode::~SStateCentricViewMainStateNode()
 	if (ViewModel)
 	{
 		ViewModel->GetOnSelectionChanged().RemoveAll(this);
+
+		if (UStateTreeEditorData* EditorData = const_cast<UStateTreeEditorData*>(ViewModel->GetStateTreeEditorData()))
+		{
+			if (UStateCentricViewEditorDataExtension* StateCentricEditorData = EditorData->GetExtension<UStateCentricViewEditorDataExtension>())
+			{
+				StateCentricEditorData->PerStateTreeUserSettings.OnNumParentsShownChanged.RemoveAll(this);
+				StateCentricEditorData->PerStateTreeUserSettings.OnParentHireachyLOD_Changed.RemoveAll(this);
+			}
+		}
+
 	}
 }
 
@@ -43,10 +53,21 @@ FName SStateCentricViewMainStateNode::GetNodeName() const
 
 void SStateCentricViewMainStateNode::Construct_InitializeNodePersistentPreLOD()
 {
+	UStateTreeEditorData* EditorData = const_cast<UStateTreeEditorData*>(ViewModel->GetStateTreeEditorData());
+	UStateCentricViewEditorDataExtension* StateCentricEditorData = EditorData->GetOrCreateExtension<UStateCentricViewEditorDataExtension>();
+
 	// Remove any prior bindings in case of issues with lifetime
-	ViewModel->GetOnSelectionChanged().RemoveAll(this);
+	{
+		ViewModel->GetOnSelectionChanged().RemoveAll(this);
+
+		StateCentricEditorData->PerStateTreeUserSettings.OnNumParentsShownChanged.RemoveAll(this);
+		StateCentricEditorData->PerStateTreeUserSettings.OnParentHireachyLOD_Changed.RemoveAll(this);
+	}
 
 	ViewModel->GetOnSelectionChanged().AddSP(this, &SStateCentricViewMainStateNode::HandleOnSelectionChanged);
+
+	StateCentricEditorData->PerStateTreeUserSettings.OnNumParentsShownChanged.AddSP(this, &SStateCentricViewMainStateNode::HandleOnNumParentsShownChanged);
+	StateCentricEditorData->PerStateTreeUserSettings.OnParentHireachyLOD_Changed.AddSP(this, &SStateCentricViewMainStateNode::HandleOnParentHireachyLOD_Changed);
 }
 
 void SStateCentricViewMainStateNode::Construct_GenerateBuiltInExtensions(const EExtendableNodeLOD InLOD)
@@ -113,6 +134,7 @@ void SStateCentricViewMainStateNode::Construct_GenerateBuiltInExtensions(const E
 		// Parent state extensions at top
 		{
 			uint32 NumParentStateShown = 0;
+			uint32 NumParentsToShow = FStateCentricViewUtils::GetNumParentsToShow(ViewModel.ToSharedRef());
 
 			FScopedMainNodeLayoutInfo ScopedMainLayoutInfo = SetScopedMainNodeLayoutInfo()
 				.SetHAlign(HAlign_Fill)
@@ -122,10 +144,15 @@ void SStateCentricViewMainStateNode::Construct_GenerateBuiltInExtensions(const E
 
 			for (UStateTreeState* ParentState : ParentStates)
 			{
-				// @TODO: Read project settings and cull after N shown parents
+				if (NumParentsToShow == NumParentStateShown)
+				{
+					break;
+				}
+
+				TOptional<EExtendableNodeLOD> DesiredParentLOD = FStateCentricViewUtils::GetDesiredParentLOD(ViewModel.ToSharedRef(), ParentState);
 
 				TSharedRef<SWidget> ParentStateTopExtension = SNew(SParentStateExtensionNode, ViewModel.ToSharedRef(), ParentState)
-					.Config({ .InitialLOD = EExtendableNodeLOD::Minimal });
+					.Config({ .InitialLOD = DesiredParentLOD.Get(EExtendableNodeLOD::Minimal) });
 
 				AddExtension(ParentStateTopExtension)
 					.SetLocation(EExtensionLocation::Top)
@@ -251,8 +278,21 @@ TSharedRef<SWidget> SStateCentricViewMainStateNode::Construct_GenerateNodeBody(c
 
 void SStateCentricViewMainStateNode::HandleOnSelectionChanged(const TArray<TWeakObjectPtr<UStateTreeState>>& SelectedStates)
 {
-	ViewState = ViewModel->GetLastSelectedState();
+	UStateTreeState* LastSelectedState = ViewModel->GetLastSelectedState();
+	if (ViewState != LastSelectedState)
+	{
+		ViewState = LastSelectedState;
+		RegenerateWidgetForLOD(CurrentLOD);
+	}
+}
+
+void SStateCentricViewMainStateNode::HandleOnParentHireachyLOD_Changed(TOptional<EExtendableNodeLOD> InLOD)
+{
+	RegenerateWidgetForLOD(CurrentLOD);
+}
 
+void SStateCentricViewMainStateNode::HandleOnNumParentsShownChanged(TOptional<uint32> InNumParentShown)
+{
 	RegenerateWidgetForLOD(CurrentLOD);
 }
 
diff --git a/Engine/Plugins/Runtime/StateTree/Source/StateTreeEditorModule/Private/StateCentricView/View/SStateCentricViewMainStateNode.h b/Engine/Plugins/Runtime/StateTree/Source/StateTreeEditorModule/Private/StateCentricView/View/SStateCentricViewMainStateNode.h
--- a/Engine/Plugins/Runtime/StateTree/Source/StateTreeEditorModule/Private/StateCentricView/View/SStateCentricViewMainStateNode.h
+++ b/Engine/Plugins/Runtime/StateTree/Source/StateTreeEditorModule/Private/StateCentricView/View/SStateCentricViewMainStateNode.h
@@ -43,6 +43,12 @@ protected:
 	/** Callback whenever selected state changes. */
 	void HandleOnSelectionChanged(const TArray<TWeakObjectPtr<UStateTreeState>>& SelectedStates);
 
+	/** Callback whenever parent state exapnsion changes */
+	void HandleOnParentHireachyLOD_Changed(TOptional<EExtendableNodeLOD> InLOD);
+
+	/** Callback whenever num parents shown changes */
+	void HandleOnNumParentsShownChanged(TOptional<uint32> InNumParentShown);
+
 private:
 
 	TArray<TSharedPtr<FDeferredCleanupSlateBrush>> ChildStateBrushes;
diff --git a/Engine/Plugins/Runtime/StateTree/Source/StateTreeEditorModule/Private/StateCentricView/View/StateCentricViewUtils.cpp b/Engine/Plugins/Runtime/StateTree/Source/StateTreeEditorModule/Private/StateCentricView/View/StateCentricViewUtils.cpp
--- a/Engine/Plugins/Runtime/StateTree/Source/StateTreeEditorModule/Private/StateCentricView/View/StateCentricViewUtils.cpp
+++ b/Engine/Plugins/Runtime/StateTree/Source/StateTreeEditorModule/Private/StateCentricView/View/StateCentricViewUtils.cpp
@@ -3,6 +3,7 @@
 #include "StateCentricView/View/StateCentricViewUtils.h"
 
 #include "StateCentricView/StateCentricViewSettings.h"
+#include "StateCentricView/ViewModel/StateCentricViewEditorDataExtension.h"
 #include "StateCentricView/View/SStateTreeExtendableNode.h"
 #include "StateCentricView/View/STransitionNode.h"
 #include "StateTreeEditorStyle.h"
@@ -250,6 +251,36 @@ void FStateCentricViewUtils::AddChildStateFootnoteExtensions(const EExtendableNo
 	}
 }
 
+TOptional<EExtendableNodeLOD> FStateCentricViewUtils::GetDesiredParentLOD(TSharedRef<FStateTreeViewModel> ViewModel, TNotNull<UStateTreeState*> ParentState)
+{
+	UStateTreeEditorData* EditorData = const_cast<UStateTreeEditorData*>(ViewModel->GetStateTreeEditorData());
+	UStateCentricViewEditorDataExtension* StateCentricEditorData = EditorData->GetOrCreateExtension<UStateCentricViewEditorDataExtension>();
+
+	TOptional<EExtendableNodeLOD> PerUserSettingsParentLOD = UStateCentricViewPerUserSettings::GetSchemaSettingsDefaultFallback(&ViewModel.Get()).GetParentHireachyLOD();
+	TOptional<EExtendableNodeLOD> PerStateTreeParentLOD = StateCentricEditorData->PerStateTreeUserSettings.GetParentHireachyLOD();
+
+	if (EExtendableNodeLOD* LastUserSetParentLOD = StateCentricEditorData->LastUserSetParentStateLOD.Find(ParentState->ID))
+	{
+		return *LastUserSetParentLOD;
+	}
+
+	return PerStateTreeParentLOD ? PerStateTreeParentLOD : PerUserSettingsParentLOD;
+}
+
+uint32 FStateCentricViewUtils::GetNumParentsToShow(TSharedRef<FStateTreeViewModel> ViewModel)
+{
+	UStateTreeEditorData* EditorData = const_cast<UStateTreeEditorData*>(ViewModel->GetStateTreeEditorData());
+	UStateCentricViewEditorDataExtension* StateCentricEditorData = EditorData->GetOrCreateExtension<UStateCentricViewEditorDataExtension>();
+
+	uint32 GlobalNumParentsToShow = UStateCentricViewSettings::GetSchemaSettingsDefaultFallback(&ViewModel.Get()).MaxNumParentsToShow;
+	TOptional<uint32> UserNumParentsToShow = UStateCentricViewPerUserSettings::GetSchemaSettingsDefaultFallback(&ViewModel.Get()).GetNumParentsShown();
+
+	uint32 SettingsNumParentsToShow = UserNumParentsToShow.Get(GlobalNumParentsToShow);
+	TOptional<uint32> PerStateTreeNumParentsToShow = StateCentricEditorData->PerStateTreeUserSettings.GetNumParentsShown();
+
+	return PerStateTreeNumParentsToShow.Get(SettingsNumParentsToShow);
+}
+
 TSharedPtr<FDeferredCleanupSlateBrush> FStateCentricViewUtils::MakeStateHalfTopOutline(TSharedRef<FStateTreeViewModel> ViewModel, TNotNull<const UStateTreeState*> State)
 {
 	const FSlateBrush* BaseBrush = FStateTreeEditorStyle::Get().GetBrush("StateTree.CentricView.Parent.HalfTopOutline");
diff --git a/Engine/Plugins/Runtime/StateTree/Source/StateTreeEditorModule/Private/StateCentricView/View/StateCentricViewUtils.h b/Engine/Plugins/Runtime/StateTree/Source/StateTreeEditorModule/Private/StateCentricView/View/StateCentricViewUtils.h
--- a/Engine/Plugins/Runtime/StateTree/Source/StateTreeEditorModule/Private/StateCentricView/View/StateCentricViewUtils.h
+++ b/Engine/Plugins/Runtime/StateTree/Source/StateTreeEditorModule/Private/StateCentricView/View/StateCentricViewUtils.h
@@ -44,6 +44,12 @@ struct FStateCentricViewUtils
 		, TNotNull<SStateTreeExtendableNode*> InSourceNode
 		, TArray<TSharedPtr<FDeferredCleanupSlateBrush>>& OutSlateBrushResouces);
 
+	/** Gets desired parent LOD based on various settings */
+	static TOptional<EExtendableNodeLOD> GetDesiredParentLOD(TSharedRef<FStateTreeViewModel> ViewModel, TNotNull<UStateTreeState*> ParentState);
+
+	/** Gets number of parents to show */
+	static uint32 GetNumParentsToShow(TSharedRef<FStateTreeViewModel> ViewModel);
+
 	/** Helpers to get brushes based on if state is parent / leaf w/ appropriate colors */
 	[[nodiscard]] static TSharedPtr<FDeferredCleanupSlateBrush> MakeStateHalfTopOutline(TSharedRef<FStateTreeViewModel> ViewModel, TNotNull<const UStateTreeState*> State);
 	[[nodiscard]] static TSharedPtr<FDeferredCleanupSlateBrush> MakeStateHalfTopOutlineMuted(TSharedRef<FStateTreeViewModel> ViewModel, TNotNull<const UStateTreeState*> State);
diff --git a/Engine/Plugins/Runtime/StateTree/Source/StateTreeEditorModule/Private/StateCentricView/ViewModel/StateCentricViewEditorDataExtension.h b/Engine/Plugins/Runtime/StateTree/Source/StateTreeEditorModule/Private/StateCentricView/ViewModel/StateCentricViewEditorDataExtension.h
--- a/Engine/Plugins/Runtime/StateTree/Source/StateTreeEditorModule/Private/StateCentricView/ViewModel/StateCentricViewEditorDataExtension.h
+++ b/Engine/Plugins/Runtime/StateTree/Source/StateTreeEditorModule/Private/StateCentricView/ViewModel/StateCentricViewEditorDataExtension.h
@@ -41,6 +41,9 @@ public:
 	UPROPERTY(Transient)
 	TMap<FGuid, EExtendableNodeLOD> LastUserSetParentStateLOD;
 
+	UPROPERTY(Transient)
+	bool bShouldStoreParentExpansion = true;
+
 public:
 
 	float LeftSpacerSplitterValue = 0.025f;
diff --git a/Engine/Plugins/Runtime/StateTree/Source/StateTreeEditorModule/Private/StateTreeViewModel.cpp b/Engine/Plugins/Runtime/StateTree/Source/StateTreeEditorModule/Private/StateTreeViewModel.cpp
--- a/Engine/Plugins/Runtime/StateTree/Source/StateTreeEditorModule/Private/StateTreeViewModel.cpp
+++ b/Engine/Plugins/Runtime/StateTree/Source/StateTreeEditorModule/Private/StateTreeViewModel.cpp
@@ -437,6 +437,16 @@ const UStateTreeEditorData* FStateTreeViewModel::GetStateTreeEditorData() const
 	return TreeDataWeak.Get();
 }
 
+const UStateTreeSchema* FStateTreeViewModel::GetStateTreeSchema() const
+{
+	if (const UStateTreeEditorData* TreeData = TreeDataWeak.Get())
+	{
+		return TreeData->Schema;
+	}
+
+	return nullptr;
+}
+
 const UStateTreeState* FStateTreeViewModel::GetStateByID(const FGuid StateID) const
 {
 	if (const UStateTreeEditorData* TreeData = TreeDataWeak.Get())
diff --git a/Engine/Plugins/Runtime/StateTree/Source/StateTreeEditorModule/Public/StateTreeViewModel.h b/Engine/Plugins/Runtime/StateTree/Source/StateTreeEditorModule/Public/StateTreeViewModel.h
--- a/Engine/Plugins/Runtime/StateTree/Source/StateTreeEditorModule/Public/StateTreeViewModel.h
+++ b/Engine/Plugins/Runtime/StateTree/Source/StateTreeEditorModule/Public/StateTreeViewModel.h
@@ -19,6 +19,7 @@ struct FStateTreeTransition;
 class FMenuBuilder;
 class UStateTree;
 class UStateTreeEditorData;
+class UStateTreeSchema;
 class UStateTreeState;
 
 enum class ECheckBoxState : uint8;
@@ -101,6 +102,7 @@ public:
 	UE_API const UStateTree* GetStateTree() const;
 
 	UE_API const UStateTreeEditorData* GetStateTreeEditorData() const;
+	UE_API const UStateTreeSchema* GetStateTreeSchema() const;
 
 	UE_API const UStateTreeState* GetStateByID(const FGuid StateID) const;
 	UE_API UStateTreeState* GetMutableStateByID(const FGuid StateID) const;
