Add support for property binding to state centric
- Note: PF Skip since Horde is down & relatively trivial / editor only change

#tests editor
#Synced-CL 48414911

diff --git a/Engine/Plugins/Runtime/StateTree/Source/StateTreeEditorModule/Private/StateCentricView/View/SParentStateExtensionNode.cpp b/Engine/Plugins/Runtime/StateTree/Source/StateTreeEditorModule/Private/StateCentricView/View/SParentStateExtensionNode.cpp
--- a/Engine/Plugins/Runtime/StateTree/Source/StateTreeEditorModule/Private/StateCentricView/View/SParentStateExtensionNode.cpp
+++ b/Engine/Plugins/Runtime/StateTree/Source/StateTreeEditorModule/Private/StateCentricView/View/SParentStateExtensionNode.cpp
@@ -2,6 +2,7 @@
 
 #include "StateCentricView/View/SParentStateExtensionNode.h"
 
+#include "Customizations/StateTreeBindingExtension.h"
 #include "StateCentricView/ViewModel/StateCentricViewEditorDataExtension.h"
 #include "StateCentricView/View/SParentStateTransitionExtensionNode.h"
 #include "StateCentricView/View/StateCentricViewUtils.h"
@@ -272,6 +273,8 @@ TSharedRef<SWidget> SParentStateExtensionNode::Construct_GenerateNodeBody(const
 			DetailsViewArgs.bAllowSearch = false;
 
 			TSharedRef<IDetailsView> StateDetailsView = PropertyEditorModule.CreateDetailView(DetailsViewArgs);
+			StateDetailsView->SetExtensionHandler(MakeShared<FStateTreeBindingExtension>());
+			StateDetailsView->SetChildrenCustomizationHandler(MakeShared<FStateTreeBindingsChildrenCustomization>());
 			StateDetailsView->SetObject(ViewState.Get());
 
 			return StateDetailsView;
diff --git a/Engine/Plugins/Runtime/StateTree/Source/StateTreeEditorModule/Private/StateCentricView/View/SStateCentricViewMainStateNode.cpp b/Engine/Plugins/Runtime/StateTree/Source/StateTreeEditorModule/Private/StateCentricView/View/SStateCentricViewMainStateNode.cpp
--- a/Engine/Plugins/Runtime/StateTree/Source/StateTreeEditorModule/Private/StateCentricView/View/SStateCentricViewMainStateNode.cpp
+++ b/Engine/Plugins/Runtime/StateTree/Source/StateTreeEditorModule/Private/StateCentricView/View/SStateCentricViewMainStateNode.cpp
@@ -2,6 +2,7 @@
 
 #include "StateCentricView/View/SStateCentricViewMainStateNode.h"
 
+#include "Customizations/StateTreeBindingExtension.h"
 #include "StateCentricView/ViewModel/StateCentricViewEditorDataExtension.h"
 #include "StateCentricView/View/SMainStateTransitionExtensionNode.h"
 #include "StateCentricView/View/SParentStateExtensionNode.h"
@@ -268,6 +269,8 @@ TSharedRef<SWidget> SStateCentricViewMainStateNode::Construct_GenerateNodeBody(c
 		DetailsViewArgs.bAllowSearch = false;
 
 		TSharedRef<IDetailsView> StateDetailsView = PropertyEditorModule.CreateDetailView(DetailsViewArgs);
+		StateDetailsView->SetExtensionHandler(MakeShared<FStateTreeBindingExtension>());
+		StateDetailsView->SetChildrenCustomizationHandler(MakeShared<FStateTreeBindingsChildrenCustomization>());
 		StateDetailsView->SetObject(ViewState.Get());
 
 		return StateDetailsView;
diff --git a/Engine/Plugins/Runtime/StateTree/Source/StateTreeEditorModule/Private/StateCentricView/View/STransitionNode.cpp b/Engine/Plugins/Runtime/StateTree/Source/StateTreeEditorModule/Private/StateCentricView/View/STransitionNode.cpp
--- a/Engine/Plugins/Runtime/StateTree/Source/StateTreeEditorModule/Private/StateCentricView/View/STransitionNode.cpp
+++ b/Engine/Plugins/Runtime/StateTree/Source/StateTreeEditorModule/Private/StateCentricView/View/STransitionNode.cpp
@@ -2,6 +2,7 @@
 
 #include "StateCentricView/View/STransitionNode.h"
 
+#include "Customizations/StateTreeBindingExtension.h"
 #include "StateCentricView/View/StateCentricViewUtils.h"
 #include "StateTreeEditorStyle.h"
 
@@ -356,6 +357,8 @@ TSharedRef<SWidget> STransitionNode::Construct_GenerateNodeBody(const EExtendabl
 			});
 
 			TSharedRef<IDetailsView> StateDetailsView = PropertyEditorModule.CreateDetailView(DetailsViewArgs);
+			StateDetailsView->SetExtensionHandler(MakeShared<FStateTreeBindingExtension>());
+			StateDetailsView->SetChildrenCustomizationHandler(MakeShared<FStateTreeBindingsChildrenCustomization>());
 			StateDetailsView->SetObject(StateOwningTransition);
 
 			return StateDetailsView;
