[StateTree] Prevent editor crash when PIE closes.
#rnx
#Synced-CL 48343893

diff --git a/Engine/Plugins/Runtime/StateTree/Source/StateTreeEditorModule/Private/StateTreeEditorMode.cpp b/Engine/Plugins/Runtime/StateTree/Source/StateTreeEditorModule/Private/StateTreeEditorMode.cpp
--- a/Engine/Plugins/Runtime/StateTree/Source/StateTreeEditorModule/Private/StateTreeEditorMode.cpp
+++ b/Engine/Plugins/Runtime/StateTree/Source/StateTreeEditorModule/Private/StateTreeEditorMode.cpp
@@ -474,18 +474,20 @@ bool UStateTreeEditorMode::CanCompile() const
 
 bool UStateTreeEditorMode::IsCompileVisible() const
 {
-	if(!HasValidStateTree())
+	const UInteractiveToolManager* ToolManager = GetToolManager();
+	if(!HasValidStateTree() || ToolManager == nullptr)
 	{
 		return false;
 	}
 
-	if (const UContextObjectStore* ContextObjectStore = GetToolManager()->GetContextObjectStore())
+	if (const UContextObjectStore* ContextObjectStore = ToolManager->GetContextObjectStore())
 	{
 		if (const UStateTreeEditorContext* Context = ContextObjectStore->FindContext<UStateTreeEditorContext>())
 		{
 			return Context->EditorHostInterface->ShouldShowCompileButton();
 		}
 	}
+
 	return true;
 }
 
