[StateTree] Fix bad memory access in EvaluationScope with BlueprintNode. FStateTreeDataView has the struct memory pointer or the UObject if the struct is a FStateTreeInstanceObjectWrapper. The EvaluationScope that it into account, but if the first item in the container was a Blueprint, then the following had the wrong address.

diff --git a/Engine/Plugins/Runtime/StateTree/Source/StateTreeModule/Private/StateTreeEvaluationScopeInstanceContainer.cpp b/Engine/Plugins/Runtime/StateTree/Source/StateTreeModule/Private/StateTreeEvaluationScopeInstanceContainer.cpp
--- a/Engine/Plugins/Runtime/StateTree/Source/StateTreeModule/Private/StateTreeEvaluationScopeInstanceContainer.cpp
+++ b/Engine/Plugins/Runtime/StateTree/Source/StateTreeModule/Private/StateTreeEvaluationScopeInstanceContainer.cpp
@@ -25,13 +25,10 @@ namespace Private
 
 void FEvaluationScopeInstanceContainer::FMemoryRequirementBuilder::Add(TNotNull<const UScriptStruct*> Struct)
 {
-	if (!Struct->IsChildOf(FStateTreeInstanceObjectWrapper::StaticStruct()))
+	Private::Add(MemoryRequirement, Struct);
+	if (FirstStructAlignment <= 0)
 	{
-		Private::Add(MemoryRequirement, Struct);
-		if (FirstStructAlignment <= 0)
-		{
-			FirstStructAlignment = Struct->GetMinAlignment();
-		}
+		FirstStructAlignment = Struct->GetMinAlignment();
 	}
 	++MemoryRequirement.NumberOfElements;
 }
@@ -77,45 +74,45 @@ void FEvaluationScopeInstanceContainer::Add(FStateTreeDataHandle DataHandle, FCo
 
 	FItem* ItemContainer = static_cast<FItem*>(Memory);
 	const UScriptStruct* Struct = DefaultInstance.GetScriptStruct();
-	if (!Struct->IsChildOf(FStateTreeInstanceObjectWrapper::StaticStruct()))
-	{
-		void* MemoryForNewItem = (NumberOfElements > 0)
-			? (uint8*)ItemContainer[NumberOfElements - 1].Instance.GetMutableMemory()
-				+ ItemContainer[NumberOfElements - 1].Instance.GetStruct()->GetStructureSize()
-				+ sizeof(Private::DebugStructEndTag)
-			: (uint8*)Memory + (MemoryRequirement.NumberOfElements * sizeof(FItem))
-				+ sizeof(Private::DebugTableEndTag);
-		MemoryForNewItem = Align(MemoryForNewItem, Struct->GetMinAlignment());
-
-		// Initialize the instance data
-		constexpr int32 ArrayDim = 1;
-		Struct->InitializeStruct(MemoryForNewItem);
-		Struct->CopyScriptStruct(MemoryForNewItem, DefaultInstance.GetMemory(), ArrayDim);
 
-		// Initialize the item in the table
-		new (ItemContainer + NumberOfElements)FItem();
-		ItemContainer[NumberOfElements].Instance = FStateTreeDataView(Struct, (uint8*)MemoryForNewItem);
-		ItemContainer[NumberOfElements].DataHandle = DataHandle;
+	void* MemoryForNewItem = (NumberOfElements > 0)
+		? (uint8*)ItemContainer[NumberOfElements - 1].Instance.GetMemory()
+			+ ItemContainer[NumberOfElements - 1].Instance.GetScriptStruct()->GetStructureSize()
+			+ sizeof(Private::DebugStructEndTag)
+		: (uint8*)Memory + (MemoryRequirement.NumberOfElements * sizeof(FItem))
+			+ sizeof(Private::DebugTableEndTag);
+	MemoryForNewItem = Align(MemoryForNewItem, Struct->GetMinAlignment());
 
-		AddStructDebugTag(NumberOfElements);
+	// Initialize the instance data
+	constexpr int32 ArrayDim = 1;
+	Struct->InitializeStruct(MemoryForNewItem);
 
-		if ((Struct->StructFlags & (STRUCT_IsPlainOldData | STRUCT_NoDestructor)) == 0)
-		{
-			bStructsHaveDestructor = true;
-		}
+	if (!Struct->IsChildOf(FStateTreeInstanceObjectWrapper::StaticStruct()))
+	{
+		Struct->CopyScriptStruct(MemoryForNewItem, DefaultInstance.GetMemory(), ArrayDim);
 	}
 	else
 	{
-		new (ItemContainer + NumberOfElements)FItem();
-
+		// Duplicate the instance object. Copy nullptr otherwise.
 		const FStateTreeInstanceObjectWrapper& Wrapper = DefaultInstance.Get<const FStateTreeInstanceObjectWrapper>();
+		FStateTreeInstanceObjectWrapper NewWrapperInstance;
 		if (Wrapper.InstanceObject)
 		{
-			UObject* DuplicatedObject = ::DuplicateObject(Wrapper.InstanceObject, GetTransientPackage());
-
-			ItemContainer[NumberOfElements].Instance = FStateTreeDataView(DuplicatedObject);
-			ItemContainer[NumberOfElements].DataHandle = DataHandle;
+			NewWrapperInstance.InstanceObject = ::DuplicateObject(Wrapper.InstanceObject, GetTransientPackage());
 		}
+		Struct->CopyScriptStruct(MemoryForNewItem, &NewWrapperInstance, ArrayDim);
+	}
+
+	// Initialize the item in the table
+	new (ItemContainer + NumberOfElements)FItem();
+	ItemContainer[NumberOfElements].Instance = FStructView(Struct, (uint8*)MemoryForNewItem);
+	ItemContainer[NumberOfElements].DataHandle = DataHandle;
+
+	AddStructDebugTag(NumberOfElements);
+
+	if ((Struct->StructFlags & (STRUCT_IsPlainOldData | STRUCT_NoDestructor)) == 0)
+	{
+		bStructsHaveDestructor = true;
 	}
 	++NumberOfElements;
 }
@@ -129,11 +126,11 @@ void FEvaluationScopeInstanceContainer::Reset()
 		FItem* ItemContainer = static_cast<FItem*>(Memory);
 		for (; NumberOfElements > 0; --NumberOfElements)
 		{
-			const UStruct* InstanceStruct = ItemContainer[NumberOfElements - 1].Instance.GetStruct();
+			const UStruct* InstanceStruct = ItemContainer[NumberOfElements - 1].Instance.GetScriptStruct();
 			if (InstanceStruct && Cast<const UClass>(InstanceStruct) == nullptr)
 			{
-				FStateTreeDataView& Instance = ItemContainer[NumberOfElements - 1].Instance;
-				Instance.GetStruct()->DestroyStruct(Instance.GetMutableMemory());
+				FStructView& Instance = ItemContainer[NumberOfElements - 1].Instance;
+				Instance.GetScriptStruct()->DestroyStruct(Instance.GetMemory());
 			}
 			if constexpr (!std::is_trivially_destructible_v<FItem>)
 			{
@@ -170,7 +167,7 @@ FEvaluationScopeInstanceContainer::FItem* FEvaluationScopeInstanceContainer::Get
 void FEvaluationScopeInstanceContainer::AddTableDebugTag()
 {
 #if WITH_STATETREE_DEBUG
-	void* MemoryForDebugTag = (uint8*)Memory + (sizeof(FItem) * MemoryRequirement.NumberOfElements);
+	void* MemoryForDebugTag = (uint8*)Memory + (MemoryRequirement.NumberOfElements * sizeof(FItem));
 	*reinterpret_cast<uint32*>(MemoryForDebugTag) = Private::DebugTableEndTag;
 #endif
 }
@@ -179,8 +176,8 @@ void FEvaluationScopeInstanceContainer::AddStructDebugTag(int32 Index)
 {
 #if WITH_STATETREE_DEBUG
 	FItem* ItemContainer = static_cast<FItem*>(Memory);
-	void* MemoryForDebugTag = (uint8*)ItemContainer[Index].Instance.GetMutableMemory()
-		+ ItemContainer[Index].Instance.GetStruct()->GetStructureSize();
+	void* MemoryForDebugTag = (uint8*)ItemContainer[Index].Instance.GetMemory()
+		+ ItemContainer[Index].Instance.GetScriptStruct()->GetStructureSize();
 	*reinterpret_cast<uint32*>(MemoryForDebugTag) = Private::DebugStructEndTag;
 #endif
 }
@@ -199,10 +196,10 @@ void FEvaluationScopeInstanceContainer::TestDebugTags() const
 		for (int32 Index = 0; Index < NumberOfElements; ++Index)
 		{
 			// Object wrapper do not have EndTag, they use the FItem directly.
-			const UStruct* InstanceStruct = ItemContainer[Index].Instance.GetStruct();
+			const UStruct* InstanceStruct = ItemContainer[Index].Instance.GetScriptStruct();
 			if (InstanceStruct && Cast<const UClass>(InstanceStruct) == nullptr)
 			{
-				const void* MemoryForStructDebugTag = (const uint8*)ItemContainer[Index].Instance.GetMutableMemory()
+				const void* MemoryForStructDebugTag = (const uint8*)ItemContainer[Index].Instance.GetMemory()
 					+ InstanceStruct->GetStructureSize();
 				const int32 StructEndTag = *reinterpret_cast<const uint32*>(MemoryForStructDebugTag);
 				ensure(StructEndTag == Private::DebugStructEndTag);
diff --git a/Engine/Plugins/Runtime/StateTree/Source/StateTreeModule/Private/StateTreeExecutionContext.cpp b/Engine/Plugins/Runtime/StateTree/Source/StateTreeModule/Private/StateTreeExecutionContext.cpp
--- a/Engine/Plugins/Runtime/StateTree/Source/StateTreeModule/Private/StateTreeExecutionContext.cpp
+++ b/Engine/Plugins/Runtime/StateTree/Source/StateTreeModule/Private/StateTreeExecutionContext.cpp
@@ -2662,9 +2662,14 @@ FStateTreeDataView FStateTreeExecutionContext::GetDataView(const FStateTreeExecu
 		{
 			if (EvaluationScopeInstanceCaches[Index].StateTree == CurrentFrame.StateTree)
 			{
-				if (FStateTreeDataView* Result = EvaluationScopeInstanceCaches[Index].Container->GetDataViewPtr(Handle))
+				if (FStructView* Result = EvaluationScopeInstanceCaches[Index].Container->GetDataViewPtr(Handle))
 				{
-					return *Result;
+					if (Handle.IsObjectSource())
+					{
+						const FStateTreeInstanceObjectWrapper& Wrapper = Result->Get<FStateTreeInstanceObjectWrapper>();
+						return FStateTreeDataView(Wrapper.InstanceObject);
+					};
+					return FStateTreeDataView(*Result);
 				}
 			}
 		}
@@ -3011,7 +3016,7 @@ bool FStateTreeExecutionContext::IsHandleSourceValid(const FStateTreeExecutionFr
 		{
 			if (EvaluationScopeInstanceCaches[Index].StateTree == CurrentFrame.StateTree)
 			{
-				if (FStateTreeDataView* Result = EvaluationScopeInstanceCaches[Index].Container->GetDataViewPtr(Handle))
+				if (const FStructView* Result = EvaluationScopeInstanceCaches[Index].Container->GetDataViewPtr(Handle))
 				{
 					return true;
 				}
diff --git a/Engine/Plugins/Runtime/StateTree/Source/StateTreeModule/Public/StateTreeEvaluationScopeInstanceContainer.h b/Engine/Plugins/Runtime/StateTree/Source/StateTreeModule/Public/StateTreeEvaluationScopeInstanceContainer.h
--- a/Engine/Plugins/Runtime/StateTree/Source/StateTreeModule/Public/StateTreeEvaluationScopeInstanceContainer.h
+++ b/Engine/Plugins/Runtime/StateTree/Source/StateTreeModule/Public/StateTreeEvaluationScopeInstanceContainer.h
@@ -72,13 +72,13 @@ namespace UE::StateTree::InstanceData
 		}
 
 		/** @return specified item of an added structure. */
-		[[nodiscard]] FStateTreeDataView GetDataView(FStateTreeDataHandle DataHandle)
+		[[nodiscard]] FStructView GetDataView(FStateTreeDataHandle DataHandle)
 		{
 			return Get(DataHandle)->Instance;
 		}
 
 		/** @return specified item if the structure was added. */
-		[[nodiscard]] FStateTreeDataView* GetDataViewPtr(FStateTreeDataHandle DataHandle)
+		[[nodiscard]] FStructView* GetDataViewPtr(FStateTreeDataHandle DataHandle)
 		{
 			FItem* FoundItem = Get(DataHandle);
 			return FoundItem ? &FoundItem->Instance : nullptr;
@@ -103,7 +103,7 @@ namespace UE::StateTree::InstanceData
 	private:
 		struct FItem
 		{
-			FStateTreeDataView Instance;
+			FStructView Instance;
 			FStateTreeDataHandle DataHandle;
 		};
 
