[StateTree] Enable async StateTree compilation on asset load. Ensure the asset is compiled before validation.
#Synced-CL 48343893

diff --git a/Engine/Plugins/Runtime/StateTree/Source/StateTreeEditorModule/Private/StateTreeCompilerManager.cpp b/Engine/Plugins/Runtime/StateTree/Source/StateTreeEditorModule/Private/StateTreeCompilerManager.cpp
--- a/Engine/Plugins/Runtime/StateTree/Source/StateTreeEditorModule/Private/StateTreeCompilerManager.cpp
+++ b/Engine/Plugins/Runtime/StateTree/Source/StateTreeEditorModule/Private/StateTreeCompilerManager.cpp
@@ -43,7 +43,8 @@ FAutoConsoleVariableRef CVarEnableCompileSynchronouslyInPIESession(
 	TEXT("Any new asset will be compile synchronously event is async was requested.")
 );
 
-int32 NumberOfQueuedCompilationPerBatch = 1;
+// Run 2 compilations per frame until all queued state tree assets are compiled.
+int32 NumberOfQueuedCompilationPerBatch = 2;
 FAutoConsoleVariableRef CVarNumberOfQueuedCompilationPerBatch(
 	TEXT("StateTree.Compiler.NumberOfQueuedCompilationPerBatch"),
 	NumberOfQueuedCompilationPerBatch,
@@ -590,6 +591,9 @@ void FCompilerManagerImpl::UpdateBindingsInstanceStructsIfNeeded(TSet<const UStr
 
 void FCompilerManagerImpl::HandleStateTreeCooking(TNotNull<UStateTree*> StateTree)
 {
+	// We do not use CompileIfNeededSynchronously here.
+	//CompileIfNeededSynchronously has extra logic to prevent multiple compilations for the same asset.
+	//We always want to compile when cooking.
 	switch (StateTree->CompileStatus)
 	{
 		case UStateTree::ECompileStatus::PublicDirty:
diff --git a/Engine/Plugins/Runtime/StateTree/Source/StateTreeEditorModule/Private/StateTreeEditorModule.cpp b/Engine/Plugins/Runtime/StateTree/Source/StateTreeEditorModule/Private/StateTreeEditorModule.cpp
--- a/Engine/Plugins/Runtime/StateTree/Source/StateTreeEditorModule/Private/StateTreeEditorModule.cpp
+++ b/Engine/Plugins/Runtime/StateTree/Source/StateTreeEditorModule/Private/StateTreeEditorModule.cpp
@@ -57,7 +57,7 @@ namespace UE::StateTree::Editor
 	 * That ensures that dependent state tree assets are compiled before being used.
 	 * When compiling in the editor (once the asset are loaded), you might have a dependency bug.
 	 */
-	bool bEnableQueuedCompilationOnAssetLoad = false;
+	bool bEnableQueuedCompilationOnAssetLoad = true;
 	FAutoConsoleVariableRef CVarEnableQueuedCompilationOnAssetLoad(
 		TEXT("StateTree.Compiler.EnableQueuedCompilationOnAssetLoad"),
 		bEnableQueuedCompilationOnAssetLoad,
diff --git a/Engine/Plugins/Runtime/StateTree/Source/StateTreeModule/Private/StateTree.cpp b/Engine/Plugins/Runtime/StateTree/Source/StateTreeModule/Private/StateTree.cpp
--- a/Engine/Plugins/Runtime/StateTree/Source/StateTreeModule/Private/StateTree.cpp
+++ b/Engine/Plugins/Runtime/StateTree/Source/StateTreeModule/Private/StateTree.cpp
@@ -28,6 +28,8 @@
 
 #include UE_INLINE_GENERATED_CPP_BY_NAME(StateTree)
 
+#define LOCTEXT_NAMESPACE "StateTree"
+
 PRAGMA_DISABLE_DEPRECATION_WARNINGS
 const FGuid FStateTreeCustomVersion::GUID(0x28E21331, 0x501F4723, 0x8110FA64, 0xEA10DA1E);
 namespace UE::StateTree::Private
@@ -343,20 +345,17 @@ EDataValidationResult UStateTree::IsDataValid(FDataValidationContext& Context) c
 	// Don't warn user that the tree they just saved is not compiled. Only for submit or manual validation
 	if (Context.GetValidationUsecase() != EDataValidationUsecase::Save)
 	{
-		if (UE::StateTree::Delegates::OnRequestEditorHash.IsBound())
+		const_cast<UStateTree*>(this)->CompileIfChanged();
+		if (CompileStatus != ECompileStatus::Executable)
 		{
-			const uint32 CurrentHash = UE::StateTree::Delegates::OnRequestEditorHash.Execute(*this);
-			if (CurrentHash != LastCompiledEditorDataHash)
-			{
-				Context.AddWarning(FText::FromString(FString::Printf(TEXT("%s is not compiled. Please recompile the State Tree."), *GetPathName())));
-				return EDataValidationResult::Invalid;
-			}
+			Context.AddWarning(FText::Format(LOCTEXT("ValidateFailedCompile", "{0} failed to compiled.Please recompile the State Tree for more details errors."), FText::FromString(GetPathName())));
+			return EDataValidationResult::Invalid;
 		}
 	}
 
 	if (!const_cast<UStateTree*>(this)->Link())
 	{
-		Context.AddError(FText::FromString(FString::Printf(TEXT("%s failed to link. Please recompile the State Tree for more details errors."), *GetPathName())));
+		Context.AddError(FText::Format(LOCTEXT("ValidateFailedLink", "{0} failed to link. Please recompile the State Tree for more details errors."), FText::FromString(GetPathName())));
 		return EDataValidationResult::Invalid;
 	}
 
@@ -1828,3 +1827,5 @@ void UStateTree::HandleRuntimeValidationPostGC()
 	GCObjectDatas.Reset();
 }
 #endif
+
+#undef LOCTEXT_NAMESPACE
