[State Tree] Expose UE::StateTree::Conditions::CompareNumbers funciton to new header

#Synced-CL 48638146

diff --git a/Engine/Plugins/Runtime/StateTree/Source/StateTreeModule/Private/Conditions/StateTreeCommonConditions.cpp b/Engine/Plugins/Runtime/StateTree/Source/StateTreeModule/Private/Conditions/StateTreeCommonConditions.cpp
--- a/Engine/Plugins/Runtime/StateTree/Source/StateTreeModule/Private/Conditions/StateTreeCommonConditions.cpp
+++ b/Engine/Plugins/Runtime/StateTree/Source/StateTreeModule/Private/Conditions/StateTreeCommonConditions.cpp
@@ -1,6 +1,8 @@
 // Copyright Epic Games, Inc. All Rights Reserved.
 
 #include "Conditions/StateTreeCommonConditions.h"
+
+#include "Conditions/StateTreeConditionHelpers.h"
 #include "StateTreeExecutionContext.h"
 #include "UObject/EnumProperty.h"
 #include "StateTreeNodeDescriptionHelpers.h"
@@ -9,42 +11,6 @@
 
 #define LOCTEXT_NAMESPACE "StateTree"
 
-namespace UE::StateTree::Conditions
-{
-
-template<typename T>
-bool CompareNumbers(const T Left, const T Right, const EGenericAICheck Operator)
-{
-	switch (Operator)
-	{
-	case EGenericAICheck::Equal:
-		return Left == Right;
-		break;
-	case EGenericAICheck::NotEqual:
-		return Left != Right;
-		break;
-	case EGenericAICheck::Less:
-		return Left < Right;
-		break;
-	case EGenericAICheck::LessOrEqual:
-		return Left <= Right;
-		break;
-	case EGenericAICheck::Greater:
-		return Left > Right;
-		break;
-	case EGenericAICheck::GreaterOrEqual:
-		return Left >= Right;
-		break;
-	default:
-		ensureMsgf(false, TEXT("Unhandled operator %d"), Operator);
-		return false;
-		break;
-	}
-}
-
-} // UE::StateTree::Conditions
-
-
 //----------------------------------------------------------------------//
 //  FStateTreeCompareIntCondition
 //----------------------------------------------------------------------//
diff --git a/Engine/Plugins/Runtime/StateTree/Source/StateTreeModule/Public/Conditions/StateTreeConditionHelpers.h b/Engine/Plugins/Runtime/StateTree/Source/StateTreeModule/Public/Conditions/StateTreeConditionHelpers.h
new file mode 100644
--- /dev/null
+++ b/Engine/Plugins/Runtime/StateTree/Source/StateTreeModule/Public/Conditions/StateTreeConditionHelpers.h
@@ -0,0 +1,31 @@
+// Copyright Epic Games, Inc. All Rights Reserved.
+
+#pragma once
+
+#include "AITypes.h"
+
+namespace UE::StateTree::Conditions
+{
+	template<typename T>
+	bool CompareNumbers(const T Left, const T Right, const EGenericAICheck Operator)
+	{
+		switch (Operator)
+		{
+		case EGenericAICheck::Equal:
+			return Left == Right;
+		case EGenericAICheck::NotEqual:
+			return Left != Right;
+		case EGenericAICheck::Less:
+			return Left < Right;
+		case EGenericAICheck::LessOrEqual:
+			return Left <= Right;
+		case EGenericAICheck::Greater:
+			return Left > Right;
+		case EGenericAICheck::GreaterOrEqual:
+			return Left >= Right;
+		default:
+			ensureMsgf(false, TEXT("Unhandled operator %d"), Operator);
+			return false;
+		}
+	}
+} // UE::StateTree::Conditions
