Show number of enter conditions / tasks / transitions / considerations / parameters in details header
- Also show utility weight if it's just a weight

#tests editor
#Synced-CL 48677313

diff --git a/Engine/Plugins/Runtime/StateTree/Source/StateTreeEditorModule/Private/Customizations/StateTreeEditorDataDetails.cpp b/Engine/Plugins/Runtime/StateTree/Source/StateTreeEditorModule/Private/Customizations/StateTreeEditorDataDetails.cpp
--- a/Engine/Plugins/Runtime/StateTree/Source/StateTreeEditorModule/Private/Customizations/StateTreeEditorDataDetails.cpp
+++ b/Engine/Plugins/Runtime/StateTree/Source/StateTreeEditorModule/Private/Customizations/StateTreeEditorDataDetails.cpp
@@ -12,6 +12,7 @@
 #include "StateTreeEditorNodeUtils.h"
 #include "StateTreeEditorSchema.h"
 #include "StateTreeEditorStyle.h"
+#include "StateTreeEditorUserSettings.h"
 #include "StateTreeGraphSchema.h"
 #include "StateTreeSchema.h"
 #include "Styling/StyleColors.h"
@@ -148,8 +149,12 @@ void FStateTreeEditorDataDetails::CustomizeDetails(IDetailLayoutBuilder& DetailB
 
 	if (Schema && Schema->AllowGlobalParameters())
 	{
+		uint32 GlobalNumParameters = EditorData ? EditorData->GetRootParametersPropertyBag().GetNumPropertiesInBag() : 0;
+
 		// Parameters category
-		const FText ParametersDisplayName = LOCTEXT("EditorDataDetailsParameters", "Parameters");
+		const FText ParametersDisplayName = EnumHasAllFlags(GetDefault<UStateTreeEditorUserSettings>()->GetDetailsViewDisplayCount(), EStateTreeEditorUserSettingsNodeType::Parameter)
+			? FText::Format(LOCTEXT("EditorDataDetailsParametersCounted", "Parameters ({0})"), GlobalNumParameters)
+			: LOCTEXT("EditorDataDetailsParameters", "Parameters"); 
 		IDetailCategoryBuilder& ParametersCategory = DetailBuilder.EditCategory(TEXT("Parameters"), ParametersDisplayName);
 		ParametersCategory.SetSortOrder(2);
 		{
@@ -198,6 +203,8 @@ void FStateTreeEditorDataDetails::CustomizeDetails(IDetailLayoutBuilder& DetailB
 	}	
 
 	// Evaluators category
+	int32 GlobalNumEvaluators = EditorData ? EditorData->Evaluators.Num() : 0;
+
 	TSharedPtr<IPropertyHandle> EvaluatorsProperty = DetailBuilder.GetProperty(GET_MEMBER_NAME_CHECKED(UStateTreeEditorData, Evaluators));
 	check(EvaluatorsProperty.IsValid());
 	const FName EvalCategoryName(TEXT("Evaluators"));
@@ -207,7 +214,9 @@ void FStateTreeEditorDataDetails::CustomizeDetails(IDetailLayoutBuilder& DetailB
 			DetailBuilder,
 			EvaluatorsProperty,
 			EvalCategoryName,
-			LOCTEXT("EditorDataDetailsEvaluators", "Evaluators"),
+			EnumHasAllFlags(GetDefault<UStateTreeEditorUserSettings>()->GetDetailsViewDisplayCount(), EStateTreeEditorUserSettingsNodeType::Evaluator)
+				? FText::Format(LOCTEXT("EditorDataDetailsEvaluatorsCounted", "Evaluators ({0})"), GlobalNumEvaluators)
+				: LOCTEXT("EditorDataDetailsEvaluators", "Evaluators"),
 			FName("StateTreeEditor.Evaluators"),
 			UE::StateTree::Colors::Bronze,
 			UE::StateTree::Colors::Bronze.WithAlpha(192),
@@ -220,6 +229,8 @@ void FStateTreeEditorDataDetails::CustomizeDetails(IDetailLayoutBuilder& DetailB
 	}
 
 	// Global Tasks category
+	int32 GlobalNumTasks = EditorData ? EditorData->GlobalTasks.Num() : 0;
+
 	TSharedPtr<IPropertyHandle> GlobalTasksProperty = DetailBuilder.GetProperty(GET_MEMBER_NAME_CHECKED(UStateTreeEditorData, GlobalTasks));
 	check(GlobalTasksProperty.IsValid());
 	const FName GlobalTasksCategoryName(TEXT("Global Tasks"));
@@ -232,7 +243,9 @@ void FStateTreeEditorDataDetails::CustomizeDetails(IDetailLayoutBuilder& DetailB
 		DetailBuilder,
 		GlobalTasksProperty,
 		GlobalTasksCategoryName,
-		LOCTEXT("EditorDataDetailsGlobalTasks", "Global Tasks"),
+		EnumHasAllFlags(GetDefault<UStateTreeEditorUserSettings>()->GetDetailsViewDisplayCount(), EStateTreeEditorUserSettingsNodeType::Task)
+			? FText::Format(LOCTEXT("EditorDataDetailsGlobalTasksCounted", "Global Tasks ({0})"), GlobalNumTasks)
+			: LOCTEXT("EditorDataDetailsGlobalTasks", "Global Tasks"),
 		FName("StateTreeEditor.Tasks"),
 		UE::StateTree::Colors::Cyan,
 		bAllowTasksCompletion ? GlobalTasksCompletionProperty->CreatePropertyValueWidget(/*bDisplayDefaultPropertyButtons*/false) : TSharedPtr<SWidget>(),
diff --git a/Engine/Plugins/Runtime/StateTree/Source/StateTreeEditorModule/Private/Customizations/StateTreeStateDetails.cpp b/Engine/Plugins/Runtime/StateTree/Source/StateTreeEditorModule/Private/Customizations/StateTreeStateDetails.cpp
--- a/Engine/Plugins/Runtime/StateTree/Source/StateTreeEditorModule/Private/Customizations/StateTreeStateDetails.cpp
+++ b/Engine/Plugins/Runtime/StateTree/Source/StateTreeEditorModule/Private/Customizations/StateTreeStateDetails.cpp
@@ -13,6 +13,7 @@
 #include "StateTreeEditorNodeUtils.h"
 #include "StateTreeEditorStyle.h"
 #include "StateTreeEditorDataExtension.h"
+#include "StateTreeEditorUserSettings.h"
 #include "StateTreePropertyHelpers.h"
 #include "StateTreeSchema.h"
 #include "StateTreeStateParametersDetails.h"
@@ -253,7 +254,10 @@ void FStateTreeStateDetails::CustomizeDetails(IDetailLayoutBuilder& DetailBuilde
 	}
 
 	// Parameters category
-	const FText ParametersDisplayName = LOCTEXT("EditorStateDetailsParameters", "Parameters");
+	int32 NumParameters = WeakState.IsValid() ? WeakState->Parameters.Parameters.GetNumPropertiesInBag() : 0;
+	const FText ParametersDisplayName = EnumHasAllFlags(GetDefault<UStateTreeEditorUserSettings>()->GetDetailsViewDisplayCount(), EStateTreeEditorUserSettingsNodeType::Parameter)
+		? FText::Format(LOCTEXT("EditorStateDetailsParametersCounted", "Parameters ({0})"), NumParameters)
+		: LOCTEXT("EditorStateDetailsParameters", "Parameters");
 	IDetailCategoryBuilder& ParametersCategory = DetailBuilder.EditCategory(TEXT("Parameters"), ParametersDisplayName);
 	ParametersCategory.SetSortOrder(1);
 	{
@@ -322,12 +326,14 @@ void FStateTreeStateDetails::CustomizeDetails(IDetailLayoutBuilder& DetailBuilde
 	const FName EnterConditionsCategoryName(TEXT("Enter Conditions"));
 	if (Schema && Schema->AllowEnterConditions())
 	{
-
+		int32 NumEnterConditions = WeakState.IsValid() ? WeakState->EnterConditions.Num() : 0;
 		IDetailCategoryBuilder& EnterConditionsCategory = UE::StateTreeEditor::EditorNodeUtils::MakeArrayCategory(
 			DetailBuilder,
 			EnterConditionsProperty,
 			EnterConditionsCategoryName,
-			LOCTEXT("StateDetailsEnterConditions", "Enter Conditions"),
+			EnumHasAllFlags(GetDefault<UStateTreeEditorUserSettings>()->GetDetailsViewDisplayCount(), EStateTreeEditorUserSettingsNodeType::Condition)
+				? FText::Format(LOCTEXT("StateDetailsEnterConditionsCounted", "Enter Conditions ({0})"), NumEnterConditions)
+				: LOCTEXT("StateDetailsEnterConditions", "Enter Conditions"),
 			FName("StateTreeEditor.Conditions"),
 			UE::StateTree::Colors::Yellow,
 			UE::StateTree::Colors::Yellow.WithAlpha(192),
@@ -354,11 +360,16 @@ void FStateTreeStateDetails::CustomizeDetails(IDetailLayoutBuilder& DetailBuilde
 	if (Schema && Schema->AllowUtilityConsiderations())
 	{
 		const FName UtilityCategoryName(TEXT("Selection Utility"));
+		int32 NumConsiderations = WeakState.IsValid() ? WeakState->Considerations.Num() : 0;
+		FText UtilityCategoryText = EnumHasAllFlags(GetDefault<UStateTreeEditorUserSettings>()->GetDetailsViewDisplayCount(), EStateTreeEditorUserSettingsNodeType::Utility)
+			? FText::Format(LOCTEXT("StateDetailsSelectionUtilityCounted", "Selection Utility ({0})"), NumConsiderations)
+			: LOCTEXT("StateDetailsSelectionUtility", "Selection Utility");
+
 		IDetailCategoryBuilder& UtilityConsiderationsCategory = UE::StateTreeEditor::EditorNodeUtils::MakeArrayCategory(
 			DetailBuilder,
 			ConsiderationsProperty,
 			UtilityCategoryName,
-			LOCTEXT("StateDetailsSelectionUtility", "Selection Utility"),
+			UtilityCategoryText,
 			FName("StateTreeEditor.Utility"),
 			UE::StateTree::Colors::Orange,
 			UE::StateTree::Colors::Orange.WithAlpha(192),
@@ -374,13 +385,16 @@ void FStateTreeStateDetails::CustomizeDetails(IDetailLayoutBuilder& DetailBuilde
 	{
 		if (Schema && Schema->AllowMultipleTasks())
 		{
+			int32 NumTasks = WeakState.IsValid() ? WeakState->Tasks.Num() : 0;
 			const bool bAllowTasksCompletion = Schema->AllowTasksCompletion();
 			const FName TasksCategoryName(TEXT("Tasks"));
 			IDetailCategoryBuilder& TasksCategory = UE::StateTreeEditor::EditorNodeUtils::MakeArrayCategoryHeader(
 				DetailBuilder,
 				TasksProperty,
 				TasksCategoryName,
-				LOCTEXT("StateDetailsTasks", "Tasks"),
+				EnumHasAllFlags(GetDefault<UStateTreeEditorUserSettings>()->GetDetailsViewDisplayCount(), EStateTreeEditorUserSettingsNodeType::Task)
+					? FText::Format(LOCTEXT("StateDetailsTasksCounted", "Tasks ({0})"), NumTasks)
+					: LOCTEXT("StateDetailsTasks", "Tasks"),
 				FName("StateTreeEditor.Tasks"),
 				UE::StateTree::Colors::Cyan,
 				bAllowTasksCompletion ? TasksCompletionProperty->CreatePropertyValueWidget(/*bDisplayDefaultPropertyButtons*/false) : TSharedPtr<SWidget>(),
@@ -409,11 +423,14 @@ void FStateTreeStateDetails::CustomizeDetails(IDetailLayoutBuilder& DetailBuilde
 	}
 
 	// Transitions
+	int32 NumTransitions = WeakState.IsValid() ? WeakState->Transitions.Num() : 0;
 	UE::StateTreeEditor::EditorNodeUtils::MakeArrayCategory(
 		DetailBuilder,
 		TransitionsProperty,
 		"Transitions",
-		LOCTEXT("StateDetailsTransitions", "Transitions"),
+		EnumHasAllFlags(GetDefault<UStateTreeEditorUserSettings>()->GetDetailsViewDisplayCount(), EStateTreeEditorUserSettingsNodeType::Transition)
+			? FText::Format(LOCTEXT("StateDetailsTransitionsCounted", "Transitions ({0})"), NumTransitions)
+			: LOCTEXT("StateDetailsTransitions", "Transitions"),
 		FName("StateTreeEditor.Transitions"),
 		UE::StateTree::Colors::Magenta,
 		UE::StateTree::Colors::Magenta.WithAlpha(192),
@@ -429,11 +446,14 @@ void FStateTreeStateDetails::CustomizeDetails(IDetailLayoutBuilder& DetailBuilde
 			}
 		}));
 
-	for (UStateTreeEditorDataExtension* Extension : EditorData->Extensions)
+	if (EditorData)
 	{
-		if (Extension)
+		for (UStateTreeEditorDataExtension* Extension : EditorData->Extensions)
 		{
-			Extension->CustomizeDetails(WeakState.Get(), DetailBuilder);
+			if (Extension)
+			{
+				Extension->CustomizeDetails(WeakState.Get(), DetailBuilder);
+			}
 		}
 	}
 }
diff --git a/Engine/Plugins/Runtime/StateTree/Source/StateTreeEditorModule/Private/StateTreeEditorUserSettings.cpp b/Engine/Plugins/Runtime/StateTree/Source/StateTreeEditorModule/Private/StateTreeEditorUserSettings.cpp
--- a/Engine/Plugins/Runtime/StateTree/Source/StateTreeEditorModule/Private/StateTreeEditorUserSettings.cpp
+++ b/Engine/Plugins/Runtime/StateTree/Source/StateTreeEditorModule/Private/StateTreeEditorUserSettings.cpp
@@ -13,6 +13,15 @@ void UStateTreeEditorUserSettings::SetStatesViewDisplayNodeType(EStateTreeEditor
 	}
 }
 
+void UStateTreeEditorUserSettings::SetDetailsViewDisplayCount(EStateTreeEditorUserSettingsNodeType Value)
+{
+	if (DetailsViewDisplayCount != Value)
+	{
+		DetailsViewDisplayCount = Value;
+		OnSettingsChanged.Broadcast();
+	}
+}
+
 void UStateTreeEditorUserSettings::SetDisplayMigrationToRewindDebuggerWarning(bool bNewValue)
 {
 	if (bDisplayMigrationToRewindDebuggerWarning != bNewValue)
@@ -27,7 +36,8 @@ void UStateTreeEditorUserSettings::PostEditChangeChainProperty(struct FPropertyC
 	FName PropertyName = PropertyChangedEvent.Property->GetFName();
 	if (PropertyName == GET_MEMBER_NAME_CHECKED(UStateTreeEditorUserSettings, StatesViewDisplayNodeType)
 		|| PropertyName == GET_MEMBER_NAME_CHECKED(UStateTreeEditorUserSettings, StatesViewStateRowHeight)
-		|| PropertyName == GET_MEMBER_NAME_CHECKED(UStateTreeEditorUserSettings, StatesViewNodeRowHeight))
+		|| PropertyName == GET_MEMBER_NAME_CHECKED(UStateTreeEditorUserSettings, StatesViewNodeRowHeight)
+		|| PropertyName == GET_MEMBER_NAME_CHECKED(UStateTreeEditorUserSettings, DetailsViewDisplayCount))
 	{
 		OnSettingsChanged.Broadcast();
 	}
diff --git a/Engine/Plugins/Runtime/StateTree/Source/StateTreeEditorModule/Private/StateTreeEditorUserSettings.h b/Engine/Plugins/Runtime/StateTree/Source/StateTreeEditorModule/Private/StateTreeEditorUserSettings.h
--- a/Engine/Plugins/Runtime/StateTree/Source/StateTreeEditorModule/Private/StateTreeEditorUserSettings.h
+++ b/Engine/Plugins/Runtime/StateTree/Source/StateTreeEditorModule/Private/StateTreeEditorUserSettings.h
@@ -8,11 +8,15 @@
 UENUM(meta=(Bitflags, UseEnumValuesAsMaskValuesInEditor = "true"))
 enum class EStateTreeEditorUserSettingsNodeType : uint8
 {
+	None = 0 << 0,
 	Condition = 1 << 0,
 	Task = 1 << 1,
 	Transition = 1 << 2,
 	Flag = 1 << 3,
-	All = (Condition | Task | Transition | Flag) UMETA(Hidden),
+	Parameter = 1 << 4,
+	Utility = 1 << 5,
+	Evaluator = 1 << 6,
+	All = (Condition | Task | Transition | Flag | Parameter | Utility | Evaluator) UMETA(Hidden),
 };
 ENUM_CLASS_FLAGS(EStateTreeEditorUserSettingsNodeType);
 
@@ -31,6 +35,12 @@ public:
 	}
 	void SetStatesViewDisplayNodeType(EStateTreeEditorUserSettingsNodeType Value);
 
+	EStateTreeEditorUserSettingsNodeType GetDetailsViewDisplayCount() const
+	{
+		return DetailsViewDisplayCount;
+	}
+	void SetDetailsViewDisplayCount(EStateTreeEditorUserSettingsNodeType Value);
+
 	float GetStatesViewStateRowHeight() const
 	{
 		return StatesViewStateRowHeight;
@@ -67,6 +77,10 @@ private:
 	UPROPERTY(config, EditAnywhere, Category = "State View", meta = (DisplayName = "Node Height", ClampMin = 8.0f))
 	float StatesViewNodeRowHeight = 16.0f;
 
+	/** True if we should show the count for that node in details. */
+	UPROPERTY(config, EditAnywhere, Category = "Details View", meta = (DisplayName = "Display Count"))
+	EStateTreeEditorUserSettingsNodeType DetailsViewDisplayCount = EStateTreeEditorUserSettingsNodeType::All;
+
 	/** Display the debugger migration warning. */
 	UPROPERTY(config, EditAnywhere, Category = "Debugger")
 	bool bDisplayMigrationToRewindDebuggerWarning = true;
diff --git a/Engine/Plugins/Runtime/StateTree/Source/StateTreeEditorModule/Public/StateTreeState.h b/Engine/Plugins/Runtime/StateTree/Source/StateTreeEditorModule/Public/StateTreeState.h
--- a/Engine/Plugins/Runtime/StateTree/Source/StateTreeEditorModule/Public/StateTreeState.h
+++ b/Engine/Plugins/Runtime/StateTree/Source/StateTreeEditorModule/Public/StateTreeState.h
@@ -436,7 +436,7 @@ public:
 	UPROPERTY(EditDefaultsOnly, Category = "Tasks", meta = (BaseStruct = "/Script/StateTreeModule.StateTreeTaskBase", BaseClass = "/Script/StateTreeModule.StateTreeTaskBlueprintBase"))
 	TArray<FStateTreeEditorNode> Tasks;
 
-	/** Expression of enter conditions that needs to evaluate true to allow the state to be selected. */
+	/** Expression used to calculate a selection utility score for state weighing. */
 	UPROPERTY(EditDefaultsOnly, Category = "Utility", meta = (BaseStruct = "/Script/StateTreeModule.StateTreeConsiderationBase", BaseClass = "/Script/StateTreeModule.StateTreeConsiderationBlueprintBase"))
 	TArray<FStateTreeEditorNode> Considerations;
 
