[State Tree]
Clear parallel tree task's tick handle in exit state regardless of the tree validity.
Update Parallel Tree's instance data by start parameters.

#Synced-CL 48638146

diff --git a/Engine/Plugins/Runtime/StateTree/Source/StateTreeModule/Private/Tasks/StateTreeRunParallelStateTreeTask.cpp b/Engine/Plugins/Runtime/StateTree/Source/StateTreeModule/Private/Tasks/StateTreeRunParallelStateTreeTask.cpp
--- a/Engine/Plugins/Runtime/StateTree/Source/StateTreeModule/Private/Tasks/StateTreeRunParallelStateTreeTask.cpp
+++ b/Engine/Plugins/Runtime/StateTree/Source/StateTreeModule/Private/Tasks/StateTreeRunParallelStateTreeTask.cpp
@@ -45,19 +45,6 @@ EStateTreeRunStatus FStateTreeRunParallelStateTreeTask::EnterState(FStateTreeExe
 		return EStateTreeRunStatus::Failed;
 	}
 
-	if (InstanceData.ScheduledTickHandle.IsValid())
-	{
-		Context.RemoveScheduledTickRequest(InstanceData.ScheduledTickHandle);
-		InstanceData.ScheduledTickHandle = UE::StateTree::FScheduledTickHandle();
-		InstanceData.TreeInstanceData.GetMutableExecutionState()->ExecutionExtension.GetMutable<FStateTreeRunParallelStateTreeExecutionExtension>().ScheduledTickHandle = InstanceData.ScheduledTickHandle;
-	}
-
-	// Share event queue with parent tree.
-	if (FStateTreeInstanceData* OuterInstanceData = Context.GetMutableInstanceData())
-	{
-		InstanceData.TreeInstanceData.SetSharedEventQueue(OuterInstanceData->GetSharedMutableEventQueue());
-	}
-	
 	InstanceData.RunningStateTree = StateTreeToRun.GetStateTree();
 	FStateTreeExecutionContext ParallelTreeContext(Context, *InstanceData.RunningStateTree, InstanceData.TreeInstanceData);
 	if (!ParallelTreeContext.IsValid())
@@ -65,16 +52,21 @@ EStateTreeRunStatus FStateTreeRunParallelStateTreeTask::EnterState(FStateTreeExe
 		return EStateTreeRunStatus::Failed;
 	}
 
+	FStateTreeInstanceData* OuterInstanceData = Context.GetMutableInstanceData();
+	check(OuterInstanceData);
+
 	FStateTreeRunParallelStateTreeExecutionExtension Extension;
 	Extension.WeakExecutionContext = Context.MakeWeakExecutionContext();
 	const EStateTreeRunStatus RunStatus = ParallelTreeContext.Start(FStateTreeExecutionContext::FStartParameters
 		{
 			.InitialGlobalParameters = StateTreeToRun.GetGlobalParameters(),
-			.ExecutionExtension = TInstancedStruct<FStateTreeRunParallelStateTreeExecutionExtension>::Make(MoveTemp(Extension))
+			.ExecutionExtension = TInstancedStruct<FStateTreeRunParallelStateTreeExecutionExtension>::Make(MoveTemp(Extension)),
+			.SharedEventQueue = OuterInstanceData->GetSharedMutableEventQueue().ToSharedPtr()
 		});
 
 	if (RunStatus == EStateTreeRunStatus::Running)
 	{
+		check(!InstanceData.ScheduledTickHandle.IsValid());
 		InstanceData.ScheduledTickHandle = Context.AddScheduledTickRequest(ParallelTreeContext.GetNextScheduledTick());
 		InstanceData.TreeInstanceData.GetMutableExecutionState()->ExecutionExtension.GetMutable<FStateTreeRunParallelStateTreeExecutionExtension>().ScheduledTickHandle = InstanceData.ScheduledTickHandle;
 	}
@@ -130,6 +122,15 @@ void FStateTreeRunParallelStateTreeTask::TriggerTransitions(FStateTreeExecutionC
 void FStateTreeRunParallelStateTreeTask::ExitState(FStateTreeExecutionContext& Context, const FStateTreeTransitionResult& Transition) const
 {
 	FInstanceDataType& InstanceData = Context.GetInstanceData(*this);
+
+	// Reset Tick Handle regardless of whether the tree is still valid
+	if (InstanceData.ScheduledTickHandle.IsValid())
+	{
+		Context.RemoveScheduledTickRequest(InstanceData.ScheduledTickHandle);
+		InstanceData.ScheduledTickHandle = UE::StateTree::FScheduledTickHandle();
+		InstanceData.TreeInstanceData.GetMutableExecutionState()->ExecutionExtension.GetMutable<FStateTreeRunParallelStateTreeExecutionExtension>().ScheduledTickHandle = InstanceData.ScheduledTickHandle;
+	}
+
 	if (!InstanceData.RunningStateTree)
 	{
 		return;
@@ -142,7 +143,6 @@ void FStateTreeRunParallelStateTreeTask::ExitState(FStateTreeExecutionContext& C
 	}
 
 	ParallelTreeContext.Stop();
-	Context.RemoveScheduledTickRequest(InstanceData.ScheduledTickHandle);
 }
 
 const FStateTreeReference& FStateTreeRunParallelStateTreeTask::GetStateTreeToRun(FStateTreeExecutionContext& Context, FInstanceDataType& InstanceData) const
