[State Tree] added missing EditorNode Descriptor for property function node's node template

#Synced-CL 48638146

diff --git a/Engine/Plugins/Runtime/StateTree/Source/StateTreeEditorModule/Private/StateTreeEditorData.cpp b/Engine/Plugins/Runtime/StateTree/Source/StateTreeEditorModule/Private/StateTreeEditorData.cpp
--- a/Engine/Plugins/Runtime/StateTree/Source/StateTreeEditorModule/Private/StateTreeEditorData.cpp
+++ b/Engine/Plugins/Runtime/StateTree/Source/StateTreeEditorModule/Private/StateTreeEditorData.cpp
@@ -1538,10 +1538,10 @@ EStateTreeVisitor UStateTreeEditorData::VisitStructBoundPropertyFunctions(FGuid
 				return EStateTreeVisitor::Break;
 			}
 
-			FStateTreeBindableStructDesc Desc;
-			Desc.Struct = FunctionNode->GetInstance().GetStruct();
 			if (const UStruct* NodeStruct = FunctionNode->Node.GetScriptStruct())
 			{
+				FStateTreeBindableStructDesc Desc;
+				Desc.Struct = FunctionNode->GetInstance().GetStruct();
 				Desc.ID = FunctionNode->ID;
 				Desc.DataSource = EStateTreeBindableStructSource::PropertyFunction;
 				Desc.Name = FName(NodeStruct->GetDisplayNameText().ToString());
@@ -1551,6 +1551,18 @@ EStateTreeVisitor UStateTreeEditorData::VisitStructBoundPropertyFunctions(FGuid
 				{
 					return EStateTreeVisitor::Break;
 				}
+
+				FStateTreeBindableStructDesc NodeDesc;
+				NodeDesc.Struct = FunctionNode->GetNode().GetScriptStruct();
+				NodeDesc.ID = FunctionNode->GetNodeID();
+				NodeDesc.DataSource = EStateTreeBindableStructSource::PropertyFunction;
+				NodeDesc.Name = FName(NodeStruct->GetDisplayNameText().ToString());
+				NodeDesc.StatePath = FString::Printf(TEXT("%s/%s"), *StatePath, *UE::StateTree::Editor::PropertyFunctionStateName);
+
+				if (InFunc(*FunctionNode, NodeDesc, FunctionNode->GetNode()) == EStateTreeVisitor::Break)
+				{
+					return EStateTreeVisitor::Break;
+				}
 			}
 		}
 	}
