[State Tree] fixed comments for GetInstanceData

#Synced-CL 48638146

diff --git a/Engine/Plugins/Runtime/StateTree/Source/StateTreeModule/Public/StateTreeExecutionContext.h b/Engine/Plugins/Runtime/StateTree/Source/StateTreeModule/Public/StateTreeExecutionContext.h
--- a/Engine/Plugins/Runtime/StateTree/Source/StateTreeModule/Public/StateTreeExecutionContext.h
+++ b/Engine/Plugins/Runtime/StateTree/Source/StateTreeModule/Public/StateTreeExecutionContext.h
@@ -369,13 +369,13 @@ public:
 	};
 	UE_API void SetExternalGlobalParameters(const FExternalGlobalParameters* Parameters);
 
-	/** @return const references to the instance data in use, or nullptr if the context is not valid. */
+	/** @return const references to the instance data in use */
 	const FStateTreeInstanceData* GetInstanceData() const
 	{
 		return &InstanceData;
 	}
 
-	/** @return mutable references to the instance data in use, or nullptr if the context is not valid. */
+	/** @return mutable references to the instance data in use */
 	FStateTreeInstanceData* GetMutableInstanceData() const
 	{
 		return &InstanceData;
