[PropertyBindingUtils][State Tree] BindingView rework

- Added error icon for Missing struct
- Added error icon for binding path failed to resolve
- Click on Binding struct will traverse to the corresponding struct property row
- Click on Binding path will traverse to the corresponding property property row

#Synced-CL 48820855

diff --git a/Engine/Plugins/Runtime/PropertyBindingUtils/Source/PropertyBindingUtilsEditor/Private/Widgets/SPropertyBindingView.cpp b/Engine/Plugins/Runtime/PropertyBindingUtils/Source/PropertyBindingUtilsEditor/Private/Widgets/SPropertyBindingView.cpp
--- a/Engine/Plugins/Runtime/PropertyBindingUtils/Source/PropertyBindingUtilsEditor/Private/Widgets/SPropertyBindingView.cpp
+++ b/Engine/Plugins/Runtime/PropertyBindingUtils/Source/PropertyBindingUtilsEditor/Private/Widgets/SPropertyBindingView.cpp
@@ -3,6 +3,8 @@
 #include "Widgets/SPropertyBindingView.h"
 #include "PropertyBindingBinding.h"
 #include "PropertyBindingBindingCollection.h"
+#include "PropertyBindingDataView.h"
+#include "Widgets/Input/SButton.h"
 
 #define LOCTEXT_NAMESPACE "SPropertyBindingViewer"
 
@@ -15,80 +17,159 @@ namespace Private
 	static const FLazyName ColumnId_TargetStruct = "TargetStruct";
 	static const FLazyName ColumnId_TargetPath = "TargetPath";
 
+	// @todo: We don't check Property Compatibility here as the code is inside FCachedBindingData. We should reuse everything from there(get name, validity).
+	// @todo: The compiler may do extra stuff that we cannot implement here. The editor data should have a TArray of invalid binding discovered during compilation.
 	class SBindingViewRow : public SMultiColumnTableRow<TSharedPtr<SBindingView::FItem>>
 	{
 	public:
 		SLATE_BEGIN_ARGS(SBindingViewRow) {}
 		SLATE_END_ARGS()
 
-		void Construct(const FArguments& Args, TSharedPtr<SBindingView::FItem> InItem, TWeakInterfacePtr<const IPropertyBindingBindingCollectionOwner> InCollectionOwner, const TSharedRef<STableViewBase>& OwnerTableView)
+		void Construct(const FArguments& Args, TSharedPtr<SBindingView::FItem> InItem, TWeakInterfacePtr<const IPropertyBindingBindingCollectionOwner> InCollectionOwner, const TSharedRef<STableViewBase>& InOwnerTableView, const TSharedRef<const SBindingView>& InOwnerBindingView)
 		{
 			Item = InItem;
 			CollectionOwner = InCollectionOwner;
+			WeakOwnerBindingView = InOwnerBindingView;
 			SMultiColumnTableRow<TSharedPtr<SBindingView::FItem>>::Construct(
 				FSuperRowType::FArguments()
 				.Padding(1.0f),
-				OwnerTableView
+				InOwnerTableView
 			);
 		}
 
-		FText HandleGetStructName(FGuid StructID) const
+		const FSlateBrush* GetMalformedBindingImage() const
 		{
-			if (const IPropertyBindingBindingCollectionOwner* CollectionOwnerPtr = CollectionOwner.Get())
-			{
-				TInstancedStruct<FPropertyBindingBindableStructDescriptor> StructDesc;
-				if (CollectionOwnerPtr->GetBindableStructByID(StructID, StructDesc))
-				{
-					return FText::FromString(StructDesc.Get().ToString());
-				}
-			}
-			return FText::GetEmpty();
+			return FCoreStyle::Get().GetBrush("Icons.ErrorWithColor");
+		}
+
+		FText GetMalformedBindingSectionTooltip(bool bIsStruct) const
+		{
+			return bIsStruct ? LOCTEXT("MissingStructTooltip", "Struct is missing!") : LOCTEXT("MalformedBindingPathTooltip", "Could not find a valid property from the binding path.");
+		}
+
+		EVisibility GetMalformedBindingImageVisibility(bool bIsStruct, bool bIsSource) const
+		{
+			return IsBindingSectionValid(bIsStruct, bIsSource) ? EVisibility::Hidden : EVisibility::Visible;
+		}
+
+		FText GetBindingSectionName(bool bIsStruct, bool bIsSource) const
+		{
+			return HandleGetBindingSectionName(bIsStruct, bIsSource).Key;
 		}
 
-		FText HandleGetBindingPath(bool bIsSource) const
+		bool IsBindingSectionValid(bool bIsStruct, bool bIsSource) const
 		{
-			const FPropertyBindingPath& Path = bIsSource ? Item->SourcePath : Item->TargetPath;
-			return FText::FromString(Path.ToString());
+			return HandleGetBindingSectionName(bIsStruct, bIsSource).Value;
 		}
 
 		virtual TSharedRef<SWidget> GenerateWidgetForColumn(const FName& ColumnName) override
 		{
-			if (ColumnName == ColumnId_SourceStruct)
-			{
-				const FGuid StructID = Item->SourcePath.GetStructID();
-				return SNew(STextBlock)
-					.Text(this, &SBindingViewRow::HandleGetStructName, StructID);
-			}
-			else if (ColumnName == ColumnId_SourcePath)
-			{
-				return SNew(STextBlock)
-					.Text(this, &SBindingViewRow::HandleGetBindingPath, true);
-			}
-			else if (ColumnName == ColumnId_TargetStruct)
-			{
-				const FGuid StructID = Item->TargetPath.GetStructID();
-				return SNew(STextBlock)
-					.Text(this, &SBindingViewRow::HandleGetStructName, StructID);
-			}
-			else if (ColumnName == ColumnId_TargetPath)
+			const bool bIsStruct = ColumnName == ColumnId_SourceStruct || ColumnName == ColumnId_TargetStruct;
+			const bool bIsSource = ColumnName == ColumnId_SourceStruct || ColumnName == ColumnId_SourcePath;
+			const FPropertyBindingPath& BindingPath = bIsSource ? Item->SourcePath : Item->TargetPath;
+			const FGuid StructID = bIsSource ? Item->SourcePath.GetStructID() : Item->TargetPath.GetStructID();
+
+			TSharedRef<SWidget> InnerWidget = SNew(SHorizontalBox)
+				+ SHorizontalBox::Slot()
+				.AutoWidth()
+				.VAlign(VAlign_Center)
+				[
+					SNew(SBox)
+						.HeightOverride(16.0f)
+						[
+							SNew(SImage)
+								.Image(this, &SBindingViewRow::GetMalformedBindingImage)
+								.ColorAndOpacity(FLinearColor::White)
+								.Visibility(this, &SBindingViewRow::GetMalformedBindingImageVisibility, bIsStruct, bIsSource)
+								.ToolTipText(this, &SBindingViewRow::GetMalformedBindingSectionTooltip, bIsStruct)
+						]
+				]
+			+ SHorizontalBox::Slot()
+				.AutoWidth()
+				.VAlign(VAlign_Center)
+				[
+					SNew(STextBlock)
+						.Text(this, &SBindingViewRow::GetBindingSectionName, bIsStruct, bIsSource)
+				];
+
+			if (TSharedPtr<const SBindingView> OwnerBindingView = WeakOwnerBindingView.Pin())
 			{
-				return SNew(STextBlock)
-					.Text(this, &SBindingViewRow::HandleGetBindingPath, false);
+				if (OwnerBindingView->OnBindingClicked.IsBound())
+				{
+					return SNew(SButton)
+						.ButtonStyle(FAppStyle::Get(), "SimpleButton")
+						.ContentPadding(FMargin(0, 0))
+						.OnClicked_Lambda([&BindingPath, bIsStruct, WeakBindingView = WeakOwnerBindingView]()
+							{
+								if (TSharedPtr<const SBindingView> BindingView = WeakBindingView.Pin())
+								{
+									BindingView->OnBindingClicked.Execute(bIsStruct, BindingPath);
+								}
+
+								return FReply::Handled();
+							})
+						[
+							InnerWidget
+						];
+				}
+
+				return InnerWidget;
 			}
 
 			return SNullWidget::NullWidget;
 		}
 
 	private:
+		TPair<FText, bool> HandleGetBindingSectionName(bool bIsStruct, bool bIsSource) const
+		{
+			if (const IPropertyBindingBindingCollectionOwner* CollectionOwnerPtr = CollectionOwner.Get())
+			{
+				FPropertyBindingPath& BindingPath = bIsSource ? Item->SourcePath : Item->TargetPath;
+
+				if (bIsStruct)
+				{
+					TInstancedStruct<FPropertyBindingBindableStructDescriptor> StructDesc;
+					if (CollectionOwnerPtr->GetBindableStructByID(BindingPath.GetStructID(), StructDesc))
+					{
+						check(StructDesc.IsValid());
+
+						return { FText::FromString(StructDesc.Get().ToString()), true };
+					}
+
+					return { LOCTEXT("MissingStruct", "???"), false };
+				}
+
+				bool bIsBindingPathValid = false;
+				FPropertyBindingDataView InstanceView;
+				if (CollectionOwnerPtr->GetBindingDataViewByID(BindingPath.GetStructID(), InstanceView))
+				{
+					TArray<FPropertyBindingPathIndirection> Indirections;
+					constexpr FString* Error = nullptr;
+					constexpr bool bHandleRedirects = true;
+
+					if (BindingPath.ResolveIndirectionsWithValue(InstanceView, Indirections, Error, bHandleRedirects))
+					{
+						bIsBindingPathValid = true;
+					}
+				}
+
+				return { FText::FromString(BindingPath.ToString()), bIsBindingPathValid };
+			}
+
+			return { FText::GetEmpty(), false };
+		}
+
 		TSharedPtr<SBindingView::FItem> Item;
 		TWeakInterfacePtr<const IPropertyBindingBindingCollectionOwner> CollectionOwner;
+
+		TWeakPtr<const SBindingView> WeakOwnerBindingView;
 	};
 }
 
 void SBindingView::Construct(const FArguments& InArgs)
 {
-	OnGetBindingCollection = InArgs._GetBindingCollection;
 	CollectionOwner = InArgs._CollectionOwner;
+	OnBindingClicked = InArgs._OnBindingClicked;
 
 	ListView = SNew(SListView<TSharedPtr<FItem>>)
 		.SelectionMode(ESelectionMode::None)
@@ -116,11 +197,6 @@ void SBindingView::Construct(const FArguments& InArgs)
 	];
 }
 
-const FPropertyBindingBindingCollection* SBindingView::GetBindingCollection()
-{
-	return OnGetBindingCollection.IsBound() ? OnGetBindingCollection.Execute() : nullptr;
-}
-
 void SBindingView::RequestRefresh()
 {
 	ListView->RequestListRefresh();
@@ -130,59 +206,61 @@ void SBindingView::Tick(const FGeometry& AllottedGeometry, const double InCurren
 {
 	SWidget::Tick(AllottedGeometry, InCurrentTime, InDeltaTime);
 
-	const FPropertyBindingBindingCollection* Collection = GetBindingCollection();
-	if (Collection)
+	if (const IPropertyBindingBindingCollectionOwner* CollectionOwnerPtr = CollectionOwner.Get())
 	{
-		int32 Index = 0;
-		bool bRequestRefresh = false;
-		Collection->ForEachBinding(
-			[&Index, &bRequestRefresh, Self=this](const FPropertyBindingBinding & Binding)
-			{
-				if (Self->Values.Num() <= Index)
-				{
-					Self->Values.Add(MakeShared<FItem>(Binding.GetSourcePath(), Binding.GetTargetPath(), Binding.GetPropertyFunctionNode().GetScriptStruct()));
-					bRequestRefresh = true;
-				}
-				else
+		if (const FPropertyBindingBindingCollection* BindingCollection = CollectionOwnerPtr->GetEditorPropertyBindings())
+		{
+			int32 Index = 0;
+			bool bRequestRefresh = false;
+			BindingCollection->ForEachBinding(
+				[&Index, &bRequestRefresh, Self = this](const FPropertyBindingBinding& Binding)
 				{
-					const TSharedPtr<FItem>& Item = Self->Values[Index];
-					const bool bChanged = Item->FunctionNodeStruct.Get() != Binding.GetPropertyFunctionNode().GetScriptStruct()
-						|| Item->SourcePath != Binding.GetSourcePath()
-						|| Item->TargetPath != Binding.GetTargetPath();
-					if (bChanged)
+					if (Self->Values.Num() <= Index)
 					{
-						Self->Values.SetNum(Index);
 						Self->Values.Add(MakeShared<FItem>(Binding.GetSourcePath(), Binding.GetTargetPath(), Binding.GetPropertyFunctionNode().GetScriptStruct()));
 						bRequestRefresh = true;
 					}
-				}
-				++Index;
-			});
+					else
+					{
+						const TSharedPtr<FItem>& Item = Self->Values[Index];
+						const bool bChanged = Item->FunctionNodeStruct.Get() != Binding.GetPropertyFunctionNode().GetScriptStruct()
+							|| Item->SourcePath != Binding.GetSourcePath()
+							|| Item->TargetPath != Binding.GetTargetPath();
+						if (bChanged)
+						{
+							Self->Values.SetNum(Index);
+							Self->Values.Add(MakeShared<FItem>(Binding.GetSourcePath(), Binding.GetTargetPath(), Binding.GetPropertyFunctionNode().GetScriptStruct()));
+							bRequestRefresh = true;
+						}
+					}
+					++Index;
+				});
 
-		if (Values.Num() != Index)
-		{
-			Values.SetNum(Index);
-			bRequestRefresh = true;
-		}
+			if (Values.Num() != Index)
+			{
+				Values.SetNum(Index);
+				bRequestRefresh = true;
+			}
 
-		if (bRequestRefresh)
-		{
-			RequestRefresh();
+			if (bRequestRefresh)
+			{
+				RequestRefresh();
+			}
+
+			return;
 		}
 	}
-	else
+
+	if (!Values.IsEmpty())
 	{
-		if (!Values.IsEmpty())
-		{
-			Values.Empty();
-			RequestRefresh();
-		}
+		Values.Empty();
+		RequestRefresh();
 	}
 }
 
 TSharedRef<ITableRow> SBindingView::HandleGenerateRow(TSharedPtr<FItem> Value, const TSharedRef<STableViewBase>& OwnerTable)
 {
-	return SNew(Private::SBindingViewRow, Value, CollectionOwner, OwnerTable);
+	return SNew(Private::SBindingViewRow, Value, CollectionOwner, OwnerTable, SharedThis(this));
 }
 
 } // namespace UE::PropertyBinding
diff --git a/Engine/Plugins/Runtime/PropertyBindingUtils/Source/PropertyBindingUtilsEditor/Public/Widgets/SPropertyBindingView.h b/Engine/Plugins/Runtime/PropertyBindingUtils/Source/PropertyBindingUtilsEditor/Public/Widgets/SPropertyBindingView.h
--- a/Engine/Plugins/Runtime/PropertyBindingUtils/Source/PropertyBindingUtilsEditor/Public/Widgets/SPropertyBindingView.h
+++ b/Engine/Plugins/Runtime/PropertyBindingUtils/Source/PropertyBindingUtilsEditor/Public/Widgets/SPropertyBindingView.h
@@ -21,7 +21,7 @@ namespace UE::PropertyBinding
 	class SBindingView : public SCompoundWidget
 	{
 	private:
-		friend ::UE::PropertyBinding::Private::SBindingViewRow;
+		friend Private::SBindingViewRow;
 		struct FItem
 		{
 			FPropertyBindingPath SourcePath;
@@ -31,9 +31,13 @@ namespace UE::PropertyBinding
 
 	public:
 		DECLARE_DELEGATE_RetVal(const FPropertyBindingBindingCollection*, FOnGetBindingCollection);
+		DECLARE_DELEGATE_TwoParams(FOnBindingClicked, bool bInIsStruct, const FPropertyBindingPath& InBindingPath);
 
 		SLATE_BEGIN_ARGS(SBindingView){}
+			UE_DEPRECATED(5.8, "Use CollectionOwner->GetEditorPropertyBindings() instead")
 			SLATE_EVENT(FOnGetBindingCollection, GetBindingCollection)
+
+			SLATE_EVENT(FOnBindingClicked, OnBindingClicked)
 			SLATE_ARGUMENT(TScriptInterface<const IPropertyBindingBindingCollectionOwner>, CollectionOwner)
 		SLATE_END_ARGS()
 
@@ -41,13 +45,13 @@ namespace UE::PropertyBinding
 		PROPERTYBINDINGUTILSEDITOR_API virtual void Tick(const FGeometry& AllottedGeometry, const double InCurrentTime, const float InDeltaTime) override;
 
 	private:
-		const FPropertyBindingBindingCollection* GetBindingCollection();
 		void RequestRefresh();
 
 		TSharedRef<ITableRow> HandleGenerateRow(TSharedPtr<FItem> Value, const TSharedRef<STableViewBase>& OwnerTable);
 
 	private:
-		FOnGetBindingCollection OnGetBindingCollection;
+		FOnBindingClicked OnBindingClicked;
+
 		TWeakInterfacePtr<const IPropertyBindingBindingCollectionOwner> CollectionOwner;
 		TSharedPtr<SListView<TSharedPtr<FItem>>> ListView;
 		TArray<TSharedPtr<FItem>> Values;
diff --git a/Engine/Plugins/Runtime/StateTree/Source/StateTreeEditorModule/Private/StateTreeEditorMode.cpp b/Engine/Plugins/Runtime/StateTree/Source/StateTreeEditorModule/Private/StateTreeEditorMode.cpp
--- a/Engine/Plugins/Runtime/StateTree/Source/StateTreeEditorModule/Private/StateTreeEditorMode.cpp
+++ b/Engine/Plugins/Runtime/StateTree/Source/StateTreeEditorModule/Private/StateTreeEditorMode.cpp
@@ -222,6 +222,7 @@ void UStateTreeEditorMode::Exit()
 				ViewModel->GetOnStateNodesChanged().RemoveAll(this);
 				ViewModel->GetOnSelectionChanged().RemoveAll(this);
 				ViewModel->GetOnBringNodeToFocus().RemoveAll(this);
+				ViewModel->GetOnBringBindingPathToFocus().RemoveAll(this);
 			}
 		}
 	}
@@ -258,6 +259,7 @@ void UStateTreeEditorMode::OnStateTreeChanged()
 					OldViewModel->GetOnStateNodesChanged().RemoveAll(this);
 					OldViewModel->GetOnSelectionChanged().RemoveAll(this);
 					OldViewModel->GetOnBringNodeToFocus().RemoveAll(this);
+					OldViewModel->GetOnBringBindingPathToFocus().RemoveAll(this);
 				}
 			}
 		}
@@ -285,6 +287,7 @@ void UStateTreeEditorMode::OnStateTreeChanged()
 					NewViewModel->GetOnStateNodesChanged().AddUObject(this, &UStateTreeEditorMode::HandleOnStateNodesChanged);
 					NewViewModel->GetOnSelectionChanged().AddUObject(this, &UStateTreeEditorMode::HandleModelSelectionChanged);
 					NewViewModel->GetOnBringNodeToFocus().AddUObject(this, &UStateTreeEditorMode::HandleModelBringNodeToFocus);
+					NewViewModel->GetOnBringBindingPathToFocus().AddUObject(this, &UStateTreeEditorMode::HandleModelBringBindingPathToFocus);
 				}
 			}
 		}
@@ -520,146 +523,433 @@ void UStateTreeEditorMode::HandleModelSelectionChanged(const TArray<TWeakObjectP
 
 void UStateTreeEditorMode::HandleModelBringNodeToFocus(const UStateTreeState* State, const FGuid NodeID) const
 {
-	auto BringToFocus = [Self = this](const FPropertyPath& HighlightPath, TSharedPtr<IDetailsView>& DetailsView)
+	const FPropertyBindingPath& BindingPathToNode = FPropertyBindingPath(NodeID);
+	return HandleModelBringBindingPathToFocus(State, BindingPathToNode);
+}
+
+void UStateTreeEditorMode::HandleModelBringBindingPathToFocus(const UStateTreeState* State, const FPropertyBindingPath& InBindingPath) const
+{
+	auto BringToFocus = [Self = this](const FPropertyPath& HighlightPath, TSharedPtr<IDetailsView> DetailsView)
 		{
-			if (HighlightPath.IsValid())
+			if (Self->HighlightTimerHandle.IsValid())
 			{
-				constexpr bool bExpandProperty = true;
-				DetailsView->ScrollPropertyIntoView(HighlightPath, bExpandProperty);
-				DetailsView->HighlightProperty(HighlightPath);
-
-				constexpr bool bLoop = false;
-				GEditor->GetTimerManager()->SetTimer(
-					Self->HighlightTimerHandle,
-					FTimerDelegate::CreateLambda([WeakSelectionDetailsView = DetailsView.ToWeakPtr()]()
-						{
-							if (TSharedPtr<IDetailsView> SelectionDetailsView = WeakSelectionDetailsView.Pin())
-							{
-								SelectionDetailsView->HighlightProperty({});
-							}
-						}),
-					1.0f,
-					bLoop);
-			}
-			else if (Self->HighlightTimerHandle.IsValid())
-			{
-				// NB. SetTimer also clear the timer.
+				if (TSharedPtr<IDetailsView> AssetDetailsView = Self->GetAssetDetailsView())
+				{
+					AssetDetailsView->HighlightProperty({});
+				}
+
+				if (TSharedPtr<IDetailsView> StateDetailsView = Self->GetDetailsView())
+				{
+					StateDetailsView->HighlightProperty({});
+				}
+
 				GEditor->GetTimerManager()->ClearTimer(Self->HighlightTimerHandle);
 			}
+
+			constexpr bool bExpandProperty = true;
+			DetailsView->ScrollPropertyIntoView(HighlightPath, bExpandProperty);
+			DetailsView->HighlightProperty(HighlightPath);
+
+			constexpr bool bLoop = false;
+			GEditor->GetTimerManager()->SetTimer(
+				Self->HighlightTimerHandle,
+				FTimerDelegate::CreateLambda([WeakEditorMode = TWeakObjectPtr<const UStateTreeEditorMode>(Self), WeakSelectionDetailsView = DetailsView.ToWeakPtr()]()
+					{
+						if (TSharedPtr<IDetailsView> SelectionDetailsView = WeakSelectionDetailsView.Pin())
+						{
+							SelectionDetailsView->HighlightProperty({});
+						}
+
+						if (const UStateTreeEditorMode* EditorModePtr = WeakEditorMode.Get())
+						{
+							EditorModePtr->HighlightTimerHandle.Invalidate();
+						}
+					}),
+				1.0f,
+				bLoop);
 		};
 
-	if (State)
-	{
-		TSharedPtr<IDetailsView> DetailsView = GetDetailsView();
-		if (DetailsView == nullptr)
+	auto MakeHighlightPathFromIndirections = [](FStateTreeDataView InValue, const FPropertyBindingPath& InPropertyBindingPath, FPropertyPath& OutHighlightPath, int32 InStartIndex = 0)
 		{
-			return;
-		}
+			// We want it to be in sync with what we show
+			constexpr bool bHandleRedirects = false;
+			TArray<FPropertyBindingPathIndirection> Indirections;
+			if (InPropertyBindingPath.ResolveIndirectionsWithValue(InValue, Indirections, nullptr, bHandleRedirects))
+			{
+				for (int32 Idx = InStartIndex; Idx < Indirections.Num(); ++Idx)
+				{
+					const FPropertyBindingPathIndirection& Indirection = Indirections[Idx];
+					OutHighlightPath.AddProperty(FPropertyInfo(const_cast<FProperty*>(Indirection.GetProperty()), Indirection.GetArrayIndex()));
+				}
+			};
+		};
+
+	auto MakeHighlightPathFromEvent = [&MakeHighlightPathFromIndirections](FProperty* InEventProperty, const FStateTreeBindableStructDesc& InDesc, FStateTreeDataView InValue, const FPropertyBindingPath& InPropertyBindingPath, FPropertyPath& OutHighlightPath)
+		{
+			// @todo: see UE-356039. We can only track up to RequiredEvent level because StateTree Event populates "fake" property row not associated with any PropertyNode.
+			// Once the customization is moved to BindingExtension, StateTree Event should just use InstancedStruct to populate rows, and we will be able to track property paths.
 
-		FPropertyPath HighlightPath;
+			check(InDesc.DataSource == EStateTreeBindableStructSource::StateEvent || InDesc.DataSource == EStateTreeBindableStructSource::TransitionEvent);
 
-		if (!HighlightPath.IsValid())
+			OutHighlightPath.AddProperty(FPropertyInfo(InEventProperty));
+			OutHighlightPath.AddProperty(FPropertyInfo(FStateTreeEventDesc::StaticStruct()->FindPropertyByName(GET_MEMBER_NAME_CHECKED(FStateTreeEventDesc, PayloadStruct))));
+
+			// Skip InstanceStruct
+			constexpr int32 IndirectionStartIndex = 1;
+			MakeHighlightPathFromIndirections(InValue, InPropertyBindingPath, OutHighlightPath, IndirectionStartIndex);
+		};
+
+	auto MakeHighlightPathFromTransition = [&MakeHighlightPathFromEvent, &MakeHighlightPathFromIndirections](TNotNull<const UStateTreeState*> InState, const FStateTreeBindableStructDesc& InDesc, FStateTreeDataView InValue, const FPropertyBindingPath& InPropertyBindingPath, FPropertyPath& OutHighlightPath)
 		{
-			FArrayProperty* TasksProperty = CastFieldChecked<FArrayProperty>(UStateTreeState::StaticClass()->FindPropertyByName(GET_MEMBER_NAME_CHECKED(UStateTreeState, Tasks)));
-			const int32 TaskIndex = State->Tasks.IndexOfByPredicate([&NodeID](const FStateTreeEditorNode& Node)
-			{
-				return Node.ID == NodeID;
-			});
-			if (TaskIndex != INDEX_NONE)
+			check(InDesc.DataSource == EStateTreeBindableStructSource::TransitionEvent || InDesc.DataSource == EStateTreeBindableStructSource::Transition);
+
+			const FGuid& StructID = InDesc.ID;
+
+			if (FArrayProperty* TransitionArrayProperty = CastFieldChecked<FArrayProperty>(UStateTreeState::StaticClass()->FindPropertyByName(GET_MEMBER_NAME_CHECKED(UStateTreeState, Transitions))))
 			{
-				HighlightPath.AddProperty(FPropertyInfo(TasksProperty));
-				HighlightPath.AddProperty(FPropertyInfo(TasksProperty->Inner, TaskIndex));
+				const int32 Index = InState->Transitions.IndexOfByPredicate([&StructID, &InDesc](const FStateTreeTransition& Element)
+					{
+						if (InDesc.DataSource == EStateTreeBindableStructSource::TransitionEvent)
+						{
+							return Element.GetEventID() == StructID;
+						}
+
+						return Element.ID == StructID;
+					});
+
+				if (Index != INDEX_NONE)
+				{
+					OutHighlightPath.AddProperty(FPropertyInfo(TransitionArrayProperty));
+					OutHighlightPath.AddProperty(FPropertyInfo(TransitionArrayProperty->Inner, Index));
+
+					if (InDesc.DataSource == EStateTreeBindableStructSource::TransitionEvent)
+					{
+						FProperty* RequiredEventProperty = FStateTreeTransition::StaticStruct()->FindPropertyByName(GET_MEMBER_NAME_CHECKED(FStateTreeTransition, RequiredEvent));
+						MakeHighlightPathFromEvent(RequiredEventProperty, InDesc, InValue, InPropertyBindingPath, OutHighlightPath);
+					}
+					else
+					{
+						MakeHighlightPathFromIndirections(InValue, InPropertyBindingPath, OutHighlightPath);
+					}
+				}
 			}
-		}
+		};
 
-		if (!HighlightPath.IsValid())
+	auto MakeHighlightPathFromNode = [&MakeHighlightPathFromIndirections]<typename NodeOwnerType>(
+		FName InNodePropertyName, 
+		const TConstArrayView<FStateTreeEditorNode> InNodes, 
+		const FStateTreeBindableStructDesc& InDesc,
+		FStateTreeDataView InValue,
+		const FPropertyBindingPath& InPropertyBindingPath,
+		FPropertyPath& OutHighlightPath) requires
+		(TIsDerivedFrom<NodeOwnerType, UStateTreeState>::Value 
+		|| TIsDerivedFrom<NodeOwnerType, UStateTreeEditorData>::Value 
+		|| TIsDerivedFrom<NodeOwnerType, FStateTreeTransition>::Value
+		|| TIsDerivedFrom<NodeOwnerType, FStateTreePropertyPathBinding>::Value)
 		{
-			FProperty* SingleTaskProperty = CastFieldChecked<FProperty>(UStateTreeState::StaticClass()->FindPropertyByName(GET_MEMBER_NAME_CHECKED(UStateTreeState, SingleTask)));
-			if (State->SingleTask.ID == NodeID)
+			const FGuid& StructID = InDesc.ID;
+
+			const int32 Index = InNodes.IndexOfByPredicate([&StructID](const FStateTreeEditorNode& Element)
+				{
+					// Binding could be static or dynamic
+					return Element.ID == StructID || Element.GetNodeID() == StructID;
+				});
+
+			if (Index != INDEX_NONE)
 			{
-				HighlightPath.AddProperty(FPropertyInfo(SingleTaskProperty));
+				// For Property Functions, we directly append the property function node instance to the bound property
+				if constexpr (!TIsDerivedFrom<NodeOwnerType, FStateTreePropertyPathBinding>::Value)
+				{
+					FProperty* NodeProperty;
+					if constexpr (TIsDerivedFrom<NodeOwnerType, UObject>::Value)
+					{
+						NodeProperty = NodeOwnerType::StaticClass()->FindPropertyByName(InNodePropertyName);
+					}
+					else
+					{
+						NodeProperty = NodeOwnerType::StaticStruct()->FindPropertyByName(InNodePropertyName);
+					}
+
+					OutHighlightPath.AddProperty(FPropertyInfo(NodeProperty));
+					if (FArrayProperty* ArrayProperty = CastField<FArrayProperty>(NodeProperty))
+					{
+						OutHighlightPath.AddProperty(FPropertyInfo(ArrayProperty->Inner, Index));
+					}
+				}
+
+				check(InDesc.Struct);
+
+				FName InnerPropertyName;
+				if (InDesc.Struct->IsChildOf<UObject>())
+				{
+					InnerPropertyName = GET_MEMBER_NAME_CHECKED(FStateTreeEditorNode, InstanceObject);
+				}
+				else if (InDesc.Struct->IsChildOf<FStateTreeNodeBase>())
+				{
+					InnerPropertyName = GET_MEMBER_NAME_CHECKED(FStateTreeEditorNode, Node);
+				}
+				else
+				{
+					InnerPropertyName = GET_MEMBER_NAME_CHECKED(FStateTreeEditorNode, Instance);
+				}
+
+				OutHighlightPath.AddProperty(FPropertyInfo(FStateTreeEditorNode::StaticStruct()->FindPropertyByName(InnerPropertyName)));
+
+				MakeHighlightPathFromIndirections(InValue, InPropertyBindingPath, OutHighlightPath);
 			}
-		}
+		};
 
-		if (!HighlightPath.IsValid())
+	auto MakeHighlightPathFromParameters = [&MakeHighlightPathFromIndirections](const FStateTreeBindableStructDesc & InDesc, FStateTreeDataView InValue, const FPropertyBindingPath& InPropertyBindingPath, FPropertyPath& OutHighlightPath)
 		{
-			FArrayProperty* TransitionsProperty = CastFieldChecked<FArrayProperty>(UStateTreeState::StaticClass()->FindPropertyByName(GET_MEMBER_NAME_CHECKED(UStateTreeState, Transitions)));
-			const int32 TransitionIndex = State->Transitions.IndexOfByPredicate([&NodeID](const FStateTreeTransition& Transition)
+			check(InDesc.DataSource == EStateTreeBindableStructSource::Parameter || InDesc.DataSource == EStateTreeBindableStructSource::StateParameter);
+
+			if (InDesc.DataSource == EStateTreeBindableStructSource::Parameter)
 			{
-				return Transition.ID == NodeID;
-			});
-			if (TransitionIndex != INDEX_NONE)
+				// Properties are private, cant go through GET_MEMBER_NAME
+				OutHighlightPath.AddProperty(FPropertyInfo(UStateTreeEditorData::StaticClass()->FindPropertyByName(TEXT("RootParameterPropertyBag"))));
+			}
+			else
 			{
-				HighlightPath.AddProperty(FPropertyInfo(TransitionsProperty));
-				HighlightPath.AddProperty(FPropertyInfo(TransitionsProperty->Inner, TransitionIndex));
+				OutHighlightPath.AddProperty(FPropertyInfo(UStateTreeState::StaticClass()->FindPropertyByName(GET_MEMBER_NAME_CHECKED(UStateTreeState, Parameters))));
+				OutHighlightPath.AddProperty(FPropertyInfo(FStateTreeStateParameters::StaticStruct()->FindPropertyByName(GET_MEMBER_NAME_CHECKED(FStateTreeStateParameters, Parameters))));
 			}
-		}
 
-		if (!HighlightPath.IsValid())
+			OutHighlightPath.AddProperty(FPropertyInfo(FInstancedPropertyBag::StaticStruct()->FindPropertyByName(TEXT("Value"))));
+
+			MakeHighlightPathFromIndirections(InValue, InPropertyBindingPath, OutHighlightPath);
+		};
+
+	auto MakeHighlightPath = [State, &MakeHighlightPathFromNode, &MakeHighlightPathFromTransition, &MakeHighlightPathFromEvent, &MakeHighlightPathFromParameters](
+		auto&& Self, 
+		const FStateTreeBindableStructDesc& InDesc, 
+		const FPropertyBindingPath& InPropertyBindingPath, 
+		const FStateTreeDataView InValue, 
+		TNotNull<const UStateTreeEditorData*> InTreeData, 
+		FPropertyPath& OutHighlightPath)
+	->void
 		{
-			FArrayProperty* EnterConditionsProperty = CastFieldChecked<FArrayProperty>(UStateTreeState::StaticClass()->FindPropertyByName(GET_MEMBER_NAME_CHECKED(UStateTreeState, EnterConditions)));
-			const int32 EnterConditionIndex = State->EnterConditions.IndexOfByPredicate([&NodeID](const FStateTreeEditorNode& Node)
+			if (InDesc.DataSource == EStateTreeBindableStructSource::PropertyFunction)
 			{
-				return Node.ID == NodeID;
-			});
-			if (EnterConditionIndex != INDEX_NONE)
+				if (const FStateTreeEditorPropertyBindings* Bindings = InTreeData->GetPropertyEditorBindings())
+				{
+					for (const FStateTreePropertyPathBinding& Binding : Bindings->GetBindings())
+					{
+						if (Binding.GetSourcePath().GetStructID() == InDesc.ID)
+						{
+							InTreeData->VisitAllNodes([&Binding, &Self, InTreeData, &OutHighlightPath](const UStateTreeState* State, const FStateTreeBindableStructDesc& InTargetStructDesc, const FStateTreeDataView InTargetStructValue)
+								{
+									if (Binding.GetTargetPath().GetStructID() == InTargetStructDesc.ID)
+									{
+										Self(Self, InTargetStructDesc, Binding.GetTargetPath(), InTargetStructValue, InTreeData, OutHighlightPath);
+										return EStateTreeVisitor::Break;
+									}
+
+									return EStateTreeVisitor::Continue;
+								});
+
+							constexpr FName EmptyName;
+							MakeHighlightPathFromNode.operator()<FStateTreePropertyPathBinding>(
+								EmptyName,
+								TConstArrayView<FStateTreeEditorNode>(Binding.GetPropertyFunctionNode().GetPtr<FStateTreeEditorNode>(), 1),
+								InDesc,
+								InValue,
+								InPropertyBindingPath,
+								OutHighlightPath
+								);
+
+							break;
+						}
+					}
+				}
+			}
+			else if (InDesc.DataSource == EStateTreeBindableStructSource::GlobalTask)
 			{
-				HighlightPath.AddProperty(FPropertyInfo(EnterConditionsProperty));
-				HighlightPath.AddProperty(FPropertyInfo(EnterConditionsProperty->Inner, EnterConditionIndex));
+				MakeHighlightPathFromNode.operator()<UStateTreeEditorData>(
+					GET_MEMBER_NAME_CHECKED(UStateTreeEditorData, GlobalTasks), 
+					InTreeData->GlobalTasks,
+					InDesc,
+					InValue,
+					InPropertyBindingPath,
+					OutHighlightPath
+					);
 			}
-		}
+			else if (InDesc.DataSource == EStateTreeBindableStructSource::Evaluator)
+			{
+				MakeHighlightPathFromNode.operator()<UStateTreeEditorData>(
+					GET_MEMBER_NAME_CHECKED(UStateTreeEditorData, Evaluators), 
+					InTreeData->Evaluators,
+					InDesc,
+					InValue,
+					InPropertyBindingPath,
+					OutHighlightPath);
+			}
+			else if (InDesc.DataSource == EStateTreeBindableStructSource::Parameter)
+			{
+				MakeHighlightPathFromParameters(InDesc, InValue, InPropertyBindingPath, OutHighlightPath);
+			}
+			else if (InDesc.DataSource == EStateTreeBindableStructSource::Context)
+			{
+				// Do nothing. We cannot highlight a category node and there is no property nodes associated with context
+			}
+			else if (State)
+			{
+				if (InDesc.DataSource == EStateTreeBindableStructSource::StateParameter)
+				{
+					MakeHighlightPathFromParameters(InDesc, InValue, InPropertyBindingPath, OutHighlightPath);
+				}
+				else if (InDesc.DataSource == EStateTreeBindableStructSource::Task)
+				{
+					MakeHighlightPathFromNode.operator()<UStateTreeState>(
+						GET_MEMBER_NAME_CHECKED(UStateTreeState, Tasks), 
+						State->Tasks,
+						InDesc,
+						InValue,
+						InPropertyBindingPath,
+						OutHighlightPath);
+					if (!OutHighlightPath.IsValid())
+					{
+						MakeHighlightPathFromNode.operator()<UStateTreeState>(
+							GET_MEMBER_NAME_CHECKED(UStateTreeState, SingleTask), 
+							TConstArrayView<FStateTreeEditorNode>(&State->SingleTask, 1),
+							InDesc,
+							InValue,
+							InPropertyBindingPath,
+							OutHighlightPath);
+					}
+				}
+				else if (InDesc.DataSource == EStateTreeBindableStructSource::Condition)
+				{
+					MakeHighlightPathFromNode.operator()<UStateTreeState>(
+						GET_MEMBER_NAME_CHECKED(UStateTreeState, EnterConditions), 
+						State->EnterConditions,
+						InDesc,
+						InValue,
+						InPropertyBindingPath,
+						OutHighlightPath);
+
+					if (!OutHighlightPath.IsValid())
+					{
+						FPropertyPath TempHighlightPath;
+						for (int32 Idx = 0; Idx < State->Transitions.Num(); ++Idx)
+						{
+							const FStateTreeTransition& Transition = State->Transitions[Idx];
+
+							MakeHighlightPathFromNode.operator()<FStateTreeTransition>(
+								GET_MEMBER_NAME_CHECKED(FStateTreeTransition, Conditions),
+								Transition.Conditions,
+								InDesc,
+								InValue,
+								InPropertyBindingPath,
+								TempHighlightPath);
+
+							if (TempHighlightPath.IsValid())
+							{
+								FArrayProperty* TransitionsProperty = CastFieldChecked<FArrayProperty>(UStateTreeState::StaticClass()->FindPropertyByName(GET_MEMBER_NAME_CHECKED(UStateTreeState, Transitions)));
+								check(TransitionsProperty);
 
-		BringToFocus(HighlightPath, DetailsView);
+								OutHighlightPath.AddProperty(FPropertyInfo(TransitionsProperty));
+								OutHighlightPath.AddProperty(FPropertyInfo(TransitionsProperty->Inner, Idx));
+
+								for (int32 PropertyInfoIdx = 0; PropertyInfoIdx < TempHighlightPath.GetNumProperties(); ++PropertyInfoIdx)
+								{
+									OutHighlightPath.AddProperty(TempHighlightPath.GetPropertyInfo(PropertyInfoIdx));
+								}
+
+								break;
+							}
+						}
+					}
+				}
+				else if (InDesc.DataSource == EStateTreeBindableStructSource::Consideration)
+				{
+					MakeHighlightPathFromNode.operator()<UStateTreeState>(
+						GET_MEMBER_NAME_CHECKED(UStateTreeState, Considerations), 
+						State->Considerations,
+						InDesc,
+						InValue,
+						InPropertyBindingPath,
+						OutHighlightPath);
+				}
+				else if (InDesc.DataSource == EStateTreeBindableStructSource::Transition)
+				{
+					MakeHighlightPathFromTransition(
+						State,
+						InDesc,
+						InValue,
+						InPropertyBindingPath,
+						OutHighlightPath);
+				}
+				else if (InDesc.DataSource == EStateTreeBindableStructSource::TransitionEvent)
+				{
+					MakeHighlightPathFromTransition(
+						State,
+						InDesc,
+						InValue,
+						InPropertyBindingPath,
+						OutHighlightPath);
+				}
+				else if (InDesc.DataSource == EStateTreeBindableStructSource::StateEvent)
+				{
+					MakeHighlightPathFromEvent(
+						UStateTreeState::StaticClass()->FindPropertyByName(GET_MEMBER_NAME_CHECKED(UStateTreeState, RequiredEventToEnter)),
+						InDesc,
+						InValue,
+						InPropertyBindingPath,
+						OutHighlightPath
+					);
+				}
+				
+			}
+		};
+
+	UStateTree* StateTree = GetStateTree();
+	if (StateTree == nullptr)
+	{
+		return;
 	}
-	else
+	const UStateTreeEditorData* TreeData = Cast<UStateTreeEditorData>(StateTree->EditorData);
+	if (TreeData == nullptr)
 	{
-		UStateTree* StateTree = GetStateTree();
-		if (StateTree == nullptr)
-		{
-			return;
-		}
-		const UStateTreeEditorData* TreeData = Cast<UStateTreeEditorData>(StateTree->EditorData);
-		if (TreeData == nullptr)
-		{
-			return;
-		}
-		TSharedPtr<IDetailsView> DetailsView = GetAssetDetailsView();
-		if (DetailsView == nullptr)
-		{
-			return;
-		}
+		return;
+	}
 
-		FPropertyPath HighlightPath;
-		if (!HighlightPath.IsValid())
-		{
-			FArrayProperty* TasksProperty = CastFieldChecked<FArrayProperty>(UStateTreeEditorData::StaticClass()->FindPropertyByName(GET_MEMBER_NAME_CHECKED(UStateTreeEditorData, GlobalTasks)));
-			const int32 TaskIndex = TreeData->GlobalTasks.IndexOfByPredicate([&NodeID](const FStateTreeEditorNode& Node)
-			{
-				return Node.ID == NodeID;
-			});
-			if (TaskIndex != INDEX_NONE)
-			{
-				HighlightPath.AddProperty(FPropertyInfo(TasksProperty));
-				HighlightPath.AddProperty(FPropertyInfo(TasksProperty->Inner, TaskIndex));
-			}
-		}
+	TSharedPtr<IDetailsView> DetailsView = State ? GetDetailsView() : GetAssetDetailsView();
+	if (!DetailsView.IsValid())
+	{
+		return;
+	}
 
-		if (!HighlightPath.IsValid())
-		{
-			FArrayProperty* EvaluatorsProperty = CastFieldChecked<FArrayProperty>(UStateTreeEditorData::StaticClass()->FindPropertyByName(GET_MEMBER_NAME_CHECKED(UStateTreeEditorData, Evaluators)));
-			const int32 TaskIndex = TreeData->Evaluators.IndexOfByPredicate([&NodeID](const FStateTreeEditorNode& Node)
+	const FGuid& StructID = InBindingPath.GetStructID();
+
+	FPropertyPath HighlightPath;
+
+	if (State)
+	{
+		TreeData->VisitStateNodes(*State, [&StructID, TreeData, &HighlightPath, &InBindingPath, &MakeHighlightPath](const UStateTreeState* State, const FStateTreeBindableStructDesc& Desc, const FStateTreeDataView Value)
 			{
-				return Node.ID == NodeID;
+				if (Desc.ID == StructID)
+				{
+					MakeHighlightPath(MakeHighlightPath, Desc, InBindingPath, Value, TreeData, HighlightPath);
+					return EStateTreeVisitor::Break;
+				}
+				return EStateTreeVisitor::Continue;
 			});
-			if (TaskIndex != INDEX_NONE)
+	}
+	else
+	{
+		TreeData->VisitGlobalNodes([&StructID, TreeData, &HighlightPath, &InBindingPath, &MakeHighlightPath](const UStateTreeState* State, const FStateTreeBindableStructDesc& Desc, const FStateTreeDataView Value)
 			{
-				HighlightPath.AddProperty(FPropertyInfo(EvaluatorsProperty));
-				HighlightPath.AddProperty(FPropertyInfo(EvaluatorsProperty->Inner, TaskIndex));
-			}
-		}
+				if (Desc.ID == StructID)
+				{
+					MakeHighlightPath(MakeHighlightPath, Desc, InBindingPath, Value, TreeData, HighlightPath);
+					return EStateTreeVisitor::Break;
+				}
 
-		BringToFocus(HighlightPath, DetailsView);
+				return EStateTreeVisitor::Continue;
+			});
 	}
+
+	BringToFocus(HighlightPath, DetailsView);
 }
 
 void UStateTreeEditorMode::UpdateAsset()
diff --git a/Engine/Plugins/Runtime/StateTree/Source/StateTreeEditorModule/Private/StateTreeEditorModeToolkit.cpp b/Engine/Plugins/Runtime/StateTree/Source/StateTreeEditorModule/Private/StateTreeEditorModeToolkit.cpp
--- a/Engine/Plugins/Runtime/StateTree/Source/StateTreeEditorModule/Private/StateTreeEditorModeToolkit.cpp
+++ b/Engine/Plugins/Runtime/StateTree/Source/StateTreeEditorModule/Private/StateTreeEditorModeToolkit.cpp
@@ -127,8 +127,24 @@ void FStateTreeEditorModeToolkit::HandleTabSpawned(UE::StateTreeEditor::FSpawned
 
 			DockTab->SetContent(
 				SNew(UE::PropertyBinding::SBindingView)
-				.GetBindingCollection(this, &FStateTreeEditorModeToolkit::GetBindingCollection)
 				.CollectionOwner(TScriptInterface<const IPropertyBindingBindingCollectionOwner>(EditorData))
+				.OnBindingClicked_Lambda([WeakEditorHost = EditorHost.ToWeakPtr()](bool bInIsStruct, const FPropertyBindingPath& InBindingPath)
+					{
+						if (const TSharedPtr<const IStateTreeEditorHost> EditorHost = WeakEditorHost.Pin())
+						{
+							if (UStateTree* StateTree = EditorHost->GetStateTree())
+							{
+								if (UStateTreeEditingSubsystem* StateTreeEditingSubsystem = GEditor->GetEditorSubsystem<UStateTreeEditingSubsystem>())
+								{
+									// Make an empty binding path only containing struct ID if it is a struct
+									const FPropertyBindingPath& BindingPath = bInIsStruct ? FPropertyBindingPath(InBindingPath.GetStructID()) : InBindingPath;
+
+									const TSharedRef<FStateTreeViewModel> ViewModel = StateTreeEditingSubsystem->FindOrAddViewModel(StateTree);
+									ViewModel->BringBindingPathToFocus(BindingPath);
+								}
+							}
+						}
+					})
 			);
 		}
 	}
diff --git a/Engine/Plugins/Runtime/StateTree/Source/StateTreeEditorModule/Private/StateTreeViewModel.cpp b/Engine/Plugins/Runtime/StateTree/Source/StateTreeEditorModule/Private/StateTreeViewModel.cpp
--- a/Engine/Plugins/Runtime/StateTree/Source/StateTreeEditorModule/Private/StateTreeViewModel.cpp
+++ b/Engine/Plugins/Runtime/StateTree/Source/StateTreeEditorModule/Private/StateTreeViewModel.cpp
@@ -1447,6 +1447,20 @@ void FStateTreeViewModel::BringNodeToFocus(UStateTreeState* State, const FGuid N
 	OnBringNodeToFocus.Broadcast(State, NodeID);
 }
 
+void FStateTreeViewModel::BringBindingPathToFocus(const FPropertyBindingPath& InBindingPath)
+{
+	if (UStateTreeEditorData* TreeData = TreeDataWeak.Get())
+	{
+		UStateTreeState* BindingPathOwnerState = const_cast<UStateTreeState*>(TreeData->GetStateByStructID(InBindingPath.GetStructID()));
+		if (BindingPathOwnerState)
+		{
+			SetSelection(BindingPathOwnerState);
+		}
+
+		OnBringBindingPathToFocus.Broadcast(BindingPathOwnerState, InBindingPath);
+	}
+}
+
 UStateTreeState* FStateTreeViewModel::GetViewState() const
 {
 	UStateTreeState* ViewState = GetLastSelectedState();
diff --git a/Engine/Plugins/Runtime/StateTree/Source/StateTreeEditorModule/Public/StateTreeEditorMode.h b/Engine/Plugins/Runtime/StateTree/Source/StateTreeEditorModule/Public/StateTreeEditorMode.h
--- a/Engine/Plugins/Runtime/StateTree/Source/StateTreeEditorModule/Public/StateTreeEditorMode.h
+++ b/Engine/Plugins/Runtime/StateTree/Source/StateTreeEditorModule/Public/StateTreeEditorMode.h
@@ -54,6 +54,7 @@ protected:
 	UE_API void HandleModelAssetChanged();
 	UE_API void HandleModelSelectionChanged(const TArray<TWeakObjectPtr<UStateTreeState>>& SelectedStates) const;
 	UE_API void HandleModelBringNodeToFocus(const UStateTreeState* State, const FGuid NodeID) const;
+	UE_API void HandleModelBringBindingPathToFocus(const UStateTreeState* State, const FPropertyBindingPath& InBindingPath) const;
 
 	void HandleStateAdded(UStateTreeState* , UStateTreeState*)
 	{
diff --git a/Engine/Plugins/Runtime/StateTree/Source/StateTreeEditorModule/Public/StateTreeViewModel.h b/Engine/Plugins/Runtime/StateTree/Source/StateTreeEditorModule/Public/StateTreeViewModel.h
--- a/Engine/Plugins/Runtime/StateTree/Source/StateTreeEditorModule/Public/StateTreeViewModel.h
+++ b/Engine/Plugins/Runtime/StateTree/Source/StateTreeEditorModule/Public/StateTreeViewModel.h
@@ -28,6 +28,7 @@ enum class EStateTreeNodeFormatting : uint8;
 enum class EStateTreeVisitor : uint8;
 
 struct FSlateBrush;
+struct FPropertyBindingPath;
 struct FPropertyChangedEvent;
 struct FStateTreeDebugger;
 struct FStateTreeDebuggerBreakpoint;
@@ -72,6 +73,7 @@ public:
 	DECLARE_MULTICAST_DELEGATE_OneParam(FOnStateNodesChanged, const UStateTreeState* /*AffectedState*/);
 	DECLARE_MULTICAST_DELEGATE_OneParam(FOnSelectionChanged, const TArray<TWeakObjectPtr<UStateTreeState>>& /*SelectedStates*/);
 	DECLARE_MULTICAST_DELEGATE_TwoParams(FOnBringNodeToFocus, const UStateTreeState* /*State*/, const FGuid /*NodeID*/);
+	DECLARE_MULTICAST_DELEGATE_TwoParams(FOnBringBindingPathToFocus, const UStateTreeState* /*State*/, const FPropertyBindingPath& /*BindingPath*/);
 
 	UE_API FStateTreeViewModel();
 	UE_API virtual ~FStateTreeViewModel() override;
@@ -94,6 +96,7 @@ public:
 	UE_API UStateTreeState* GetLastSelectedState() const;
 
 	UE_API void BringNodeToFocus(UStateTreeState* State, const FGuid NodeID);
+	UE_API void BringBindingPathToFocus(const FPropertyBindingPath& InBindingPath);
 
 	// The view state is the last selected state, or first root if that is invalid. Nullptr if neither of those exist.
 	UE_API UStateTreeState* GetViewState() const;
@@ -255,6 +258,11 @@ public:
 		return OnBringNodeToFocus;
 	}
 
+	FOnBringBindingPathToFocus& GetOnBringBindingPathToFocus()
+	{
+		return OnBringBindingPathToFocus;
+	}
+
 protected:
 	UE_API void GetExpandedStatesRecursive(UStateTreeState* State, TSet<TWeakObjectPtr<UStateTreeState>>& ExpandedStates);
 
@@ -288,6 +296,7 @@ protected:
 	FOnStateNodesChanged OnStateNodesChanged;
 	FOnSelectionChanged OnSelectionChanged;
 	FOnBringNodeToFocus OnBringNodeToFocus;
+	FOnBringBindingPathToFocus OnBringBindingPathToFocus;
 };
 
 /** Helper class to allow to copy bindings into clipboard. */
