Viewmodel Driven Transition Control State Centric View
- Add Transition View Model used to track transtion changes
- Make StateTreeState broadcast via TranstionVM when it's transitions change
- Also Add support for VM driven response to states added / removed / etc to state centric
- Fix UE::StateTree::PropertyHelpers::ModifyStateInPreAndPostEdit not propagating array index

#tests editor
#Synced-CL 48823676

diff --git a/Engine/Plugins/Runtime/StateTree/Source/StateTreeEditorModule/Private/StateCentricView/View/SMainStateTransitionExtensionNode.cpp b/Engine/Plugins/Runtime/StateTree/Source/StateTreeEditorModule/Private/StateCentricView/View/SMainStateTransitionExtensionNode.cpp
--- a/Engine/Plugins/Runtime/StateTree/Source/StateTreeEditorModule/Private/StateCentricView/View/SMainStateTransitionExtensionNode.cpp
+++ b/Engine/Plugins/Runtime/StateTree/Source/StateTreeEditorModule/Private/StateCentricView/View/SMainStateTransitionExtensionNode.cpp
@@ -35,6 +35,10 @@ SMainStateTransitionExtensionNode::~SMainStateTransitionExtensionNode()
 				StateCentricEditorData->PerStateTreeUserSettings.OnOutTransitionViewLOD_Changed.RemoveAll(this);
 			}
 		}
+
+		ViewModel->GetTransitionViewModel().GetOnTransitionsAdded().RemoveAll(this);
+		ViewModel->GetTransitionViewModel().GetOnTransitionsRemoved().RemoveAll(this);
+		ViewModel->GetTransitionViewModel().GetOnTransitionsMoved().RemoveAll(this);
 	}
 }
 
@@ -46,6 +50,9 @@ void SMainStateTransitionExtensionNode::Construct_InitializeNodePersistentPreLOD
 	// Remove any prior bindings in case of issues with lifetime
 	StateCentricEditorData->PerStateTreeUserSettings.OnInTransitionViewLOD_Changed.RemoveAll(this);
 	StateCentricEditorData->PerStateTreeUserSettings.OnOutTransitionViewLOD_Changed.RemoveAll(this);
+	ViewModel->GetTransitionViewModel().GetOnTransitionsAdded().RemoveAll(this);
+	ViewModel->GetTransitionViewModel().GetOnTransitionsRemoved().RemoveAll(this);
+	ViewModel->GetTransitionViewModel().GetOnTransitionsMoved().RemoveAll(this);
 
 	TOptional<EExtendableNodeLOD> TransitionViewUserSetLOD;
 	if (TransitionDirection == EStateTransitionDirection::In)
@@ -64,6 +71,10 @@ void SMainStateTransitionExtensionNode::Construct_InitializeNodePersistentPreLOD
 	{
 		CurrentLOD = TransitionViewUserSetLOD.GetValue();
 	}
+
+	ViewModel->GetTransitionViewModel().GetOnTransitionsAdded().AddSP(this, &SMainStateTransitionExtensionNode::HandleOnTransitionsNumChanged);
+	ViewModel->GetTransitionViewModel().GetOnTransitionsRemoved().AddSP(this, &SMainStateTransitionExtensionNode::HandleOnTransitionsNumChanged);
+	ViewModel->GetTransitionViewModel().GetOnTransitionsMoved().AddSP(this, &SMainStateTransitionExtensionNode::HandleOnTransitionsNumChanged);
 }
 
 void SMainStateTransitionExtensionNode::Construct_GenerateBuiltInExtensions(const EExtendableNodeLOD InLOD)
@@ -302,10 +313,7 @@ FReply SMainStateTransitionExtensionNode::HandleOnAddTransitionClicked()
 	if (ViewState.IsValid())
 	{
 		FStateTreeTransition DefaultTransition = {};
-		ViewState->AddTransition(DefaultTransition.Trigger, EStateTreeTransitionType::None);
-
-		// @TODO: Temp local LOD regen instead of listening to delegate
-		RegenerateWidgetForLOD(CurrentLOD);
+		ViewModel->AddTransition(ViewState.Get(), DefaultTransition.Trigger, EStateTreeTransitionType::None);
 	}
 	return FReply::Handled();
 }
@@ -341,6 +349,11 @@ void SMainStateTransitionExtensionNode::HandleOnTransitionViewLOD_SettingsChange
 	}
 }
 
+void SMainStateTransitionExtensionNode::HandleOnTransitionsNumChanged(const UStateTreeState* /*OwningState*/, const TSet<FGuid>& /*TransitionIDs*/)
+{
+	RegenerateWidgetForLOD(CurrentLOD);
+}
+
 FSlateColor SMainStateTransitionExtensionNode::GetAddTransitionColor() const
 {
 	return AddTransitionButton->IsHovered() ? FStyleColors::White : FStyleColors::Foreground;
diff --git a/Engine/Plugins/Runtime/StateTree/Source/StateTreeEditorModule/Private/StateCentricView/View/SMainStateTransitionExtensionNode.h b/Engine/Plugins/Runtime/StateTree/Source/StateTreeEditorModule/Private/StateCentricView/View/SMainStateTransitionExtensionNode.h
--- a/Engine/Plugins/Runtime/StateTree/Source/StateTreeEditorModule/Private/StateCentricView/View/SMainStateTransitionExtensionNode.h
+++ b/Engine/Plugins/Runtime/StateTree/Source/StateTreeEditorModule/Private/StateCentricView/View/SMainStateTransitionExtensionNode.h
@@ -48,6 +48,9 @@ private:
 	/** Callback for when LOD settings change */
 	void HandleOnTransitionViewLOD_SettingsChanged(TOptional<EExtendableNodeLOD> InLOD);
 
+	/** Callback for transtions added / removed */
+	void HandleOnTransitionsNumChanged(const UStateTreeState* /*OwningState*/, const TSet<FGuid>& /*TransitionIDs*/);
+
 	/** Minor UI / UX callbacks */
 	FSlateColor GetAddTransitionColor() const;
 	FSlateColor GetHideButtonColor() const;
diff --git a/Engine/Plugins/Runtime/StateTree/Source/StateTreeEditorModule/Private/StateCentricView/View/SParentStateTransitionExtensionNode.cpp b/Engine/Plugins/Runtime/StateTree/Source/StateTreeEditorModule/Private/StateCentricView/View/SParentStateTransitionExtensionNode.cpp
--- a/Engine/Plugins/Runtime/StateTree/Source/StateTreeEditorModule/Private/StateCentricView/View/SParentStateTransitionExtensionNode.cpp
+++ b/Engine/Plugins/Runtime/StateTree/Source/StateTreeEditorModule/Private/StateCentricView/View/SParentStateTransitionExtensionNode.cpp
@@ -36,6 +36,10 @@ SParentStateTransitionExtensionNode::~SParentStateTransitionExtensionNode()
 				StateCentricEditorData->PerStateTreeUserSettings.OnOutTransitionViewLOD_Changed.RemoveAll(this);
 			}
 		}
+
+		ViewModel->GetTransitionViewModel().GetOnTransitionsAdded().RemoveAll(this);
+		ViewModel->GetTransitionViewModel().GetOnTransitionsRemoved().RemoveAll(this);
+		ViewModel->GetTransitionViewModel().GetOnTransitionsMoved().RemoveAll(this);
 	}
 }
 
@@ -47,6 +51,9 @@ void SParentStateTransitionExtensionNode::Construct_InitializeNodePersistentPreL
 	// Remove any prior bindings in case of issues with lifetime
 	StateCentricEditorData->PerStateTreeUserSettings.OnInTransitionViewLOD_Changed.RemoveAll(this);
 	StateCentricEditorData->PerStateTreeUserSettings.OnOutTransitionViewLOD_Changed.RemoveAll(this);
+	ViewModel->GetTransitionViewModel().GetOnTransitionsAdded().RemoveAll(this);
+	ViewModel->GetTransitionViewModel().GetOnTransitionsRemoved().RemoveAll(this);
+	ViewModel->GetTransitionViewModel().GetOnTransitionsMoved().RemoveAll(this);
 
 	TOptional<EExtendableNodeLOD> TransitionViewUserSetLOD;
 	if (TransitionDirection == EStateTransitionDirection::In)
@@ -65,6 +72,10 @@ void SParentStateTransitionExtensionNode::Construct_InitializeNodePersistentPreL
 	{
 		CurrentLOD = TransitionViewUserSetLOD.GetValue();
 	}
+
+	ViewModel->GetTransitionViewModel().GetOnTransitionsAdded().AddSP(this, &SParentStateTransitionExtensionNode::HandleOnTransitionsNumChanged);
+	ViewModel->GetTransitionViewModel().GetOnTransitionsRemoved().AddSP(this, &SParentStateTransitionExtensionNode::HandleOnTransitionsNumChanged);
+	ViewModel->GetTransitionViewModel().GetOnTransitionsMoved().AddSP(this, &SParentStateTransitionExtensionNode::HandleOnTransitionsNumChanged);
 }
 
 void SParentStateTransitionExtensionNode::Construct_GenerateBuiltInExtensions(const EExtendableNodeLOD InLOD)
@@ -191,10 +202,7 @@ FReply SParentStateTransitionExtensionNode::HandleOnAddTransitionClicked()
 	if (ViewState.IsValid())
 	{
 		FStateTreeTransition DefaultTransition = {};
-		ViewState->AddTransition(DefaultTransition.Trigger, EStateTreeTransitionType::None);
-
-		// @TODO: Temp local LOD regen instead of listening to delegate
-		RegenerateWidgetForLOD(CurrentLOD);
+		ViewModel->AddTransition(ViewState.Get(), DefaultTransition.Trigger, EStateTreeTransitionType::None);
 	}
 	return FReply::Handled();
 }
@@ -227,6 +235,11 @@ void SParentStateTransitionExtensionNode::HandleOnTransitionViewLOD_SettingsChan
 	}
 }
 
+void SParentStateTransitionExtensionNode::HandleOnTransitionsNumChanged(const UStateTreeState* /*OwningState*/, const TSet<FGuid>& /*TransitionIDs*/)
+{
+	RegenerateWidgetForLOD(CurrentLOD);
+}
+
 FSlateColor SParentStateTransitionExtensionNode::GetAddTransitionColor() const
 {
 	return AddTransitionButton->IsHovered() ? FStyleColors::White : FStyleColors::Foreground;
diff --git a/Engine/Plugins/Runtime/StateTree/Source/StateTreeEditorModule/Private/StateCentricView/View/SParentStateTransitionExtensionNode.h b/Engine/Plugins/Runtime/StateTree/Source/StateTreeEditorModule/Private/StateCentricView/View/SParentStateTransitionExtensionNode.h
--- a/Engine/Plugins/Runtime/StateTree/Source/StateTreeEditorModule/Private/StateCentricView/View/SParentStateTransitionExtensionNode.h
+++ b/Engine/Plugins/Runtime/StateTree/Source/StateTreeEditorModule/Private/StateCentricView/View/SParentStateTransitionExtensionNode.h
@@ -47,6 +47,9 @@ private:
 	/** Callback for when LOD settings change */
 	void HandleOnTransitionViewLOD_SettingsChanged(TOptional<EExtendableNodeLOD> InLOD);
 
+	/** Callback for transtions added / removed */
+	void HandleOnTransitionsNumChanged(const UStateTreeState* /*OwningState*/, const TSet<FGuid>& /*TransitionIDs*/);
+
 	/** Minor UI / UX callbacks */
 	FSlateColor GetAddTransitionColor() const;
 
diff --git a/Engine/Plugins/Runtime/StateTree/Source/StateTreeEditorModule/Private/StateCentricView/View/SStateCentricViewMainStateNode.cpp b/Engine/Plugins/Runtime/StateTree/Source/StateTreeEditorModule/Private/StateCentricView/View/SStateCentricViewMainStateNode.cpp
--- a/Engine/Plugins/Runtime/StateTree/Source/StateTreeEditorModule/Private/StateCentricView/View/SStateCentricViewMainStateNode.cpp
+++ b/Engine/Plugins/Runtime/StateTree/Source/StateTreeEditorModule/Private/StateCentricView/View/SStateCentricViewMainStateNode.cpp
@@ -34,6 +34,9 @@ SStateCentricViewMainStateNode::~SStateCentricViewMainStateNode()
 	if (ViewModel)
 	{
 		ViewModel->GetOnSelectionChanged().RemoveAll(this);
+		ViewModel->GetOnStateAdded().RemoveAll(this);
+		ViewModel->GetOnStatesRemoved().RemoveAll(this);
+		ViewModel->GetOnStatesMoved().RemoveAll(this);
 
 		if (UStateTreeEditorData* EditorData = const_cast<UStateTreeEditorData*>(ViewModel->GetStateTreeEditorData()))
 		{
@@ -60,12 +63,18 @@ void SStateCentricViewMainStateNode::Construct_InitializeNodePersistentPreLOD()
 	// Remove any prior bindings in case of issues with lifetime
 	{
 		ViewModel->GetOnSelectionChanged().RemoveAll(this);
+		ViewModel->GetOnStateAdded().RemoveAll(this);
+		ViewModel->GetOnStatesRemoved().RemoveAll(this);
+		ViewModel->GetOnStatesMoved().RemoveAll(this);
 
 		StateCentricEditorData->PerStateTreeUserSettings.OnNumParentsShownChanged.RemoveAll(this);
 		StateCentricEditorData->PerStateTreeUserSettings.OnParentHireachyLOD_Changed.RemoveAll(this);
 	}
 
 	ViewModel->GetOnSelectionChanged().AddSP(this, &SStateCentricViewMainStateNode::HandleOnSelectionChanged);
+	ViewModel->GetOnStateAdded().AddSP(this, &SStateCentricViewMainStateNode::HandleOnStateAdded);
+	ViewModel->GetOnStatesRemoved().AddSP(this, &SStateCentricViewMainStateNode::HandleOnStateRemoved);
+	ViewModel->GetOnStatesMoved().AddSP(this, &SStateCentricViewMainStateNode::HandleOnStateMoved);
 
 	StateCentricEditorData->PerStateTreeUserSettings.OnNumParentsShownChanged.AddSP(this, &SStateCentricViewMainStateNode::HandleOnNumParentsShownChanged);
 	StateCentricEditorData->PerStateTreeUserSettings.OnParentHireachyLOD_Changed.AddSP(this, &SStateCentricViewMainStateNode::HandleOnParentHireachyLOD_Changed);
@@ -299,6 +308,21 @@ void SStateCentricViewMainStateNode::HandleOnNumParentsShownChanged(TOptional<ui
 	RegenerateWidgetForLOD(CurrentLOD);
 }
 
+void SStateCentricViewMainStateNode::HandleOnStateAdded(UStateTreeState*, UStateTreeState*)
+{
+	RegenerateWidgetForLOD(CurrentLOD);
+}
+
+void SStateCentricViewMainStateNode::HandleOnStateRemoved(const TSet<UStateTreeState*>&)
+{
+	RegenerateWidgetForLOD(CurrentLOD);
+}
+
+void SStateCentricViewMainStateNode::HandleOnStateMoved(const TSet<UStateTreeState*>&, const TSet<UStateTreeState*>&)
+{
+	RegenerateWidgetForLOD(CurrentLOD);
+}
+
 } // namespace UE::StateTree::Editor::StateCentricView
 
 #undef LOCTEXT_NAMESPACE
diff --git a/Engine/Plugins/Runtime/StateTree/Source/StateTreeEditorModule/Private/StateCentricView/View/SStateCentricViewMainStateNode.h b/Engine/Plugins/Runtime/StateTree/Source/StateTreeEditorModule/Private/StateCentricView/View/SStateCentricViewMainStateNode.h
--- a/Engine/Plugins/Runtime/StateTree/Source/StateTreeEditorModule/Private/StateCentricView/View/SStateCentricViewMainStateNode.h
+++ b/Engine/Plugins/Runtime/StateTree/Source/StateTreeEditorModule/Private/StateCentricView/View/SStateCentricViewMainStateNode.h
@@ -49,6 +49,15 @@ protected:
 	/** Callback whenever num parents shown changes */
 	void HandleOnNumParentsShownChanged(TOptional<uint32> InNumParentShown);
 
+	/** Callback whenever states added */
+	void HandleOnStateAdded(UStateTreeState* /*ParentState*/, UStateTreeState* /*NewState*/);
+
+	/** Callback whenever states removed */
+	void HandleOnStateRemoved(const TSet<UStateTreeState*>& /*AffectedParents*/);
+
+	/** Callback whenever states moved */
+	void HandleOnStateMoved(const TSet<UStateTreeState*>& /*AffectedParents*/, const TSet<UStateTreeState*>& /*MovedStates*/);
+
 private:
 
 	TArray<TSharedPtr<FDeferredCleanupSlateBrush>> ChildStateBrushes;
diff --git a/Engine/Plugins/Runtime/StateTree/Source/StateTreeEditorModule/Private/StateCentricView/View/STransitionNode.cpp b/Engine/Plugins/Runtime/StateTree/Source/StateTreeEditorModule/Private/StateCentricView/View/STransitionNode.cpp
--- a/Engine/Plugins/Runtime/StateTree/Source/StateTreeEditorModule/Private/StateCentricView/View/STransitionNode.cpp
+++ b/Engine/Plugins/Runtime/StateTree/Source/StateTreeEditorModule/Private/StateCentricView/View/STransitionNode.cpp
@@ -6,6 +6,7 @@
 #include "StateCentricView/View/StateCentricViewUtils.h"
 #include "StateTreeEditorStyle.h"
 
+#include "Containers/Ticker.h"
 #include "Customizations/StateTreeTransitionDetails.h"
 #include "DetailsViewArgs.h"
 #include "IDetailsView.h"
@@ -34,16 +35,30 @@ namespace UE::StateTree::Editor::StateCentricView
 //////////////////////////////////////////////////////////////////////////
 // STransitionNode
 
+STransitionNode::~STransitionNode()
+{
+	if (ViewModel)
+	{
+		ViewModel->GetTransitionViewModel().GetOnTransitionModified().RemoveAll(this);
+	}
+}
+
 void STransitionNode::Construct(const SExtendableNode::FArguments& InArgs
-		, TSharedRef<FStateTreeViewModel> InViewModel
-		, TNotNull<UStateTreeState*> InViewState
-		, TNotNull<FStateTreeTransition*> InTransition)
+	, TSharedRef<FStateTreeViewModel> InViewModel
+	, TNotNull<UStateTreeState*> InViewState
+	, TNotNull<FStateTreeTransition*> InTransition)
 {
 	TransitionID = InTransition->ID;
 
 	SStateTreeExtendableNode::Construct(InArgs, InViewModel, InViewState);
 }
 
+
+void STransitionNode::Construct_InitializeNodePersistentPreLOD()
+{
+	ViewModel->GetTransitionViewModel().GetOnTransitionModified().AddSP(this, &STransitionNode::HandleOnTransitionModified);
+}
+
 void STransitionNode::Construct_GenerateBuiltInExtensions(const EExtendableNodeLOD InLOD)
 {
 	// @TODO: Rather than handle the transition arrow via padding. Maybe just add an inner node for the transition & make the 
@@ -282,6 +297,37 @@ TSharedRef<SWidget> STransitionNode::Construct_GenerateNodeHeader(const EExtenda
 					]
 				];
 		}
+		else
+		{
+			// State Invalid. Indicate so
+			return SNew(SOverlay)
+				+ SOverlay::Slot()
+				[
+					SNew(SHorizontalBox)
+					+ SHorizontalBox::Slot()
+					.HAlign(HAlign_Left)
+					.VAlign(VAlign_Center)
+					.Padding(FMargin(12.0, 0.0, 2.0, 0.0))
+					.AutoWidth()
+					[
+						SNew(SImage)
+						.Image(FStateTreeEditorStyle::Get().GetBrush("StateTreeEditor.Transition.None"))
+						.ColorAndOpacity(FStyleColors::Foreground)
+						.Visibility(EVisibility::HitTestInvisible)
+					]
+					+ SHorizontalBox::Slot()
+					.HAlign(HAlign_Left)
+					.VAlign(VAlign_Center)
+					.AutoWidth()
+					[
+						SNew(STextBlock)
+						.Text(LOCTEXT("TransitionNodeMissingState", "Invalid State"))
+						.TextStyle(&FStateTreeEditorStyle::Get().GetWidgetStyle<FTextBlockStyle>("Normal.Normal"))
+						.ColorAndOpacity(FStyleColors::White)
+						.Visibility(EVisibility::HitTestInvisible)
+					]
+				];
+		}
 	}
 
 	return SNullWidget::NullWidget;
@@ -368,6 +414,28 @@ TSharedRef<SWidget> STransitionNode::Construct_GenerateNodeBody(const EExtendabl
 	return SNullWidget::NullWidget;
 }
 
+void STransitionNode::HandleOnTransitionModified(const UStateTreeState* OwningState, const FGuid& InTransitionID)
+{
+	if (TransitionID == InTransitionID)
+	{
+		TWeakPtr<STransitionNode> ThisWeak = StaticCastSharedRef<STransitionNode>(AsShared());
+		auto RefreshTransitionNode = [ThisWeak](float)
+		{
+			if (TSharedPtr<STransitionNode> ThisPinned = ThisWeak.Pin())
+			{
+				ThisPinned->RegenerateWidgetForLOD(ThisPinned->CurrentLOD);
+			}
+
+			// Run just once
+			return false;
+		};
+
+		// Refresh on a delay. We might have been modified from our own details panel
+		// If we instantly refresh we can kill that details mid-update which can remove property handles & cause update issues
+		FTSTicker::GetCoreTicker().AddTicker(FTickerDelegate::CreateLambda(RefreshTransitionNode));
+	}
+}
+
 
 //////////////////////////////////////////////////////////////////////////
 // SInTransitionNode
@@ -382,6 +450,7 @@ FName SInTransitionNode::GetNodeName() const
 void SInTransitionNode::Construct_InitializeNodePersistentPreLOD()
 {
 	TransitionDirection = EStateTransitionDirection::In;
+	STransitionNode::Construct_InitializeNodePersistentPreLOD();
 }
 
 
@@ -398,6 +467,7 @@ FName SOutTransitionNode::GetNodeName() const
 void SOutTransitionNode::Construct_InitializeNodePersistentPreLOD()
 {
 	TransitionDirection = EStateTransitionDirection::Out;
+	STransitionNode::Construct_InitializeNodePersistentPreLOD();
 }
 
 
diff --git a/Engine/Plugins/Runtime/StateTree/Source/StateTreeEditorModule/Private/StateCentricView/View/STransitionNode.h b/Engine/Plugins/Runtime/StateTree/Source/StateTreeEditorModule/Private/StateCentricView/View/STransitionNode.h
--- a/Engine/Plugins/Runtime/StateTree/Source/StateTreeEditorModule/Private/StateCentricView/View/STransitionNode.h
+++ b/Engine/Plugins/Runtime/StateTree/Source/StateTreeEditorModule/Private/StateCentricView/View/STransitionNode.h
@@ -29,6 +29,8 @@ class STransitionNode : public SStateTreeExtendableNode
 
 public:
 
+	virtual ~STransitionNode() override;
+
 	void Construct(const SExtendableNode::FArguments& InArgs
 		, TSharedRef<FStateTreeViewModel> InViewModel
 		, TNotNull<UStateTreeState*> InViewState
@@ -37,11 +39,17 @@ public:
 protected:
 
 	//~ Begin SExtendableNode Interface
+	virtual void Construct_InitializeNodePersistentPreLOD() override;
 	virtual void Construct_GenerateBuiltInExtensions(const EExtendableNodeLOD InLOD) override;
 	virtual TSharedRef<SWidget> Construct_GenerateNodeHeader(const EExtendableNodeLOD InLOD) override;
 	virtual TSharedRef<SWidget> Construct_GenerateNodeBody(const EExtendableNodeLOD InLOD) override;
 	//~ End SExtendableNode Interface
 
+protected:
+
+	/** Callback for transtion modified */
+	void HandleOnTransitionModified(const UStateTreeState* OwningState, const FGuid& InTransitionID);
+
 protected:
 
 	/** 
diff --git a/Engine/Plugins/Runtime/StateTree/Source/StateTreeEditorModule/Private/StateTreeState.cpp b/Engine/Plugins/Runtime/StateTree/Source/StateTreeEditorModule/Private/StateTreeState.cpp
--- a/Engine/Plugins/Runtime/StateTree/Source/StateTreeEditorModule/Private/StateTreeState.cpp
+++ b/Engine/Plugins/Runtime/StateTree/Source/StateTreeEditorModule/Private/StateTreeState.cpp
@@ -9,6 +9,7 @@
 #include "StateTreeTaskBase.h"
 #include "StateTreeDelegates.h"
 #include "StateTreePropertyHelpers.h"
+#include "StateTreeViewModel.h"
 
 #include "Customizations/StateTreeEditorNodeUtils.h"
 
@@ -105,6 +106,7 @@ void UStateTreeState::PreEditChange(FEditPropertyChain& PropertyAboutToChange)
 	const FStateTreeEditPropertyPath PropertyChainPath(PropertyAboutToChange);
 
 	static const FStateTreeEditPropertyPath StateTypePath(UStateTreeState::StaticClass(), TEXT("Type"));
+	static const FStateTreeEditPropertyPath StateTransitionsPath(UStateTreeState::StaticClass(), TEXT("Transitions"));
 
 	if (PropertyChainPath.IsPathExact(StateTypePath))
 	{
@@ -115,6 +117,12 @@ void UStateTreeState::PreEditChange(FEditPropertyChain& PropertyAboutToChange)
 			Parameters.ResetParametersAndOverrides();
 		}
 	}
+
+	if (PropertyChainPath.IsPathExact(StateTransitionsPath))
+	{
+		// Copy Transitions so we can do diff compares later
+		PreEditTransitions = Transitions;
+	}
 }
 
 void UStateTreeState::PostEditChangeChainProperty(FPropertyChangedChainEvent& PropertyChangedEvent)
@@ -139,6 +147,16 @@ void UStateTreeState::PostEditChangeChainProperty(FPropertyChangedChainEvent& Pr
 	UStateTree* StateTree = GetTypedOuter<UStateTree>();
 	ensure(StateTree);
 
+	TSharedPtr<FStateTreeViewModel> StateTreeViewModel = nullptr;
+	if (UStateTreeEditorData* TreeData = GetTypedOuter<UStateTreeEditorData>())
+	{
+		if (UStateTreeEditingSubsystem* StateTreeEditingSubsystem = TreeData ? GEditor->GetEditorSubsystem<UStateTreeEditingSubsystem>() : nullptr)
+		{
+			// Only find, don't create a view model if it doesn't exist
+			StateTreeViewModel = StateTreeEditingSubsystem->FindViewModel(TreeData->GetOuterUStateTree());
+		}
+	}
+
 	// Broadcast name changes so that the UI can update.
 	if (ChangePropertyPath.IsPathExact(StateNamePath))
 	{
@@ -224,7 +242,6 @@ void UStateTreeState::PostEditChangeChainProperty(FPropertyChangedChainEvent& Pr
 		}
 	}
 
-	// Reset delay on completion transitions
 	if (ChangePropertyPath.ContainsPath(StateTransitionsPath))
 	{
 		const int32 TransitionsIndex = ChangePropertyPath.GetPropertyArrayIndex(StateTransitionsPath);
@@ -232,28 +249,86 @@ void UStateTreeState::PostEditChangeChainProperty(FPropertyChangedChainEvent& Pr
 		{
 			FStateTreeTransition& Transition = Transitions[TransitionsIndex];
 
+			// Reset delay on completion transitions
 			if (EnumHasAnyFlags(Transition.Trigger, EStateTreeTransitionTrigger::OnStateCompleted))
 			{
 				Transition.bDelayTransition = false;
 			}
+
+			// Let ViewModel know the transition was modified
+			if (StateTreeViewModel)
+			{
+				StateTreeViewModel->GetTransitionViewModel().HandleOnTransitionModified(this, Transition.ID);
+			}
 		}
 	}
 
-	// Set default state to root and Id on new transitions.
-	if (PropertyChangedEvent.ChangeType == EPropertyChangeType::ArrayAdd)
+	// Test exact path for array modifications
+	if (ChangePropertyPath.IsPathExact(StateTransitionsPath))
 	{
-		if (ChangePropertyPath.IsPathExact(StateTransitionsPath))
+		if (PropertyChangedEvent.ChangeType == EPropertyChangeType::ArrayAdd)
 		{
 			const int32 TransitionsIndex = ChangePropertyPath.GetPropertyArrayIndex(StateTransitionsPath);
 			if (Transitions.IsValidIndex(TransitionsIndex))
 			{
+				// Set default state to root and Id on new transitions.
 				FStateTreeTransition& Transition = Transitions[TransitionsIndex];
 				Transition.Trigger = EStateTreeTransitionTrigger::OnStateCompleted;
 				const UStateTreeState* RootState = GetRootState();
 				Transition.State = RootState->GetLinkToState();
 				Transition.ID = FGuid::NewGuid();
+
+				// Notify Transtions ViewModel of new transtion
+				if (StateTreeViewModel)
+				{
+					StateTreeViewModel->GetTransitionViewModel().HandleOnTransitionsAdded(this, { Transition.ID });
+				}
 			}
 		}
+		else if (PropertyChangedEvent.ChangeType == EPropertyChangeType::ArrayRemove || PropertyChangedEvent.ChangeType == EPropertyChangeType::ArrayClear)
+		{
+			TSet<FGuid> RemovedTransitions;
+			RemovedTransitions.Reserve(PreEditTransitions.Num());
+
+			for (const FStateTreeTransition& PreEditTransition : PreEditTransitions)
+			{
+				if (!Transitions.Contains(PreEditTransition))
+				{
+					RemovedTransitions.Add(PreEditTransition.ID);
+				}
+			}
+
+			// Notify Transtions ViewModel of removed transtions
+			if (StateTreeViewModel)
+			{
+				StateTreeViewModel->GetTransitionViewModel().HandleOnTransitionsRemoved(this, RemovedTransitions);
+			}
+		}
+		else if (PropertyChangedEvent.ChangeType == EPropertyChangeType::ArrayMove)
+		{
+			TSet<FGuid> MovedTransitions;
+			MovedTransitions.Reserve(PreEditTransitions.Num());
+
+			if (ensure(PreEditTransitions.Num() == Transitions.Num()))
+			{
+				for (int32 Index = 0; Index < PreEditTransitions.Num(); Index++)
+				{
+					if (PreEditTransitions[Index].ID != Transitions[Index].ID)
+					{
+						MovedTransitions.Add(PreEditTransitions[Index].ID);
+					}
+				}
+			}
+
+			// Notify Transtions ViewModel of removed transtions
+			if (StateTreeViewModel)
+			{
+				StateTreeViewModel->GetTransitionViewModel().HandleOnTransitionsMoved(this, MovedTransitions);
+			}
+		}
+
+		// Always reset temp captured old transitions
+		PreEditTransitions.Reset();
 	}
 
 	if (StateTree)
diff --git a/Engine/Plugins/Runtime/StateTree/Source/StateTreeEditorModule/Private/StateTreeTransitionViewModel.cpp b/Engine/Plugins/Runtime/StateTree/Source/StateTreeEditorModule/Private/StateTreeTransitionViewModel.cpp
new file mode 100644
--- /dev/null
+++ b/Engine/Plugins/Runtime/StateTree/Source/StateTreeEditorModule/Private/StateTreeTransitionViewModel.cpp
@@ -0,0 +1,31 @@
+// Copyright Epic Games, Inc. All Rights Reserved.
+
+#include "StateTreeTransitionViewModel.h"
+
+#include "StateTree.h"
+#include "Misc/NotNull.h"
+#include "Misc/StringOutputDevice.h"
+
+#define LOCTEXT_NAMESPACE "StateTreeEditor"
+
+void FStateTreeTransitionViewModel::HandleOnTransitionsAdded(const UStateTreeState* OwningState, const TSet<FGuid>& TransitionID) const
+{
+	OnTransitionsAdded.Broadcast(OwningState, TransitionID);
+}
+
+void FStateTreeTransitionViewModel::HandleOnTransitionsRemoved(const UStateTreeState* OwningState, const TSet<FGuid>& TransitionIDs) const
+{
+	OnTransitionsRemoved.Broadcast(OwningState, TransitionIDs);
+}
+
+void FStateTreeTransitionViewModel::HandleOnTransitionsMoved(const UStateTreeState* OwningState, const TSet<FGuid>& TransitionIDs) const
+{
+	OnTransitionsMoved.Broadcast(OwningState, TransitionIDs);
+}
+
+void FStateTreeTransitionViewModel::HandleOnTransitionModified(const UStateTreeState* OwningState, const FGuid& TransitionID) const
+{
+	OnTransitionModified.Broadcast(OwningState, TransitionID);
+}
+
+#undef LOCTEXT_NAMESPACE
diff --git a/Engine/Plugins/Runtime/StateTree/Source/StateTreeEditorModule/Private/StateTreeViewModel.cpp b/Engine/Plugins/Runtime/StateTree/Source/StateTreeEditorModule/Private/StateTreeViewModel.cpp
--- a/Engine/Plugins/Runtime/StateTree/Source/StateTreeEditorModule/Private/StateTreeViewModel.cpp
+++ b/Engine/Plugins/Runtime/StateTree/Source/StateTreeEditorModule/Private/StateTreeViewModel.cpp
@@ -382,9 +382,10 @@ namespace Private
 			return ArrayIndex;
 		}
 	};
-}
 
-};
+} //namespace Private
+
+} // namespace UE::StateTree::Editor
 
 FStateTreeViewModel::FStateTreeViewModel()
 	: TreeDataWeak(nullptr)
@@ -2125,6 +2126,68 @@ const UStateTreeState* FStateTreeViewModel::GetTargetParentStateFromTransition(T
 	return nullptr;
 }
 
+FStateTreeTransition& FStateTreeViewModel::AddTransition(TNotNull<UStateTreeState*> State, const EStateTreeTransitionTrigger InTrigger, const EStateTreeTransitionType InType, const UStateTreeState* InState) const
+{
+	FStateTreeTransition* Result = nullptr;
+	if (UStateTreeEditorData* EditorData = TreeDataWeak.Get())
+	{
+		static const FString TransitionsPropertyNodeName = GET_MEMBER_NAME_CHECKED(UStateTreeState, Transitions).ToString();
+
+		auto AddFunc = [&](TNotNull<UStateTreeState*> InOwnerState, TNotNull<UStateTreeEditorData*> InEditorData, const FStateTreeEditPropertyPath& InPropertyPath)
+		{
+			Result = &State->AddTransition(InTrigger, InType, InState);
+		};
+
+		UE::StateTree::PropertyHelpers::ModifyStateInPreAndPostEdit(
+			LOCTEXT("AddTransitionTransaction", "Add Transition"),
+			State,
+			EditorData,
+			TransitionsPropertyNodeName,
+			AddFunc,
+			State->Transitions.Num(),
+			EPropertyChangeType::ArrayAdd);
+	}
+
+	// No editor, still add the state
+	if (!Result)
+	{
+		Result = &State->AddTransition(InTrigger, InType, InState);
+	}
+
+	return *Result;
+}
+
+FStateTreeTransition& FStateTreeViewModel::AddTransition(TNotNull<UStateTreeState*> State, const EStateTreeTransitionTrigger InTrigger, const FGameplayTag InEventTag, const EStateTreeTransitionType InType, const UStateTreeState* InState) const
+{
+	FStateTreeTransition* Result = nullptr;
+	if (UStateTreeEditorData* EditorData = TreeDataWeak.Get())
+	{
+		static const FString TransitionsPropertyNodeName = GET_MEMBER_NAME_CHECKED(UStateTreeState, Transitions).ToString();
+
+		auto AddFunc = [&](TNotNull<UStateTreeState*> InOwnerState, TNotNull<UStateTreeEditorData*> InEditorData, const FStateTreeEditPropertyPath& InPropertyPath)
+		{
+			Result = &State->AddTransition(InTrigger, InEventTag, InType, InState);
+		};
+
+		UE::StateTree::PropertyHelpers::ModifyStateInPreAndPostEdit(
+			LOCTEXT("AddTransitionTransaction", "Add Transition"),
+			State,
+			EditorData,
+			TransitionsPropertyNodeName,
+			AddFunc,
+			State->Transitions.Num(),
+			EPropertyChangeType::ArrayAdd);
+	}
+
+	// No editor, still add the state
+	if (!Result)
+	{
+		Result = &State->AddTransition(InTrigger, InEventTag, InType, InState);
+	}
+
+	return *Result;
+}
+
 FLinearColor FStateTreeViewModel::GetStateColor(TNotNull<const UStateTreeState*> State, EStateTreeViewModelColorAdjustment ColorAdjustment) const
 {
 	if (bEnableSimpleHierarchyDrivenColors)
diff --git a/Engine/Plugins/Runtime/StateTree/Source/StateTreeEditorModule/Public/StateTreeState.h b/Engine/Plugins/Runtime/StateTree/Source/StateTreeEditorModule/Public/StateTreeState.h
--- a/Engine/Plugins/Runtime/StateTree/Source/StateTreeEditorModule/Public/StateTreeState.h
+++ b/Engine/Plugins/Runtime/StateTree/Source/StateTreeEditorModule/Public/StateTreeState.h
@@ -115,6 +115,11 @@ PRAGMA_ENABLE_DEPRECATION_WARNINGS
 		return FGuid::Combine(ID, FGuid::NewDeterministicGuid(TEXT("Event")));
 	}
 
+	bool operator==(const FStateTreeTransition& Other) const
+	{
+		return ID == Other.ID;
+	}
+
 	UE_API void PostSerialize(const FArchive& Ar);
 
 	/** When to try trigger the transition. */
@@ -461,6 +466,11 @@ public:
 
 	UPROPERTY(meta = (ExcludeFromHash))
 	TObjectPtr<UStateTreeState> Parent = nullptr;
+
+private:
+
+	/** Transaction temporary. Transitions we captured in PreEdit for remove / move comparision */
+	TArray<FStateTreeTransition> PreEditTransitions;
 };
 
 #undef UE_API
diff --git a/Engine/Plugins/Runtime/StateTree/Source/StateTreeEditorModule/Public/StateTreeTransitionViewModel.h b/Engine/Plugins/Runtime/StateTree/Source/StateTreeEditorModule/Public/StateTreeTransitionViewModel.h
new file mode 100644
--- /dev/null
+++ b/Engine/Plugins/Runtime/StateTree/Source/StateTreeEditorModule/Public/StateTreeTransitionViewModel.h
@@ -0,0 +1,69 @@
+// Copyright Epic Games, Inc. All Rights Reserved.
+
+#pragma once
+
+#include "UObject/WeakObjectPtr.h"
+#include "EditorUndoClient.h"
+#include "IStateTreeEditorHost.h"
+
+#define UE_API STATETREEEDITORMODULE_API
+
+class UStateTreeState;
+
+/**
+ * Viewmodel for editing state transtions.
+ */
+struct FStateTreeTransitionViewModel
+{
+	DECLARE_MULTICAST_DELEGATE_TwoParams(FOnTransitionsAdded, const UStateTreeState* /*OwningState*/, const TSet<FGuid>& /*TransitionIDs*/);
+	DECLARE_MULTICAST_DELEGATE_TwoParams(FOnTransitionsRemoved, const UStateTreeState* /*OwningState*/, const TSet<FGuid>& /*TransitionIDs*/);
+	DECLARE_MULTICAST_DELEGATE_TwoParams(FOnTransitionsMoved, const UStateTreeState* /*OwningState*/, const TSet<FGuid>& /*TransitionIDs*/);
+	DECLARE_MULTICAST_DELEGATE_TwoParams(FOnTransitionModified, const UStateTreeState* /*OwningState*/, const FGuid& /*TransitionID*/);
+
+public:
+
+	/** Fired externally to let the view model know a transition(s) was added. */
+	void HandleOnTransitionsAdded(const UStateTreeState* OwningState, const TSet<FGuid>& TransitionIDs) const;
+
+	/** Fired externally to let the view model know a transition(s) was removed. */
+	void HandleOnTransitionsRemoved(const UStateTreeState* OwningState, const TSet<FGuid>& TransitionIDs) const;
+
+	/** Fired externally to let the view model know a transition(s) was moved. */
+	void HandleOnTransitionsMoved(const UStateTreeState* OwningState, const TSet<FGuid>& TransitionIDs) const;
+
+	/** Fired externally to let the view model know a transition was modified. */
+	void HandleOnTransitionModified(const UStateTreeState* OwningState, const FGuid& TransitionID) const;
+
+	// Called each time a transition(s) is added
+	FOnTransitionsAdded& GetOnTransitionsAdded()
+	{
+		return OnTransitionsAdded;
+	}
+
+	// Called each time a transition(s) is removed.
+	FOnTransitionsRemoved& GetOnTransitionsRemoved()
+	{
+		return OnTransitionsRemoved;
+	}
+
+	// Called each time a transition(s) is moved.
+	FOnTransitionsRemoved& GetOnTransitionsMoved()
+	{
+		return OnTransitionsMoved;
+	}
+
+	// Called each time a transition is modified
+	FOnTransitionModified& GetOnTransitionModified()
+	{
+		return OnTransitionModified;
+	}
+
+protected:
+
+	FOnTransitionsAdded OnTransitionsAdded;
+	FOnTransitionsRemoved OnTransitionsRemoved;
+	FOnTransitionsMoved OnTransitionsMoved;
+	FOnTransitionModified OnTransitionModified;
+};
+
+#undef UE_API
diff --git a/Engine/Plugins/Runtime/StateTree/Source/StateTreeEditorModule/Public/StateTreeViewModel.h b/Engine/Plugins/Runtime/StateTree/Source/StateTreeEditorModule/Public/StateTreeViewModel.h
--- a/Engine/Plugins/Runtime/StateTree/Source/StateTreeEditorModule/Public/StateTreeViewModel.h
+++ b/Engine/Plugins/Runtime/StateTree/Source/StateTreeEditorModule/Public/StateTreeViewModel.h
@@ -5,6 +5,8 @@
 #include "UObject/WeakObjectPtr.h"
 #include "EditorUndoClient.h"
 #include "IStateTreeEditorHost.h"
+#include "StateTreeTransitionViewModel.h"
+
 #include "StateTreeViewModel.generated.h"
 
 #define UE_API STATETREEEDITORMODULE_API
@@ -25,8 +27,11 @@ class UStateTreeState;
 enum class ECheckBoxState : uint8;
 enum class EStateTreeBreakpointType : uint8;
 enum class EStateTreeNodeFormatting : uint8;
+enum class EStateTreeTransitionTrigger : uint8;
+enum class EStateTreeTransitionType : uint8;
 enum class EStateTreeVisitor : uint8;
 
+struct FGameplayTag;
 struct FSlateBrush;
 struct FPropertyBindingPath;
 struct FPropertyChangedEvent;
@@ -162,6 +167,10 @@ public:
 	UE_API const UStateTreeState* GetTargetStateFromTransition(TNotNull<const FStateTreeTransition*> Transition) const;
 	UE_API const UStateTreeState* GetTargetParentStateFromTransition(TNotNull<const FStateTreeTransition*> Transition) const;
 
+	// Transition Manipulation
+	UE_API FStateTreeTransition& AddTransition(TNotNull<UStateTreeState*> State, const EStateTreeTransitionTrigger InTrigger, const EStateTreeTransitionType InType, const UStateTreeState* InState = nullptr) const;
+	UE_API FStateTreeTransition& AddTransition(TNotNull<UStateTreeState*> State, const EStateTreeTransitionTrigger InTrigger, const FGameplayTag InEventTag, const EStateTreeTransitionType InType, const UStateTreeState* InState = nullptr) const;
+
 	// Color manipulation
 	UE_API FLinearColor GetStateColor(TNotNull<const UStateTreeState*> State, EStateTreeViewModelColorAdjustment ColorAdjustment = EStateTreeViewModelColorAdjustment::Normal) const;
 	UE_API FLinearColor GetHierarchyBasedHighlightColor(TNotNull<const UStateTreeState*> State, EStateTreeViewModelColorAdjustment ColorAdjustment = EStateTreeViewModelColorAdjustment::Normal) const;
@@ -263,6 +272,17 @@ public:
 		return OnBringBindingPathToFocus;
 	}
 
+	// Get View Model for transitions. Holds various callbacks for transtions added / modified / etc.
+	FStateTreeTransitionViewModel& GetTransitionViewModel()
+	{
+		return TransitionViewModel;
+	}
+
+	const FStateTreeTransitionViewModel& GetTransitionViewModel() const
+	{
+		return TransitionViewModel;
+	}
+
 protected:
 	UE_API void GetExpandedStatesRecursive(UStateTreeState* State, TSet<TWeakObjectPtr<UStateTreeState>>& ExpandedStates);
 
@@ -297,6 +317,9 @@ protected:
 	FOnSelectionChanged OnSelectionChanged;
 	FOnBringNodeToFocus OnBringNodeToFocus;
 	FOnBringBindingPathToFocus OnBringBindingPathToFocus;
+
+	/** Viewmodel for transition modifications */
+	FStateTreeTransitionViewModel TransitionViewModel;
 };
 
 /** Helper class to allow to copy bindings into clipboard. */
