Add file accidentally split from state centric transition control change on submit
- Fix UE::StateTree::PropertyHelpers::ModifyStateInPreAndPostEdit not propagating array index

#tests editor
#Synced-CL 48823676

diff --git a/Engine/Plugins/Runtime/StateTree/Source/StateTreeEditorModule/Private/StateTreePropertyHelpers.cpp b/Engine/Plugins/Runtime/StateTree/Source/StateTreeEditorModule/Private/StateTreePropertyHelpers.cpp
--- a/Engine/Plugins/Runtime/StateTree/Source/StateTreeEditorModule/Private/StateTreePropertyHelpers.cpp
+++ b/Engine/Plugins/Runtime/StateTree/Source/StateTreeEditorModule/Private/StateTreePropertyHelpers.cpp
@@ -279,6 +279,12 @@ void ModifyStateInPreAndPostEdit(
 	FPropertyChangedEvent ChangedEvent(PropertyChain.GetActiveNode()->GetValue(), ChangeType);
 	ChangedEvent.SetArrayIndexPerObject(ArrayIndicesPerObject);
 
+	// Set ObjectIteratorIndex to index we used above as FPropertyChangedEvent::GetArrayIndex uses that to read the array indexes.
+	if (ArrayIndex != INDEX_NONE)
+	{
+		ChangedEvent.ObjectIteratorIndex = 0;
+	}
+
 	FPropertyChangedChainEvent ChainEvent(PropertyChain, ChangedEvent);
 	State->PostEditChangeChainProperty(ChainEvent);
 }
