Add Parent Selection Entry Node

#tests editor
#Synced-CL 48823676

diff --git a/Engine/Plugins/Runtime/StateTree/Source/StateTreeEditorModule/Private/StateCentricView/StateCentricViewSettings.h b/Engine/Plugins/Runtime/StateTree/Source/StateTreeEditorModule/Private/StateCentricView/StateCentricViewSettings.h
--- a/Engine/Plugins/Runtime/StateTree/Source/StateTreeEditorModule/Private/StateCentricView/StateCentricViewSettings.h
+++ b/Engine/Plugins/Runtime/StateTree/Source/StateTreeEditorModule/Private/StateCentricView/StateCentricViewSettings.h
@@ -39,7 +39,7 @@ public:
 
 	/** True if central / leafmost view should have transitions justified instead of top aligned */
 	UPROPERTY(EditAnywhere, Category = "Layout")
-	bool bJustifyCentralViewTransitions = true;
+	bool bJustifyCentralViewTransitions = false;
 
 	/** True if state centric view allows for full detail expansion of parents, else max expansion is entry conditions */
 	UPROPERTY(EditAnywhere, Category="Layout")
diff --git a/Engine/Plugins/Runtime/StateTree/Source/StateTreeEditorModule/Private/StateCentricView/View/SParentSelectionNode.cpp b/Engine/Plugins/Runtime/StateTree/Source/StateTreeEditorModule/Private/StateCentricView/View/SParentSelectionNode.cpp
new file mode 100644
--- /dev/null
+++ b/Engine/Plugins/Runtime/StateTree/Source/StateTreeEditorModule/Private/StateCentricView/View/SParentSelectionNode.cpp
@@ -0,0 +1,211 @@
+// Copyright Epic Games, Inc. All Rights Reserved.
+
+#include "StateCentricView/View/SParentSelectionNode.h"
+
+#include "Customizations/StateTreeBindingExtension.h"
+#include "StateCentricView/View/StateCentricViewUtils.h"
+#include "StateTreeEditorStyle.h"
+
+#include "Layout/Visibility.h"
+#include "Modules/ModuleManager.h"
+#include "ObjectEditorUtils.h"
+#include "PropertyEditorModule.h"
+#include "PropertyPathHelpers.h"
+#include "Styling/SlateBrush.h"
+#include "Styling/StyleColors.h"
+#include "Widgets/Images/SImage.h"
+#include "Widgets/Input/SButton.h"
+#include "Widgets/Layout/SBorder.h"
+#include "Widgets/Text/STextBlock.h"
+
+// Bad API, need to make not private
+#include "Editor/PropertyEditor/Private/PropertyNode.h"
+
+#define LOCTEXT_NAMESPACE "SParentSelectionNode"
+
+namespace UE::StateTree::Editor::StateCentricView
+{
+
+
+//////////////////////////////////////////////////////////////////////////
+// SParentSelectionNode
+
+const FName SParentSelectionNode::NodeID(TEXT("StateCentricView.ParentSelection.Node"));
+
+FName SParentSelectionNode::GetNodeName() const
+{
+	return NodeID;
+}
+
+void SParentSelectionNode::Construct_GenerateBuiltInExtensions(const EExtendableNodeLOD InLOD)
+{
+	// @TODO: Rather than handle the ParentSelection arrow via padding. Maybe just add an inner node for the ParentSelection & make the 
+	// ParentSelection arrow not part of this node
+
+	if (ViewState.IsValid())
+	{
+		// Parent state name as top extension
+		if (const UStateTreeState* ParentState = ViewState->Parent)
+		{
+			ParentStateBrush = FStateCentricViewUtils::MakeStateSubnodeBackgroundOutlined(ViewModel.ToSharedRef(), ParentState);
+
+			FScopedMainNodeLayoutInfo ScopedMainLayoutInfo = SetScopedMainNodeLayoutInfo()
+				.SetPadding(FMargin(0.0));
+
+			// @TODO: Make sure to add a warning icon if one of the parents in the nested hierachy has a entry condition
+			TSharedRef<SWidget> ParentStateNameTopExtension = SNew(SOverlay)
+				+ SOverlay::Slot()
+				[
+					SNew(SImage)
+					.Image(ParentStateBrush->GetSlateBrush())
+				]
+				+ SOverlay::Slot()
+				[
+					FStateCentricViewUtils::GenerateGotoStateButton(InLOD, ViewModel.ToSharedRef(), ParentState)
+				]
+				+ SOverlay::Slot()
+				[
+					SNew(SHorizontalBox)
+					+ SHorizontalBox::Slot()
+					.HAlign(HAlign_Left)
+					.VAlign(VAlign_Center)
+					.Padding(FMargin(12.0, 0.0, 2.0, 0.0))
+					.AutoWidth()
+					[
+						SNew(SImage)
+						.Image(ViewModel->GetSelectorIcon(ParentState))
+						.ColorAndOpacity(FStyleColors::Foreground)
+						.Visibility(EVisibility::HitTestInvisible)
+					]
+					+ SHorizontalBox::Slot()
+					.HAlign(HAlign_Left)
+					.VAlign(VAlign_Center)
+					.AutoWidth()
+					[
+						SNew(STextBlock)
+						.Text(FText::FromName(ParentState->Name))
+						.TextStyle(&FStateTreeEditorStyle::Get().GetWidgetStyle<FTextBlockStyle>("Normal.Normal"))
+						.ColorAndOpacity(FStyleColors::White)
+						.Visibility(EVisibility::HitTestInvisible)
+					]
+				];
+
+			AddExtension(ParentStateNameTopExtension)
+				.SetLocation(EExtensionLocation::Top)
+				.SetAutoSize()
+				.SetPadding(FMargin(0.0, 0.0, 0.0, 4.0));
+		}
+
+		// Small arrow used to indicate ParentSelection on left / right depending on direction
+		{
+			FScopedMainNodeLayoutInfo ScopedMainLayoutInfo = SetScopedMainNodeLayoutInfo()
+				.SetFillSize(1.0f)
+				.SetPadding(FMargin(0.0));
+
+			TSharedRef<SWidget> ParentSelectionArrowExtension =
+				SNew(SButton)
+				.NormalPaddingOverride(FMargin(2.0, 4.0))
+				.PressedPaddingOverride(FMargin(2.0, 4.0))
+				.ContentPadding(FMargin(0.0f))
+				[
+					SNew(SHorizontalBox)
+						+ SHorizontalBox::Slot()
+						[
+							SNew(SImage)
+								.Image(FAppStyle::Get().GetBrush("Icons.ChevronRight"))
+								.DesiredSizeOverride(FVector2D(10, 24))
+								.ColorAndOpacity(FStyleColors::Foreground)
+						]
+						+ SHorizontalBox::Slot()
+						.Padding(FMargin(-8.0, 0.0, 0.0, 0.0))
+						[
+							SNew(SImage)
+								.Image(FAppStyle::Get().GetBrush("Icons.ChevronRight"))
+								.DesiredSizeOverride(FVector2D(10, 24))
+								.ColorAndOpacity(FStyleColors::Foreground)
+						]
+				];
+				
+			AddExtension(ParentSelectionArrowExtension)
+				.SetLocation(EExtensionLocation::Right)
+				.SetVAlign(VAlign_Center)
+				.SetAutoSize()
+				.SetPadding(FMargin(6.0, 0.0, -6.0, 0.0));
+		}
+
+		// Generic background extension, padded to not overlap with the small ParentSelection arrow
+		{
+			FScopedMainNodeLayoutInfo ScopedMainLayoutInfo = SetScopedMainNodeLayoutInfo()
+				.SetPadding(FMargin(4.0));
+
+			const FSlateBrush* SubNodeBackgroundBrush = FStateCentricViewUtils::GetSubNodeBackground();
+
+			TSharedPtr<SWidget> SubNodeBackground = SNew(SBorder)
+				.BorderImage(SubNodeBackgroundBrush)
+				.Visibility(EVisibility::HitTestInvisible);
+
+			AddExtension(SubNodeBackground.ToSharedRef())
+				.SetLocation(EExtensionLocation::Background)
+				.SetAddLocationBehavior(EExtensionAddLocationBehavior::OverlayEntireNode)
+				.SetPadding(FMargin(0.0, 0.0, 16.0, 0.0));
+		}
+	}
+}
+
+TSharedRef<SWidget> SParentSelectionNode::Construct_GenerateNodeHeader(const EExtendableNodeLOD InLOD)
+{
+	NodeBodyPadding = FMargin(0.0, 0.0, 0.0, 0.0);
+
+	if (ViewState.IsValid())
+	{
+		// Show Parent Selection method
+		if (const UStateTreeState* ParentState = ViewState->Parent)
+		{
+			NodeBodyPadding = FMargin(0.0, 2.0, 0.0, 0.0);
+			
+			return SNew(SOverlay)
+				+ SOverlay::Slot()
+				[
+					SNew(SHorizontalBox)
+					+ SHorizontalBox::Slot()
+					.HAlign(HAlign_Left)
+					.VAlign(VAlign_Center)
+					.Padding(FMargin(12.0, 0.0, 2.0, 0.0))
+					.AutoWidth()
+					[
+						SNew(SImage)
+						.Image(FStateTreeEditorStyle::Get().GetBrush("StateTreeEditor.TrySelectChildrenInOrder"))
+						.ColorAndOpacity(FStyleColors::Foreground)
+						.Visibility(EVisibility::HitTestInvisible)
+					]
+					+ SHorizontalBox::Slot()
+					.HAlign(HAlign_Left)
+					.VAlign(VAlign_Center)
+					.AutoWidth()
+					[
+						SNew(STextBlock)
+						.Text(StaticEnum<EStateTreeStateSelectionBehavior>()->GetDisplayNameTextByValue((int64)ParentState->SelectionBehavior))
+						.TextStyle(&FStateTreeEditorStyle::Get().GetWidgetStyle<FTextBlockStyle>("Normal.Normal"))
+						.ColorAndOpacity(FStyleColors::White)
+						.Visibility(EVisibility::HitTestInvisible)
+						.AutoWrapText(true)
+					]
+				];
+		}
+	}
+
+	return SNullWidget::NullWidget;
+}
+
+TSharedRef<SWidget> SParentSelectionNode::Construct_GenerateNodeBody(const EExtendableNodeLOD InLOD)
+{
+	return SNullWidget::NullWidget;
+}
+
+
+//////////////////////////////////////////////////////////////////////////
+
+} // namespace UE::StateTree::Editor::StateCentricView
+
+#undef LOCTEXT_NAMESPACE
+
diff --git a/Engine/Plugins/Runtime/StateTree/Source/StateTreeEditorModule/Private/StateCentricView/View/SParentSelectionNode.h b/Engine/Plugins/Runtime/StateTree/Source/StateTreeEditorModule/Private/StateCentricView/View/SParentSelectionNode.h
new file mode 100644
--- /dev/null
+++ b/Engine/Plugins/Runtime/StateTree/Source/StateTreeEditorModule/Private/StateCentricView/View/SParentSelectionNode.h
@@ -0,0 +1,51 @@
+// Copyright Epic Games, Inc. All Rights Reserved.
+
+#pragma once
+
+#include "CoreMinimal.h"
+
+#include "StateCentricView/View/SStateTreeExtendableNode.h"
+#include "StateTreeState.h"
+
+struct FSlateBrush;
+
+class IStructureDetailsView;
+class FStateTreeViewModel;
+class SButton;
+class UStateTreeState;
+
+namespace UE::StateTree::Editor::StateCentricView
+{
+
+
+/**
+ * Experimental. Node used to show a single state's parent selection behavior
+ */
+class SParentSelectionNode : public SStateTreeExtendableNode
+{
+
+public:
+
+	static const FName NodeID;
+
+public:
+
+	//~ Begin SExtendableNode Interface
+	virtual FName GetNodeName() const override;
+	//~ End SExtendableNode Interface
+
+protected:
+
+	//~ Begin SExtendableNode Interface
+	virtual void Construct_GenerateBuiltInExtensions(const EExtendableNodeLOD InLOD) override;
+	virtual TSharedRef<SWidget> Construct_GenerateNodeHeader(const EExtendableNodeLOD InLOD) override;
+	virtual TSharedRef<SWidget> Construct_GenerateNodeBody(const EExtendableNodeLOD InLOD) override;
+	//~ End SExtendableNode Interface
+
+private:
+
+	TSharedPtr<FDeferredCleanupSlateBrush> ParentStateBrush;
+};
+
+
+} // namespace UE::StateTree::Editor::StateCentricView
\ No newline at end of file
diff --git a/Engine/Plugins/Runtime/StateTree/Source/StateTreeEditorModule/Private/StateCentricView/View/StateCentricViewUtils.cpp b/Engine/Plugins/Runtime/StateTree/Source/StateTreeEditorModule/Private/StateCentricView/View/StateCentricViewUtils.cpp
--- a/Engine/Plugins/Runtime/StateTree/Source/StateTreeEditorModule/Private/StateCentricView/View/StateCentricViewUtils.cpp
+++ b/Engine/Plugins/Runtime/StateTree/Source/StateTreeEditorModule/Private/StateCentricView/View/StateCentricViewUtils.cpp
@@ -4,6 +4,7 @@
 
 #include "StateCentricView/StateCentricViewSettings.h"
 #include "StateCentricView/ViewModel/StateCentricViewEditorDataExtension.h"
+#include "StateCentricView/View/SParentSelectionNode.h"
 #include "StateCentricView/View/SStateTreeExtendableNode.h"
 #include "StateCentricView/View/STransitionNode.h"
 #include "StateTreeEditorStyle.h"
@@ -60,6 +61,41 @@ TSharedRef<SWidget> FStateCentricViewUtils::GenerateTransitionsViewWidget(const
 		}
 
 		TSharedRef<SVerticalBox> TransitionListBody = SNew(SVerticalBox);
+
+		if (ViewState == CentralState && TransitionDirection == EStateTransitionDirection::In)
+		{
+			if (const UStateTreeState* ParentState = ViewState->Parent)
+			{
+				if (ParentState->SelectionBehavior != EStateTreeStateSelectionBehavior::None
+					&& ParentState->SelectionBehavior != EStateTreeStateSelectionBehavior::TryEnterState
+					&& ParentState->SelectionBehavior != EStateTreeStateSelectionBehavior::TryFollowTransitions)
+				{
+					TSharedPtr<SWidget> ParenSelectionWidget = SNew(SParentSelectionNode, ViewModel, ViewState)
+						.Config({ .InitialLOD = LOD });
+
+					if (UStateCentricViewSettings::GetSchemaSettingsDefaultFallback(&ViewModel.Get()).bJustifyCentralViewTransitions)
+					{
+						TransitionListBody->AddSlot()
+						.FillContentHeight(1.0f)
+						.VAlign(VAlign_Center)
+						.Padding(FMargin(8.0, 6.0, 8.0, 6.0))
+						[
+							ParenSelectionWidget.ToSharedRef()
+						];
+					}
+					else
+					{
+						TransitionListBody->AddSlot()
+						.AutoHeight()
+						.Padding(FMargin(8.0, 6.0, 8.0, 6.0))
+						[
+							ParenSelectionWidget.ToSharedRef()
+						];
+					}
+				}
+			}
+		}
+
 		for (FStateTreeTransition* Transition : Transitions)
 		{
 			TSharedPtr<SWidget> TransitionNodeWidget;
@@ -591,11 +627,6 @@ TSharedPtr<FDeferredCleanupSlateBrush> FStateCentricViewUtils::MakeStateSubnodeB
 		Result.TintColor = ViewModel->GetHierarchyBasedBackgroundColor(State, EStateTreeViewModelColorAdjustment::Dark);
 		Result.OutlineSettings.Color = ViewModel->GetHierarchyBasedBackgroundColor(State, EStateTreeViewModelColorAdjustment::Dark);
 
-		if (bEnableDepthBasedOutline)
-		{
-			Result.OutlineSettings.Color = ViewModel->GetHierarchyBasedBackgroundColor(State, EStateTreeViewModelColorAdjustment::Bright);
-		}
-
 		if (bDisableDarkenHeader)
 		{
 			Result.TintColor = ViewModel->GetHierarchyBasedBackgroundColor(State);
@@ -619,11 +650,6 @@ TSharedPtr<FDeferredCleanupSlateBrush> FStateCentricViewUtils::MakeStateCollapse
 		Result.TintColor = ViewModel->GetHierarchyBasedBackgroundColor(State, EStateTreeViewModelColorAdjustment::Dark);
 		Result.OutlineSettings.Color = ViewModel->GetHierarchyBasedBackgroundColor(State, EStateTreeViewModelColorAdjustment::Dark);
 
-		if (bEnableDepthBasedOutline)
-		{
-			Result.OutlineSettings.Color = ViewModel->GetHierarchyBasedBackgroundColor(State, EStateTreeViewModelColorAdjustment::Bright);
-		}
-
 		if (bDisableDarkenHeader)
 		{
 			Result.TintColor = ViewModel->GetHierarchyBasedBackgroundColor(State);
