Fix state centric issues with unloaded schema references not saving in project settings

#tests editor

diff --git a/Engine/Plugins/Runtime/StateTree/Source/StateTreeEditorModule/Private/StateCentricView/StateCentricViewPerUserSettings.h b/Engine/Plugins/Runtime/StateTree/Source/StateTreeEditorModule/Private/StateCentricView/StateCentricViewPerUserSettings.h
--- a/Engine/Plugins/Runtime/StateTree/Source/StateTreeEditorModule/Private/StateCentricView/StateCentricViewPerUserSettings.h
+++ b/Engine/Plugins/Runtime/StateTree/Source/StateTreeEditorModule/Private/StateCentricView/StateCentricViewPerUserSettings.h
@@ -153,7 +153,7 @@ public:
 
 	UE_API static const FStateCentricViewPerUserSetting& GetSchemaSettingsDefaultFallback(const TSubclassOf<UStateTreeSchema> InSchema)
 	{
-		if (const FStateCentricViewPerUserSetting* SchemaSettings = UStateCentricViewPerUserSettings::Get().PerSchemaSettings.Find(InSchema))
+		if (const FStateCentricViewPerUserSetting* SchemaSettings = UStateCentricViewPerUserSettings::Get().PerSchemaSettings.Find(TSoftClassPtr<UStateTreeSchema>(InSchema)))
 		{
 			return *SchemaSettings;
 		}
@@ -164,8 +164,8 @@ public:
 public:
 
 	/** Per schema settings. Defined per schema since we have no need for an editor schema to be defined. */
-	UPROPERTY(EditAnywhere, Category="Appearance", Config, meta = (EditCondition = "StateTreeEditorModule.StateCentricViewSettings.IsStateCentricViewEnabled", EditConditionHides))
-	TMap<TSubclassOf<UStateTreeSchema>, FStateCentricViewPerUserSetting> PerSchemaSettings;
+	UPROPERTY(EditAnywhere, Category="Appearance", Config, meta = (AllowAbstract = "true", EditCondition = "StateTreeEditorModule.StateCentricViewSettings.IsStateCentricViewEnabled", EditConditionHides))
+	TMap<TSoftClassPtr<UStateTreeSchema>, FStateCentricViewPerUserSetting> PerSchemaSettings;
 
 protected:
 
diff --git a/Engine/Plugins/Runtime/StateTree/Source/StateTreeEditorModule/Private/StateCentricView/StateCentricViewSettings.h b/Engine/Plugins/Runtime/StateTree/Source/StateTreeEditorModule/Private/StateCentricView/StateCentricViewSettings.h
--- a/Engine/Plugins/Runtime/StateTree/Source/StateTreeEditorModule/Private/StateCentricView/StateCentricViewSettings.h
+++ b/Engine/Plugins/Runtime/StateTree/Source/StateTreeEditorModule/Private/StateCentricView/StateCentricViewSettings.h
@@ -68,7 +68,7 @@ public:
  * 
  * Settings for the state centric view per schema. As state centric is early in dev, there are many toggles.
  */
-UCLASS(Experimental, MinimalAPI, config = Editor)
+UCLASS(Experimental, MinimalAPI, config = EditorSettings, defaultconfig)
 class UStateCentricViewSettings : public UDeveloperSettings
 {
 	GENERATED_BODY()
@@ -91,7 +91,7 @@ public:
 
 	UE_API static const FStateCentricViewPerSchemaSetting& GetSchemaSettingsDefaultFallback(const TSubclassOf<UStateTreeSchema> InSchema)
 	{
-		if (const FStateCentricViewPerSchemaSetting* SchemaSettings = UStateCentricViewSettings::Get().PerSchemaSettings.Find(InSchema))
+		if (const FStateCentricViewPerSchemaSetting* SchemaSettings = UStateCentricViewSettings::Get().PerSchemaSettings.Find(TSoftClassPtr<UStateTreeSchema>(InSchema)))
 		{
 			return *SchemaSettings;
 		}
@@ -102,8 +102,8 @@ public:
 public:
 
 	/** Per schema settings. Defined per schema since we have no need for an editor schema to be defined. */
-	UPROPERTY(EditAnywhere, Category="Appearance", Config, meta = (EditCondition = "StateTreeEditorModule.StateCentricViewSettings.IsStateCentricViewEnabled", EditConditionHides))
-	TMap<TSubclassOf<UStateTreeSchema>, FStateCentricViewPerSchemaSetting> PerSchemaSettings;
+	UPROPERTY(EditAnywhere, Category="Appearance", Config, meta = (AllowAbstract = "true", EditCondition = "StateTreeEditorModule.StateCentricViewSettings.IsStateCentricViewEnabled", EditConditionHides))
+	TMap<TSoftClassPtr<UStateTreeSchema>, FStateCentricViewPerSchemaSetting> PerSchemaSettings;
 
 public:
 
