Fix word wrapping state centric

#tests editor
#Synced-CL 48823676

diff --git a/Engine/Plugins/Runtime/StateTree/Source/StateTreeEditorModule/Private/StateCentricView/View/SMainStateTransitionExtensionNode.cpp b/Engine/Plugins/Runtime/StateTree/Source/StateTreeEditorModule/Private/StateCentricView/View/SMainStateTransitionExtensionNode.cpp
--- a/Engine/Plugins/Runtime/StateTree/Source/StateTreeEditorModule/Private/StateCentricView/View/SMainStateTransitionExtensionNode.cpp
+++ b/Engine/Plugins/Runtime/StateTree/Source/StateTreeEditorModule/Private/StateCentricView/View/SMainStateTransitionExtensionNode.cpp
@@ -196,6 +196,7 @@ TSharedRef<SWidget> SMainStateTransitionExtensionNode::Construct_GenerateNodeHea
 			+ SOverlay::Slot()
 			[
 				SAssignNew(HeaderContentsBox, SHorizontalBox)
+				.Clipping(EWidgetClipping::OnDemand)
 				+ SHorizontalBox::Slot()
 				.HAlign(HeaderAlignment)
 				.VAlign(VAlign_Center)
diff --git a/Engine/Plugins/Runtime/StateTree/Source/StateTreeEditorModule/Private/StateCentricView/View/SParentStateExtensionNode.cpp b/Engine/Plugins/Runtime/StateTree/Source/StateTreeEditorModule/Private/StateCentricView/View/SParentStateExtensionNode.cpp
--- a/Engine/Plugins/Runtime/StateTree/Source/StateTreeEditorModule/Private/StateCentricView/View/SParentStateExtensionNode.cpp
+++ b/Engine/Plugins/Runtime/StateTree/Source/StateTreeEditorModule/Private/StateCentricView/View/SParentStateExtensionNode.cpp
@@ -209,6 +209,7 @@ TSharedRef<SWidget> SParentStateExtensionNode::Construct_GenerateNodeHeader(cons
 			+ SOverlay::Slot()
 			[
 				SNew(SHorizontalBox)
+				.Clipping(EWidgetClipping::OnDemand)
 				+ SHorizontalBox::Slot()
 				.HAlign(HAlign_Left)
 				.VAlign(VAlign_Center)
diff --git a/Engine/Plugins/Runtime/StateTree/Source/StateTreeEditorModule/Private/StateCentricView/View/SParentStateTransitionExtensionNode.cpp b/Engine/Plugins/Runtime/StateTree/Source/StateTreeEditorModule/Private/StateCentricView/View/SParentStateTransitionExtensionNode.cpp
--- a/Engine/Plugins/Runtime/StateTree/Source/StateTreeEditorModule/Private/StateCentricView/View/SParentStateTransitionExtensionNode.cpp
+++ b/Engine/Plugins/Runtime/StateTree/Source/StateTreeEditorModule/Private/StateCentricView/View/SParentStateTransitionExtensionNode.cpp
@@ -148,6 +148,7 @@ TSharedRef<SWidget> SParentStateTransitionExtensionNode::Construct_GenerateNodeH
 			+ SOverlay::Slot()
 			[
 				SAssignNew(HeaderContentsBox, SHorizontalBox)
+				.Clipping(EWidgetClipping::OnDemand)
 				+ SHorizontalBox::Slot()
 				.HAlign(HeaderAlignment)
 				.VAlign(VAlign_Center)
diff --git a/Engine/Plugins/Runtime/StateTree/Source/StateTreeEditorModule/Private/StateCentricView/View/SStateCentricViewMainStateNode.cpp b/Engine/Plugins/Runtime/StateTree/Source/StateTreeEditorModule/Private/StateCentricView/View/SStateCentricViewMainStateNode.cpp
--- a/Engine/Plugins/Runtime/StateTree/Source/StateTreeEditorModule/Private/StateCentricView/View/SStateCentricViewMainStateNode.cpp
+++ b/Engine/Plugins/Runtime/StateTree/Source/StateTreeEditorModule/Private/StateCentricView/View/SStateCentricViewMainStateNode.cpp
@@ -237,6 +237,7 @@ TSharedRef<SWidget> SStateCentricViewMainStateNode::Construct_GenerateNodeHeader
 			.VAlign(VAlign_Center)
 			[
 				SNew(SHorizontalBox)
+				.Clipping(EWidgetClipping::OnDemand)
 				+ SHorizontalBox::Slot()
 				.HAlign(HAlign_Right)
 				.VAlign(VAlign_Center)
diff --git a/Engine/Plugins/Runtime/StateTree/Source/StateTreeEditorModule/Private/StateCentricView/View/STransitionNode.cpp b/Engine/Plugins/Runtime/StateTree/Source/StateTreeEditorModule/Private/StateCentricView/View/STransitionNode.cpp
--- a/Engine/Plugins/Runtime/StateTree/Source/StateTreeEditorModule/Private/StateCentricView/View/STransitionNode.cpp
+++ b/Engine/Plugins/Runtime/StateTree/Source/StateTreeEditorModule/Private/StateCentricView/View/STransitionNode.cpp
@@ -103,9 +103,8 @@ void STransitionNode::Construct_GenerateBuiltInExtensions(const EExtendableNodeL
 						.Visibility(EVisibility::HitTestInvisible)
 					]
 					+ SHorizontalBox::Slot()
-					.HAlign(HAlign_Left)
+					.HAlign(HAlign_Fill)
 					.VAlign(VAlign_Center)
-					.AutoWidth()
 					[
 						SNew(STextBlock)
 						.Text(FText::FromName(ParentState->Name))
@@ -285,9 +284,8 @@ TSharedRef<SWidget> STransitionNode::Construct_GenerateNodeHeader(const EExtenda
 						.Visibility(EVisibility::HitTestInvisible)
 					]
 					+ SHorizontalBox::Slot()
-					.HAlign(HAlign_Left)
+					.HAlign(HAlign_Fill)
 					.VAlign(VAlign_Center)
-					.AutoWidth()
 					[
 						SNew(STextBlock)
 						.Text(FText::FromName(State->Name))
diff --git a/Engine/Plugins/Runtime/StateTree/Source/StateTreeEditorModule/Private/StateCentricView/View/StateCentricViewUtils.cpp b/Engine/Plugins/Runtime/StateTree/Source/StateTreeEditorModule/Private/StateCentricView/View/StateCentricViewUtils.cpp
--- a/Engine/Plugins/Runtime/StateTree/Source/StateTreeEditorModule/Private/StateCentricView/View/StateCentricViewUtils.cpp
+++ b/Engine/Plugins/Runtime/StateTree/Source/StateTreeEditorModule/Private/StateCentricView/View/StateCentricViewUtils.cpp
@@ -70,7 +70,8 @@ TSharedRef<SWidget> FStateCentricViewUtils::GenerateTransitionsViewWidget(const
 					&& ParentState->SelectionBehavior != EStateTreeStateSelectionBehavior::TryEnterState
 					&& ParentState->SelectionBehavior != EStateTreeStateSelectionBehavior::TryFollowTransitions)
 				{
-					TSharedPtr<SWidget> ParenSelectionWidget = SNew(SParentSelectionNode, ViewModel, ViewState)
+					TSharedPtr<SWidget> ParentSelectionWidget = SNew(SParentSelectionNode, ViewModel, ViewState)
+						.Clipping(EWidgetClipping::OnDemand)
 						.Config({ .InitialLOD = LOD });
 
 					if (UStateCentricViewSettings::GetSchemaSettingsDefaultFallback(&ViewModel.Get()).bJustifyCentralViewTransitions)
@@ -80,7 +81,7 @@ TSharedRef<SWidget> FStateCentricViewUtils::GenerateTransitionsViewWidget(const
 						.VAlign(VAlign_Center)
 						.Padding(FMargin(8.0, 6.0, 8.0, 6.0))
 						[
-							ParenSelectionWidget.ToSharedRef()
+							ParentSelectionWidget.ToSharedRef()
 						];
 					}
 					else
@@ -89,7 +90,7 @@ TSharedRef<SWidget> FStateCentricViewUtils::GenerateTransitionsViewWidget(const
 						.AutoHeight()
 						.Padding(FMargin(8.0, 6.0, 8.0, 6.0))
 						[
-							ParenSelectionWidget.ToSharedRef()
+							ParentSelectionWidget.ToSharedRef()
 						];
 					}
 				}
@@ -102,11 +103,13 @@ TSharedRef<SWidget> FStateCentricViewUtils::GenerateTransitionsViewWidget(const
 			if (TransitionDirection == EStateTransitionDirection::In)
 			{
 				TransitionNodeWidget = SNew(SInTransitionNode, ViewModel, ViewState, Transition)
+					.Clipping(EWidgetClipping::OnDemand)
 					.Config({ .InitialLOD = LOD });
 			}
 			else if (ensure(TransitionDirection == EStateTransitionDirection::Out))
 			{
 				TransitionNodeWidget = SNew(SOutTransitionNode, ViewModel, ViewState, Transition)
+					.Clipping(EWidgetClipping::OnDemand)
 					.Config({ .InitialLOD = LOD });
 			}
 
