Add a node for implicit transition to root to state centric.
Add toggle for word wrap around transtion state names / child state names to state centric.
Fix clipping on child states.

#tests editor
#Synced-CL 48823676

diff --git a/Engine/Plugins/Runtime/StateTree/Source/StateTreeEditorModule/Private/StateCentricView/StateCentricViewSettings.h b/Engine/Plugins/Runtime/StateTree/Source/StateTreeEditorModule/Private/StateCentricView/StateCentricViewSettings.h
--- a/Engine/Plugins/Runtime/StateTree/Source/StateTreeEditorModule/Private/StateCentricView/StateCentricViewSettings.h
+++ b/Engine/Plugins/Runtime/StateTree/Source/StateTreeEditorModule/Private/StateCentricView/StateCentricViewSettings.h
@@ -41,6 +41,14 @@ public:
 	UPROPERTY(EditAnywhere, Category = "Layout")
 	bool bJustifyCentralViewTransitions = false;
 
+	/** True if states in transitions should have word wrap on */
+	UPROPERTY(EditAnywhere, Category = "Layout")
+	bool bWordWrapTransitionStates = true;
+
+	/** True if child states should have word wrap on */
+	UPROPERTY(EditAnywhere, Category = "Layout")
+	bool bWordWrapChildStates = false;
+
 	/** True if state centric view allows for full detail expansion of parents, else max expansion is entry conditions */
 	UPROPERTY(EditAnywhere, Category="Layout")
 	bool bEnableFullParentDetailExpansion = true;
diff --git a/Engine/Plugins/Runtime/StateTree/Source/StateTreeEditorModule/Private/StateCentricView/View/SFakeTransitionInfoNode.cpp b/Engine/Plugins/Runtime/StateTree/Source/StateTreeEditorModule/Private/StateCentricView/View/SFakeTransitionInfoNode.cpp
new file mode 100644
--- /dev/null
+++ b/Engine/Plugins/Runtime/StateTree/Source/StateTreeEditorModule/Private/StateCentricView/View/SFakeTransitionInfoNode.cpp
@@ -0,0 +1,209 @@
+// Copyright Epic Games, Inc. All Rights Reserved.
+
+#include "StateCentricView/View/SFakeTransitionInfoNode.h"
+
+#include "Customizations/StateTreeBindingExtension.h"
+#include "StateCentricView/StateCentricViewSettings.h"
+#include "StateCentricView/View/StateCentricViewUtils.h"
+#include "StateTreeEditorStyle.h"
+
+#include "Layout/Visibility.h"
+#include "Modules/ModuleManager.h"
+#include "ObjectEditorUtils.h"
+#include "PropertyEditorModule.h"
+#include "PropertyPathHelpers.h"
+#include "Styling/SlateBrush.h"
+#include "Styling/StyleColors.h"
+#include "Widgets/Images/SImage.h"
+#include "Widgets/Input/SButton.h"
+#include "Widgets/Layout/SBorder.h"
+#include "Widgets/Text/STextBlock.h"
+
+#define LOCTEXT_NAMESPACE "SFakeTransitionInfoNode"
+
+namespace UE::StateTree::Editor::StateCentricView
+{
+
+
+//////////////////////////////////////////////////////////////////////////
+// SFakeTransitionInfoNode
+
+void SFakeTransitionInfoNode::Construct_GenerateBuiltInExtensions(const EExtendableNodeLOD InLOD)
+{
+	// @TODO: Rather than handle the ParentSelection arrow via padding. Maybe just add an inner node for the ParentSelection & make the 
+	// ParentSelection arrow not part of this node
+
+	if (ViewState.IsValid())
+	{
+		// Parent state name as top extension
+		if (const UStateTreeState* FakeTranstionInfoState = GetFakeTranstionInfoState())
+		{
+			StateBrush = FStateCentricViewUtils::MakeStateSubnodeBackgroundOutlined(ViewModel.ToSharedRef(), FakeTranstionInfoState);
+
+			FScopedMainNodeLayoutInfo ScopedMainLayoutInfo = SetScopedMainNodeLayoutInfo()
+				.SetPadding(FMargin(0.0));
+
+			TSharedRef<SWidget> FakeTranstionInfoStateTopExtension = SNew(SOverlay)
+				+ SOverlay::Slot()
+				[
+					SNew(SImage)
+					.Image(StateBrush->GetSlateBrush())
+				]
+				+ SOverlay::Slot()
+				[
+					FStateCentricViewUtils::GenerateGotoStateButton(InLOD, ViewModel.ToSharedRef(), FakeTranstionInfoState)
+				]
+				+ SOverlay::Slot()
+				[
+					SNew(SHorizontalBox)
+					+ SHorizontalBox::Slot()
+					.HAlign(HAlign_Left)
+					.VAlign(VAlign_Center)
+					.Padding(FMargin(12.0, 0.0, 2.0, 0.0))
+					.AutoWidth()
+					[
+						SNew(SImage)
+						.Image(ViewModel->GetSelectorIcon(FakeTranstionInfoState))
+						.ColorAndOpacity(FStyleColors::Foreground)
+						.Visibility(EVisibility::HitTestInvisible)
+					]
+					+ SHorizontalBox::Slot()
+					.HAlign(HAlign_Fill)
+					.VAlign(VAlign_Center)
+					[
+						SNew(STextBlock)
+						.Text(FText::FromName(FakeTranstionInfoState->Name))
+						.TextStyle(&FStateTreeEditorStyle::Get().GetWidgetStyle<FTextBlockStyle>("Normal.Normal"))
+						.ColorAndOpacity(FStyleColors::White)
+						.Visibility(EVisibility::HitTestInvisible)
+						.AutoWrapText(UStateCentricViewSettings::GetSchemaSettingsDefaultFallback(ViewModel.Get()).bWordWrapTransitionStates)
+					]
+				];
+
+			AddExtension(FakeTranstionInfoStateTopExtension)
+				.SetLocation(EExtensionLocation::Top)
+				.SetAutoSize()
+				.SetPadding(FMargin(0.0, 0.0, 0.0, 4.0));
+		}
+
+		// Small arrow used to indicate transition on left / right depending on direction
+		{
+			FScopedMainNodeLayoutInfo ScopedMainLayoutInfo = SetScopedMainNodeLayoutInfo()
+				.SetFillSize(1.0f)
+				.SetPadding(FMargin(0.0));
+
+			TSharedRef<SWidget> TransitionArrowExtension =
+				SNew(SButton)
+				.NormalPaddingOverride(FMargin(2.0, 4.0))
+				.PressedPaddingOverride(FMargin(2.0, 4.0))
+				.ContentPadding(FMargin(0.0f))
+				[
+					SNew(SHorizontalBox)
+						+ SHorizontalBox::Slot()
+						[
+							SNew(SImage)
+								.Image(FAppStyle::Get().GetBrush("Icons.ChevronRight"))
+								.DesiredSizeOverride(FVector2D(10, 24))
+								.ColorAndOpacity(FStyleColors::Foreground)
+						]
+						+ SHorizontalBox::Slot()
+						.Padding(FMargin(-8.0, 0.0, 0.0, 0.0))
+						[
+							SNew(SImage)
+								.Image(FAppStyle::Get().GetBrush("Icons.ChevronRight"))
+								.DesiredSizeOverride(FVector2D(10, 24))
+								.ColorAndOpacity(FStyleColors::Foreground)
+						]
+				];
+
+			EExtensionLocation TransitionArrowLocation = GetTransitionDirection() == EStateTransitionDirection::In
+				? EExtensionLocation::Right
+				: EExtensionLocation::Left;
+
+			FMargin TransitionArrowPadding = GetTransitionDirection() == EStateTransitionDirection::In
+				? FMargin(6.0, 0.0, -6.0, 0.0)
+				: FMargin(-6.0, 0.0, 6.0, 0.0);
+				
+			AddExtension(TransitionArrowExtension)
+				.SetLocation(TransitionArrowLocation)
+				.SetVAlign(VAlign_Center)
+				.SetAutoSize()
+				.SetPadding(TransitionArrowPadding);
+		}
+
+		// Generic background extension, padded to not overlap with the small transition arrow
+		{
+			FScopedMainNodeLayoutInfo ScopedMainLayoutInfo = SetScopedMainNodeLayoutInfo()
+				.SetPadding(FMargin(4.0));
+
+			FMargin TransitionArrowPadding = GetTransitionDirection() == EStateTransitionDirection::In
+				? FMargin(0.0, 0.0, 16.0, 0.0)
+				: FMargin(16.0, 0.0, 0.0, 0.0);
+
+			const FSlateBrush* SubNodeBackgroundBrush = FStateCentricViewUtils::GetSubNodeBackground();
+
+			TSharedPtr<SWidget> SubNodeBackground = SNew(SBorder)
+				.BorderImage(SubNodeBackgroundBrush)
+				.Visibility(EVisibility::HitTestInvisible);
+
+			AddExtension(SubNodeBackground.ToSharedRef())
+				.SetLocation(EExtensionLocation::Background)
+				.SetAddLocationBehavior(EExtensionAddLocationBehavior::OverlayEntireNode)
+				.SetPadding(TransitionArrowPadding);
+		}
+	}
+}
+
+TSharedRef<SWidget> SFakeTransitionInfoNode::Construct_GenerateNodeHeader(const EExtendableNodeLOD InLOD)
+{
+	NodeBodyPadding = FMargin(0.0, 0.0, 0.0, 0.0);
+
+	if (ViewState.IsValid())
+	{
+		NodeBodyPadding = FMargin(0.0, 2.0, 0.0, 0.0);
+			
+		return SNew(SOverlay)
+			+ SOverlay::Slot()
+			[
+				SNew(SHorizontalBox)
+				+ SHorizontalBox::Slot()
+				.HAlign(HAlign_Left)
+				.VAlign(VAlign_Center)
+				.Padding(FMargin(12.0, 0.0, 2.0, 0.0))
+				.AutoWidth()
+				[
+					SNew(SImage)
+					.Image(GetFakeTranstionInfoBrush())
+					.ColorAndOpacity(FStyleColors::Foreground)
+					.Visibility(EVisibility::HitTestInvisible)
+				]
+				+ SHorizontalBox::Slot()
+				.HAlign(HAlign_Fill)
+				.VAlign(VAlign_Center)
+				[
+					SNew(STextBlock)
+					.Text(GetFakeTranstionInfoText())
+					.ToolTipText(GetFakeTranstionInfoToolTipText())
+					.TextStyle(&FStateTreeEditorStyle::Get().GetWidgetStyle<FTextBlockStyle>("Normal.Normal"))
+					.ColorAndOpacity(FStyleColors::White)
+					.Visibility(EVisibility::Visible)
+					.AutoWrapText(true)
+				]
+			];
+	}
+
+	return SNullWidget::NullWidget;
+}
+
+TSharedRef<SWidget> SFakeTransitionInfoNode::Construct_GenerateNodeBody(const EExtendableNodeLOD InLOD)
+{
+	return SNullWidget::NullWidget;
+}
+
+
+//////////////////////////////////////////////////////////////////////////
+
+} // namespace UE::StateTree::Editor::StateCentricView
+
+#undef LOCTEXT_NAMESPACE
+
diff --git a/Engine/Plugins/Runtime/StateTree/Source/StateTreeEditorModule/Private/StateCentricView/View/SFakeTransitionInfoNode.h b/Engine/Plugins/Runtime/StateTree/Source/StateTreeEditorModule/Private/StateCentricView/View/SFakeTransitionInfoNode.h
new file mode 100644
--- /dev/null
+++ b/Engine/Plugins/Runtime/StateTree/Source/StateTreeEditorModule/Private/StateCentricView/View/SFakeTransitionInfoNode.h
@@ -0,0 +1,61 @@
+// Copyright Epic Games, Inc. All Rights Reserved.
+
+#pragma once
+
+#include "CoreMinimal.h"
+
+#include "StateCentricView/View/SStateTreeExtendableNode.h"
+#include "StateTreeState.h"
+
+struct FSlateBrush;
+
+class IStructureDetailsView;
+class FStateTreeViewModel;
+class SButton;
+class UStateTreeState;
+
+namespace UE::StateTree::Editor::StateCentricView
+{
+
+
+/**
+ * Experimental. Node used to show some info in transition area that isn't a real transition.
+ */
+class SFakeTransitionInfoNode : public SStateTreeExtendableNode
+{
+
+protected:
+
+	/** Get if we are intended to be shown at in or out transtions. */
+	virtual EStateTransitionDirection GetTransitionDirection() const = 0;
+
+	/** Get State to show for this node. */
+	virtual const UStateTreeState* GetFakeTranstionInfoState() const = 0;
+
+	/** Get Icon to use in header for this node. */
+	virtual const FSlateBrush* GetFakeTranstionInfoBrush() const = 0;
+
+	/** Get Text to show for this node */
+	virtual FText GetFakeTranstionInfoText() const = 0;
+
+	/** Get Text to show for this node */
+	virtual FText GetFakeTranstionInfoToolTipText() const 
+	{ 
+		return FText::GetEmpty();
+	};
+
+protected:
+
+	//~ Begin SExtendableNode Interface
+	virtual void Construct_GenerateBuiltInExtensions(const EExtendableNodeLOD InLOD) override;
+	virtual TSharedRef<SWidget> Construct_GenerateNodeHeader(const EExtendableNodeLOD InLOD) override;
+	virtual TSharedRef<SWidget> Construct_GenerateNodeBody(const EExtendableNodeLOD InLOD) override;
+	//~ End SExtendableNode Interface
+
+private:
+
+	TSharedPtr<FDeferredCleanupSlateBrush> StateBrush;
+};
+
+
+} // namespace UE::StateTree::Editor::StateCentricView
\ No newline at end of file
diff --git a/Engine/Plugins/Runtime/StateTree/Source/StateTreeEditorModule/Private/StateCentricView/View/SImplicitRootTransitionNode.cpp b/Engine/Plugins/Runtime/StateTree/Source/StateTreeEditorModule/Private/StateCentricView/View/SImplicitRootTransitionNode.cpp
new file mode 100644
--- /dev/null
+++ b/Engine/Plugins/Runtime/StateTree/Source/StateTreeEditorModule/Private/StateCentricView/View/SImplicitRootTransitionNode.cpp
@@ -0,0 +1,77 @@
+// Copyright Epic Games, Inc. All Rights Reserved.
+
+#include "StateCentricView/View/SImplicitRootTransitionNode.h"
+
+#include "Customizations/StateTreeBindingExtension.h"
+#include "StateCentricView/View/StateCentricViewUtils.h"
+#include "StateTreeEditorStyle.h"
+
+#include "Layout/Visibility.h"
+#include "Modules/ModuleManager.h"
+#include "ObjectEditorUtils.h"
+#include "PropertyEditorModule.h"
+#include "PropertyPathHelpers.h"
+#include "Styling/SlateBrush.h"
+#include "Styling/StyleColors.h"
+
+#define LOCTEXT_NAMESPACE "SImplicitRootTransitionNode"
+
+namespace UE::StateTree::Editor::StateCentricView
+{
+
+
+//////////////////////////////////////////////////////////////////////////
+// SImplicitRootTransitionNode
+
+const FName SImplicitRootTransitionNode::NodeID(TEXT("StateCentricView.ImplicitRootTranstion.Node"));
+
+FName SImplicitRootTransitionNode::GetNodeName() const
+{
+	return NodeID;
+}
+
+EStateTransitionDirection SImplicitRootTransitionNode::GetTransitionDirection() const
+{
+	return EStateTransitionDirection::Out;
+}
+
+const UStateTreeState* SImplicitRootTransitionNode::GetFakeTranstionInfoState() const
+{
+	const UStateTreeState* RootState = nullptr;
+	if (ViewState.IsValid())
+	{
+		auto SetRootState = [&RootState](UStateTreeState& ParentState)
+		{
+			RootState = &ParentState;
+			return EStateTreeVisitor::Continue;
+		};
+
+		ViewModel->ForEachParent(ViewState.Get(), SetRootState);
+	}
+
+	return RootState;
+}
+
+const FSlateBrush* SImplicitRootTransitionNode::GetFakeTranstionInfoBrush() const
+{
+	return FStateTreeEditorStyle::Get().GetBrush("StateTreeEditor.Transition.Parent");
+}
+
+FText SImplicitRootTransitionNode::GetFakeTranstionInfoText() const
+{
+	return LOCTEXT("ImplicitRootTransitionNodeHeader", "Implict Transition To Root");
+}
+
+FText SImplicitRootTransitionNode::GetFakeTranstionInfoToolTipText() const
+{
+	return LOCTEXT("ImplicitRootTransitionNodeToolTip"
+		, "When no transitions exist. StateTree will automatically add an internal OnComplete transition to 'Root' to avoid getting stuck.");
+}
+
+
+//////////////////////////////////////////////////////////////////////////
+
+} // namespace UE::StateTree::Editor::StateCentricView
+
+#undef LOCTEXT_NAMESPACE
+
diff --git a/Engine/Plugins/Runtime/StateTree/Source/StateTreeEditorModule/Private/StateCentricView/View/SImplicitRootTransitionNode.h b/Engine/Plugins/Runtime/StateTree/Source/StateTreeEditorModule/Private/StateCentricView/View/SImplicitRootTransitionNode.h
new file mode 100644
--- /dev/null
+++ b/Engine/Plugins/Runtime/StateTree/Source/StateTreeEditorModule/Private/StateCentricView/View/SImplicitRootTransitionNode.h
@@ -0,0 +1,42 @@
+// Copyright Epic Games, Inc. All Rights Reserved.
+
+#pragma once
+
+#include "CoreMinimal.h"
+
+#include "StateCentricView/View/SFakeTransitionInfoNode.h"
+
+namespace UE::StateTree::Editor::StateCentricView
+{
+
+
+/**
+ * Experimental. Node used to show the implicit transtion to root on complete.
+ */
+class SImplicitRootTransitionNode : public SFakeTransitionInfoNode
+{
+
+public:
+
+	static const FName NodeID;
+
+public:
+
+	//~ Begin SExtendableNode Interface
+	virtual FName GetNodeName() const override;
+	//~ End SExtendableNode Interface
+
+protected:
+
+	//~ Begin SFakeTransitionInfoNode Interface
+	virtual EStateTransitionDirection GetTransitionDirection() const override;
+	virtual const UStateTreeState* GetFakeTranstionInfoState() const override;
+	virtual const FSlateBrush* GetFakeTranstionInfoBrush() const override;
+	virtual FText GetFakeTranstionInfoText() const override;
+	virtual FText GetFakeTranstionInfoToolTipText() const override;
+	//~ End SFakeTransitionInfoNode Interface
+
+};
+
+
+} // namespace UE::StateTree::Editor::StateCentricView
\ No newline at end of file
diff --git a/Engine/Plugins/Runtime/StateTree/Source/StateTreeEditorModule/Private/StateCentricView/View/SParentSelectionNode.cpp b/Engine/Plugins/Runtime/StateTree/Source/StateTreeEditorModule/Private/StateCentricView/View/SParentSelectionNode.cpp
--- a/Engine/Plugins/Runtime/StateTree/Source/StateTreeEditorModule/Private/StateCentricView/View/SParentSelectionNode.cpp
+++ b/Engine/Plugins/Runtime/StateTree/Source/StateTreeEditorModule/Private/StateCentricView/View/SParentSelectionNode.cpp
@@ -13,13 +13,6 @@
 #include "PropertyPathHelpers.h"
 #include "Styling/SlateBrush.h"
 #include "Styling/StyleColors.h"
-#include "Widgets/Images/SImage.h"
-#include "Widgets/Input/SButton.h"
-#include "Widgets/Layout/SBorder.h"
-#include "Widgets/Text/STextBlock.h"
-
-// Bad API, need to make not private
-#include "Editor/PropertyEditor/Private/PropertyNode.h"
 
 #define LOCTEXT_NAMESPACE "SParentSelectionNode"
 
@@ -37,169 +30,44 @@ FName SParentSelectionNode::GetNodeName() const
 	return NodeID;
 }
 
-void SParentSelectionNode::Construct_GenerateBuiltInExtensions(const EExtendableNodeLOD InLOD)
+EStateTransitionDirection SParentSelectionNode::GetTransitionDirection() const
 {
-	// @TODO: Rather than handle the ParentSelection arrow via padding. Maybe just add an inner node for the ParentSelection & make the 
-	// ParentSelection arrow not part of this node
+	return EStateTransitionDirection::In;
+}
 
+const UStateTreeState* SParentSelectionNode::GetFakeTranstionInfoState() const
+{
 	if (ViewState.IsValid())
 	{
-		// Parent state name as top extension
-		if (const UStateTreeState* ParentState = ViewState->Parent)
-		{
-			ParentStateBrush = FStateCentricViewUtils::MakeStateSubnodeBackgroundOutlined(ViewModel.ToSharedRef(), ParentState);
-
-			FScopedMainNodeLayoutInfo ScopedMainLayoutInfo = SetScopedMainNodeLayoutInfo()
-				.SetPadding(FMargin(0.0));
-
-			// @TODO: Make sure to add a warning icon if one of the parents in the nested hierachy has a entry condition
-			TSharedRef<SWidget> ParentStateNameTopExtension = SNew(SOverlay)
-				+ SOverlay::Slot()
-				[
-					SNew(SImage)
-					.Image(ParentStateBrush->GetSlateBrush())
-				]
-				+ SOverlay::Slot()
-				[
-					FStateCentricViewUtils::GenerateGotoStateButton(InLOD, ViewModel.ToSharedRef(), ParentState)
-				]
-				+ SOverlay::Slot()
-				[
-					SNew(SHorizontalBox)
-					+ SHorizontalBox::Slot()
-					.HAlign(HAlign_Left)
-					.VAlign(VAlign_Center)
-					.Padding(FMargin(12.0, 0.0, 2.0, 0.0))
-					.AutoWidth()
-					[
-						SNew(SImage)
-						.Image(ViewModel->GetSelectorIcon(ParentState))
-						.ColorAndOpacity(FStyleColors::Foreground)
-						.Visibility(EVisibility::HitTestInvisible)
-					]
-					+ SHorizontalBox::Slot()
-					.HAlign(HAlign_Left)
-					.VAlign(VAlign_Center)
-					.AutoWidth()
-					[
-						SNew(STextBlock)
-						.Text(FText::FromName(ParentState->Name))
-						.TextStyle(&FStateTreeEditorStyle::Get().GetWidgetStyle<FTextBlockStyle>("Normal.Normal"))
-						.ColorAndOpacity(FStyleColors::White)
-						.Visibility(EVisibility::HitTestInvisible)
-					]
-				];
-
-			AddExtension(ParentStateNameTopExtension)
-				.SetLocation(EExtensionLocation::Top)
-				.SetAutoSize()
-				.SetPadding(FMargin(0.0, 0.0, 0.0, 4.0));
-		}
-
-		// Small arrow used to indicate ParentSelection on left / right depending on direction
-		{
-			FScopedMainNodeLayoutInfo ScopedMainLayoutInfo = SetScopedMainNodeLayoutInfo()
-				.SetFillSize(1.0f)
-				.SetPadding(FMargin(0.0));
-
-			TSharedRef<SWidget> ParentSelectionArrowExtension =
-				SNew(SButton)
-				.NormalPaddingOverride(FMargin(2.0, 4.0))
-				.PressedPaddingOverride(FMargin(2.0, 4.0))
-				.ContentPadding(FMargin(0.0f))
-				[
-					SNew(SHorizontalBox)
-						+ SHorizontalBox::Slot()
-						[
-							SNew(SImage)
-								.Image(FAppStyle::Get().GetBrush("Icons.ChevronRight"))
-								.DesiredSizeOverride(FVector2D(10, 24))
-								.ColorAndOpacity(FStyleColors::Foreground)
-						]
-						+ SHorizontalBox::Slot()
-						.Padding(FMargin(-8.0, 0.0, 0.0, 0.0))
-						[
-							SNew(SImage)
-								.Image(FAppStyle::Get().GetBrush("Icons.ChevronRight"))
-								.DesiredSizeOverride(FVector2D(10, 24))
-								.ColorAndOpacity(FStyleColors::Foreground)
-						]
-				];
-				
-			AddExtension(ParentSelectionArrowExtension)
-				.SetLocation(EExtensionLocation::Right)
-				.SetVAlign(VAlign_Center)
-				.SetAutoSize()
-				.SetPadding(FMargin(6.0, 0.0, -6.0, 0.0));
-		}
-
-		// Generic background extension, padded to not overlap with the small ParentSelection arrow
-		{
-			FScopedMainNodeLayoutInfo ScopedMainLayoutInfo = SetScopedMainNodeLayoutInfo()
-				.SetPadding(FMargin(4.0));
-
-			const FSlateBrush* SubNodeBackgroundBrush = FStateCentricViewUtils::GetSubNodeBackground();
-
-			TSharedPtr<SWidget> SubNodeBackground = SNew(SBorder)
-				.BorderImage(SubNodeBackgroundBrush)
-				.Visibility(EVisibility::HitTestInvisible);
-
-			AddExtension(SubNodeBackground.ToSharedRef())
-				.SetLocation(EExtensionLocation::Background)
-				.SetAddLocationBehavior(EExtensionAddLocationBehavior::OverlayEntireNode)
-				.SetPadding(FMargin(0.0, 0.0, 16.0, 0.0));
-		}
+		return ViewState->Parent;
 	}
+
+	return nullptr;
 }
 
-TSharedRef<SWidget> SParentSelectionNode::Construct_GenerateNodeHeader(const EExtendableNodeLOD InLOD)
+const FSlateBrush* SParentSelectionNode::GetFakeTranstionInfoBrush() const
 {
-	NodeBodyPadding = FMargin(0.0, 0.0, 0.0, 0.0);
+	return FStateTreeEditorStyle::Get().GetBrush("StateTreeEditor.TrySelectChildrenInOrder");
+}
 
+FText SParentSelectionNode::GetFakeTranstionInfoText() const
+{
 	if (ViewState.IsValid())
 	{
 		// Show Parent Selection method
 		if (const UStateTreeState* ParentState = ViewState->Parent)
 		{
-			NodeBodyPadding = FMargin(0.0, 2.0, 0.0, 0.0);
-			
-			return SNew(SOverlay)
-				+ SOverlay::Slot()
-				[
-					SNew(SHorizontalBox)
-					+ SHorizontalBox::Slot()
-					.HAlign(HAlign_Left)
-					.VAlign(VAlign_Center)
-					.Padding(FMargin(12.0, 0.0, 2.0, 0.0))
-					.AutoWidth()
-					[
-						SNew(SImage)
-						.Image(FStateTreeEditorStyle::Get().GetBrush("StateTreeEditor.TrySelectChildrenInOrder"))
-						.ColorAndOpacity(FStyleColors::Foreground)
-						.Visibility(EVisibility::HitTestInvisible)
-					]
-					+ SHorizontalBox::Slot()
-					.HAlign(HAlign_Left)
-					.VAlign(VAlign_Center)
-					.AutoWidth()
-					[
-						SNew(STextBlock)
-						.Text(StaticEnum<EStateTreeStateSelectionBehavior>()->GetDisplayNameTextByValue((int64)ParentState->SelectionBehavior))
-						.TextStyle(&FStateTreeEditorStyle::Get().GetWidgetStyle<FTextBlockStyle>("Normal.Normal"))
-						.ColorAndOpacity(FStyleColors::White)
-						.Visibility(EVisibility::HitTestInvisible)
-						.AutoWrapText(true)
-					]
-				];
+			return StaticEnum<EStateTreeStateSelectionBehavior>()->GetDisplayNameTextByValue((int64)ParentState->SelectionBehavior);
 		}
 	}
 
-	return SNullWidget::NullWidget;
+	return FText::GetEmpty();
 }
 
-TSharedRef<SWidget> SParentSelectionNode::Construct_GenerateNodeBody(const EExtendableNodeLOD InLOD)
+FText SParentSelectionNode::GetFakeTranstionInfoToolTipText() const
 {
-	return SNullWidget::NullWidget;
+	return LOCTEXT("ParentSelectionNodeToolTip"
+		, "Parent Selection Mode. This state be entered via parent section.");
 }
 
 
diff --git a/Engine/Plugins/Runtime/StateTree/Source/StateTreeEditorModule/Private/StateCentricView/View/SParentSelectionNode.h b/Engine/Plugins/Runtime/StateTree/Source/StateTreeEditorModule/Private/StateCentricView/View/SParentSelectionNode.h
--- a/Engine/Plugins/Runtime/StateTree/Source/StateTreeEditorModule/Private/StateCentricView/View/SParentSelectionNode.h
+++ b/Engine/Plugins/Runtime/StateTree/Source/StateTreeEditorModule/Private/StateCentricView/View/SParentSelectionNode.h
@@ -4,15 +4,7 @@
 
 #include "CoreMinimal.h"
 
-#include "StateCentricView/View/SStateTreeExtendableNode.h"
-#include "StateTreeState.h"
-
-struct FSlateBrush;
-
-class IStructureDetailsView;
-class FStateTreeViewModel;
-class SButton;
-class UStateTreeState;
+#include "StateCentricView/View/SFakeTransitionInfoNode.h"
 
 namespace UE::StateTree::Editor::StateCentricView
 {
@@ -21,7 +13,7 @@ namespace UE::StateTree::Editor::StateCentricView
 /**
  * Experimental. Node used to show a single state's parent selection behavior
  */
-class SParentSelectionNode : public SStateTreeExtendableNode
+class SParentSelectionNode : public SFakeTransitionInfoNode
 {
 
 public:
@@ -36,15 +28,14 @@ public:
 
 protected:
 
-	//~ Begin SExtendableNode Interface
-	virtual void Construct_GenerateBuiltInExtensions(const EExtendableNodeLOD InLOD) override;
-	virtual TSharedRef<SWidget> Construct_GenerateNodeHeader(const EExtendableNodeLOD InLOD) override;
-	virtual TSharedRef<SWidget> Construct_GenerateNodeBody(const EExtendableNodeLOD InLOD) override;
-	//~ End SExtendableNode Interface
-
-private:
+	//~ Begin SFakeTransitionInfoNode Interface
+	virtual EStateTransitionDirection GetTransitionDirection() const override;
+	virtual const UStateTreeState* GetFakeTranstionInfoState() const override;
+	virtual const FSlateBrush* GetFakeTranstionInfoBrush() const override;
+	virtual FText GetFakeTranstionInfoText() const override;
+	virtual FText GetFakeTranstionInfoToolTipText() const override;
+	//~ End SFakeTransitionInfoNode Interface
 
-	TSharedPtr<FDeferredCleanupSlateBrush> ParentStateBrush;
 };
 
 
diff --git a/Engine/Plugins/Runtime/StateTree/Source/StateTreeEditorModule/Private/StateCentricView/View/STransitionNode.cpp b/Engine/Plugins/Runtime/StateTree/Source/StateTreeEditorModule/Private/StateCentricView/View/STransitionNode.cpp
--- a/Engine/Plugins/Runtime/StateTree/Source/StateTreeEditorModule/Private/StateCentricView/View/STransitionNode.cpp
+++ b/Engine/Plugins/Runtime/StateTree/Source/StateTreeEditorModule/Private/StateCentricView/View/STransitionNode.cpp
@@ -3,6 +3,7 @@
 #include "StateCentricView/View/STransitionNode.h"
 
 #include "Customizations/StateTreeBindingExtension.h"
+#include "StateCentricView/StateCentricViewSettings.h"
 #include "StateCentricView/View/StateCentricViewUtils.h"
 #include "StateTreeEditorStyle.h"
 
@@ -111,6 +112,7 @@ void STransitionNode::Construct_GenerateBuiltInExtensions(const EExtendableNodeL
 						.TextStyle(&FStateTreeEditorStyle::Get().GetWidgetStyle<FTextBlockStyle>("Normal.Normal"))
 						.ColorAndOpacity(FStyleColors::White)
 						.Visibility(EVisibility::HitTestInvisible)
+						.AutoWrapText(UStateCentricViewSettings::GetSchemaSettingsDefaultFallback(ViewModel.Get()).bWordWrapTransitionStates)
 					]
 				];
 
@@ -292,6 +294,7 @@ TSharedRef<SWidget> STransitionNode::Construct_GenerateNodeHeader(const EExtenda
 						.TextStyle(&FStateTreeEditorStyle::Get().GetWidgetStyle<FTextBlockStyle>("Normal.Normal"))
 						.ColorAndOpacity(FStyleColors::White)
 						.Visibility(EVisibility::HitTestInvisible)
+						.AutoWrapText(UStateCentricViewSettings::GetSchemaSettingsDefaultFallback(ViewModel.Get()).bWordWrapTransitionStates)
 					]
 				];
 		}
@@ -314,15 +317,15 @@ TSharedRef<SWidget> STransitionNode::Construct_GenerateNodeHeader(const EExtenda
 						.Visibility(EVisibility::HitTestInvisible)
 					]
 					+ SHorizontalBox::Slot()
-					.HAlign(HAlign_Left)
+					.HAlign(HAlign_Fill)
 					.VAlign(VAlign_Center)
-					.AutoWidth()
 					[
 						SNew(STextBlock)
 						.Text(LOCTEXT("TransitionNodeMissingState", "Invalid State"))
 						.TextStyle(&FStateTreeEditorStyle::Get().GetWidgetStyle<FTextBlockStyle>("Normal.Normal"))
 						.ColorAndOpacity(FStyleColors::White)
 						.Visibility(EVisibility::HitTestInvisible)
+						.AutoWrapText(UStateCentricViewSettings::GetSchemaSettingsDefaultFallback(ViewModel.Get()).bWordWrapTransitionStates)
 					]
 				];
 		}
diff --git a/Engine/Plugins/Runtime/StateTree/Source/StateTreeEditorModule/Private/StateCentricView/View/StateCentricViewUtils.cpp b/Engine/Plugins/Runtime/StateTree/Source/StateTreeEditorModule/Private/StateCentricView/View/StateCentricViewUtils.cpp
--- a/Engine/Plugins/Runtime/StateTree/Source/StateTreeEditorModule/Private/StateCentricView/View/StateCentricViewUtils.cpp
+++ b/Engine/Plugins/Runtime/StateTree/Source/StateTreeEditorModule/Private/StateCentricView/View/StateCentricViewUtils.cpp
@@ -4,6 +4,7 @@
 
 #include "StateCentricView/StateCentricViewSettings.h"
 #include "StateCentricView/ViewModel/StateCentricViewEditorDataExtension.h"
+#include "StateCentricView/View/SImplicitRootTransitionNode.h"
 #include "StateCentricView/View/SParentSelectionNode.h"
 #include "StateCentricView/View/SStateTreeExtendableNode.h"
 #include "StateCentricView/View/STransitionNode.h"
@@ -97,6 +98,36 @@ TSharedRef<SWidget> FStateCentricViewUtils::GenerateTransitionsViewWidget(const
 			}
 		}
 
+		if (ViewState == CentralState && TransitionDirection == EStateTransitionDirection::Out)
+		{
+			if (ViewState->Children.IsEmpty() && ViewState->Transitions.IsEmpty())
+			{
+				TSharedPtr<SWidget> ImplicitRootTransitionWidget = SNew(SImplicitRootTransitionNode, ViewModel, ViewState)
+					.Clipping(EWidgetClipping::OnDemand)
+					.Config({ .InitialLOD = LOD });
+
+				if (UStateCentricViewSettings::GetSchemaSettingsDefaultFallback(&ViewModel.Get()).bJustifyCentralViewTransitions)
+				{
+					TransitionListBody->AddSlot()
+					.FillContentHeight(1.0f)
+					.VAlign(VAlign_Center)
+					.Padding(FMargin(8.0, 6.0, 8.0, 6.0))
+					[
+						ImplicitRootTransitionWidget.ToSharedRef()
+					];
+				}
+				else
+				{
+					TransitionListBody->AddSlot()
+					.AutoHeight()
+					.Padding(FMargin(8.0, 6.0, 8.0, 6.0))
+					[
+						ImplicitRootTransitionWidget.ToSharedRef()
+					];
+				}
+			}
+		}
+
 		for (FStateTreeTransition* Transition : Transitions)
 		{
 			TSharedPtr<SWidget> TransitionNodeWidget;
@@ -235,6 +266,7 @@ void FStateCentricViewUtils::AddChildStateFootnoteExtensions(const EExtendableNo
 				+ SOverlay::Slot()
 				[
 					SNew(SHorizontalBox)
+					.Clipping(EWidgetClipping::OnDemand)
 					+ SHorizontalBox::Slot()
 					.HAlign(HAlign_Left)
 					.VAlign(VAlign_Center)
@@ -247,15 +279,15 @@ void FStateCentricViewUtils::AddChildStateFootnoteExtensions(const EExtendableNo
 						.Visibility(EVisibility::HitTestInvisible)
 					]
 					+ SHorizontalBox::Slot()
-					.HAlign(HAlign_Left)
+					.HAlign(HAlign_Fill)
 					.VAlign(VAlign_Center)
-					.AutoWidth()
 					[
 						SNew(STextBlock)
 						.Text(FText::FromName(ChildState->Name))
 						.TextStyle(&FStateTreeEditorStyle::Get().GetWidgetStyle<FTextBlockStyle>("Normal.Normal"))
 						.ColorAndOpacity(FStyleColors::White)
 						.Visibility(EVisibility::HitTestInvisible)
+						.AutoWrapText(UStateCentricViewSettings::GetSchemaSettingsDefaultFallback(ViewModel.Get()).bWordWrapChildStates)
 					]
 				];
 
@@ -271,6 +303,7 @@ void FStateCentricViewUtils::AddChildStateFootnoteExtensions(const EExtendableNo
 	
 
 		TSharedRef<SExpandableArea> ChildStateArea = SNew(SExpandableArea)
+			.Clipping(EWidgetClipping::OnDemand)
 			.InitiallyCollapsed(false)
 			.BorderImage(FAppStyle::Get().GetBrush("Brushes.Header"))
 			.BodyBorderImage(FAppStyle::Get().GetBrush("Brushes.Recessed"))
