[StateTree] Add a color bar that matches the category color.
#Synced-CL 49067435

diff --git a/Engine/Plugins/Runtime/StateTree/Source/StateTreeEditorModule/Private/Customizations/StateTreeEditorDataDetails.cpp b/Engine/Plugins/Runtime/StateTree/Source/StateTreeEditorModule/Private/Customizations/StateTreeEditorDataDetails.cpp
--- a/Engine/Plugins/Runtime/StateTree/Source/StateTreeEditorModule/Private/Customizations/StateTreeEditorDataDetails.cpp
+++ b/Engine/Plugins/Runtime/StateTree/Source/StateTreeEditorModule/Private/Customizations/StateTreeEditorDataDetails.cpp
@@ -157,6 +157,7 @@ void FStateTreeEditorDataDetails::CustomizeDetails(IDetailLayoutBuilder& DetailB
 			: LOCTEXT("EditorDataDetailsParameters", "Parameters"); 
 		IDetailCategoryBuilder& ParametersCategory = DetailBuilder.EditCategory(TEXT("Parameters"), ParametersDisplayName);
 		ParametersCategory.SetSortOrder(2);
+		ParametersCategory.SetColor(UE::StateTree::Colors::Blue);
 		{
 			// Show parameters as a category.
 			TSharedPtr<IPropertyHandle> PropertyBagParametersProperty = DetailBuilder.GetProperty(GET_MEMBER_NAME_CHECKED(UStateTreeEditorData, RootParameterPropertyBag)); // FInstancedPropertyBag
diff --git a/Engine/Plugins/Runtime/StateTree/Source/StateTreeEditorModule/Private/Customizations/StateTreeEditorNodeUtils.cpp b/Engine/Plugins/Runtime/StateTree/Source/StateTreeEditorModule/Private/Customizations/StateTreeEditorNodeUtils.cpp
--- a/Engine/Plugins/Runtime/StateTree/Source/StateTreeEditorModule/Private/Customizations/StateTreeEditorNodeUtils.cpp
+++ b/Engine/Plugins/Runtime/StateTree/Source/StateTreeEditorModule/Private/Customizations/StateTreeEditorNodeUtils.cpp
@@ -897,6 +897,7 @@ IDetailCategoryBuilder& MakeArrayCategoryHeader(
 {
 	IDetailCategoryBuilder& Category = DetailBuilder.EditCategory(CategoryName, CategoryDisplayName);
 	Category.SetSortOrder(SortOrder);
+	Category.SetColor(IconColor);
 
 	bool bIsNodeArray = false;
 	if (const FArrayProperty* ArrayProperty = CastField<const FArrayProperty>(ArrayPropertyHandle->GetProperty()))
diff --git a/Engine/Plugins/Runtime/StateTree/Source/StateTreeEditorModule/Private/Customizations/StateTreeStateDetails.cpp b/Engine/Plugins/Runtime/StateTree/Source/StateTreeEditorModule/Private/Customizations/StateTreeStateDetails.cpp
--- a/Engine/Plugins/Runtime/StateTree/Source/StateTreeEditorModule/Private/Customizations/StateTreeStateDetails.cpp
+++ b/Engine/Plugins/Runtime/StateTree/Source/StateTreeEditorModule/Private/Customizations/StateTreeStateDetails.cpp
@@ -260,6 +260,7 @@ void FStateTreeStateDetails::CustomizeDetails(IDetailLayoutBuilder& DetailBuilde
 		: LOCTEXT("EditorStateDetailsParameters", "Parameters");
 	IDetailCategoryBuilder& ParametersCategory = DetailBuilder.EditCategory(TEXT("Parameters"), ParametersDisplayName);
 	ParametersCategory.SetSortOrder(1);
+	ParametersCategory.SetColor(UE::StateTree::Colors::Blue);
 	{
 		// Show parameters as a category.
 		ParametersProperty->MarkHiddenByCustomization();
