[StateTree] Fix issue with override linked asset when the state needs reentrance. It used to pick the new override for sustained transition instead of picking the previous override. Add unit test for 2 repro of that same bug.

#Synced-CL 49067435

diff --git a/Engine/Plugins/Runtime/StateTree/Source/StateTreeModule/Private/StateTreeExecutionContext.cpp b/Engine/Plugins/Runtime/StateTree/Source/StateTreeModule/Private/StateTreeExecutionContext.cpp
--- a/Engine/Plugins/Runtime/StateTree/Source/StateTreeModule/Private/StateTreeExecutionContext.cpp
+++ b/Engine/Plugins/Runtime/StateTree/Source/StateTreeModule/Private/StateTreeExecutionContext.cpp
@@ -6906,12 +6906,12 @@ bool FStateTreeExecutionContext::SelectStateInternal(
 	const FStateTreeExecutionFrame* NextParentFrame = nullptr;
 	auto CacheNextFrame = [&NextFrame, &NextParentFrame, &NextFrameID, &Exec, &OutSelectionResult]()
 		{
-			NextFrame = UE::StateTree::ExecutionContext::FindExecutionFrame(NextFrameID, MakeConstArrayView(Exec.ActiveFrames), MakeConstArrayView(OutSelectionResult->TemporaryFrames));
+			NextFrame = FindExecutionFrame(NextFrameID, MakeConstArrayView(Exec.ActiveFrames), MakeConstArrayView(OutSelectionResult->TemporaryFrames));
 			check(NextFrame);
 
 	if (OutSelectionResult->SelectedFrames.Num() > 1)
 	{
-		NextParentFrame = UE::StateTree::ExecutionContext::FindExecutionFrame(OutSelectionResult->SelectedFrames.Last(1), MakeConstArrayView(Exec.ActiveFrames), MakeConstArrayView(OutSelectionResult->TemporaryFrames));
+		NextParentFrame = FindExecutionFrame(OutSelectionResult->SelectedFrames.Last(1), MakeConstArrayView(Exec.ActiveFrames), MakeConstArrayView(OutSelectionResult->TemporaryFrames));
 		check(NextParentFrame);
 	}
 		};
@@ -6922,24 +6922,76 @@ bool FStateTreeExecutionContext::SelectStateInternal(
 
 	UE_STATETREE_DEBUG_SCOPED_STATE_PHASE(this, NextStateHandle, EStateTreeUpdatePhase::StateSelection);
 
+	const bool bIsTargetState = InternalArgs.MissingStatesToReachTarget.Num() <= 1;
+	const bool bShouldPrerequisitesBeChecked = Args.Behavior != ESelectStateBehavior::Forced
+		&& (bIsTargetState || NextState.bCheckPrerequisitesWhenActivatingChildDirectly);
+
+	bool bCreateNewState = true;
+	const bool bIsNextTargetStateInActiveStates = !bNewFrameCreated
+		&& InternalArgs.MissingActiveStates.Num() > 0
+		&& InternalArgs.MissingActiveStates[0].GetStateHandle() == NextStateHandle
+		&& InternalArgs.MissingActiveStates[0].GetFrameID() == NextFrame->FrameID;
+	if (EnumHasAnyFlags(Args.SelectionRules, EStateTreeStateSelectionRules::ReselectedStateCreatesNewStates | EStateTreeStateSelectionRules::CompletedTransitionStatesCreateNewStates))
+	{
+		bCreateNewState = false;
+		if (EnumHasAllFlags(Args.SelectionRules, EStateTreeStateSelectionRules::CompletedTransitionStatesCreateNewStates))
+		{
+			// Request a new state if it doesn't match the previous states or if the state is completed.
+			bCreateNewState = !bIsNextTargetStateInActiveStates || NextFrame->ActiveTasksStatus.GetStatus(NextState).IsCompleted();
+		}
+		if (EnumHasAnyFlags(Args.SelectionRules, EStateTreeStateSelectionRules::ReselectedStateCreatesNewStates))
+		{
+			bCreateNewState = bCreateNewState || bIsTargetState || !bIsNextTargetStateInActiveStates;
+		}
+	}
+	else
+	{
+		bCreateNewState = !bIsNextTargetStateInActiveStates;
+	}
+
 	// Look up linked state overrides
 	const UStateTree* NextLinkedStateAsset = NextState.LinkedAsset;
 	const FInstancedPropertyBag* NextLinkedStateParameterOverride = nullptr;
 	if (NextState.Type == EStateTreeStateType::LinkedAsset)
 	{
-		if (const FStateTreeReference* Override = GetLinkedStateTreeOverrideForTag(NextState.Tag))
+		// If the next state is sustained, then we should get the last valid asset.
+		if (!bCreateNewState)
 		{
-			NextLinkedStateAsset = Override->GetStateTree();
-			NextLinkedStateParameterOverride = &Override->GetParameters();
+			bCreateNewState = true;
+			if (InternalArgs.MissingActiveStates.Num() > 1)
+			{
+				if (const FStateTreeExecutionFrame* NextLinkedFrame = UE::StateTree::ExecutionContext::FindExecutionFrame(InternalArgs.MissingActiveStates[1].GetFrameID(), MakeConstArrayView(Exec.ActiveFrames), MakeConstArrayView(OutSelectionResult->TemporaryFrames)))
+				{
+					check(NextLinkedFrame->StateTree);
+					NextLinkedStateAsset = NextLinkedFrame->StateTree;
+					ensure(NextLinkedStateAsset != NextFrame->StateTree);
+
+					// The StateParameters data also need to match.
+					FCurrentlyProcessedFrameScope FrameScope(*this, NextParentFrame, *NextFrame);
+					FCurrentlyProcessedStateScope NextStateScope(*this, NextStateHandle);
+					FStateTreeDataView NextStateParametersView = GetDataViewOrTemporary(NextParentFrame, *NextFrame, NextState.ParameterDataHandle);
+					const FConstStructView DefaultStateParamsInstanceData = NextLinkedStateAsset->GetDefaultParameters().GetValue();
+					bCreateNewState = DefaultStateParamsInstanceData.GetScriptStruct() != NextStateParametersView.GetStruct();
+				}
+			}
+		}
 
-			STATETREE_LOG(VeryVerbose, TEXT("%hs: In state '%s', overriding linked asset '%s' with '%s'. '%s' using StateTree '%s'."),
-				__FUNCTION__,
-				*GetSafeStateName(*NextFrame, NextStateHandle),
-				*GetFullNameSafe(NextState.LinkedAsset),
-				*GetFullNameSafe(NextLinkedStateAsset),
-				*GetNameSafe(&Owner),
-				*GetFullNameSafe(NextFrame->StateTree)
-			);
+		if (bCreateNewState)
+		{
+			if (const FStateTreeReference* Override = GetLinkedStateTreeOverrideForTag(NextState.Tag))
+			{
+				NextLinkedStateAsset = Override->GetStateTree();
+				NextLinkedStateParameterOverride = &Override->GetParameters();
+
+				STATETREE_LOG(VeryVerbose, TEXT("%hs: In state '%s', overriding linked asset '%s' with '%s'. '%s' using StateTree '%s'."),
+					__FUNCTION__,
+					*GetSafeStateName(*NextFrame, NextStateHandle),
+					*GetFullNameSafe(NextState.LinkedAsset),
+					*GetFullNameSafe(NextLinkedStateAsset),
+					*GetNameSafe(&Owner),
+					*GetFullNameSafe(NextFrame->StateTree)
+				);
+			}
 		}
 	}
 
@@ -6987,10 +7039,6 @@ bool FStateTreeExecutionContext::SelectStateInternal(
 		}
 	}
 
-	const bool bIsTargetState = InternalArgs.MissingStatesToReachTarget.Num() <= 1;
-	const bool bShouldPrerequisitesBeChecked = Args.Behavior != ESelectStateBehavior::Forced
-		&& (bIsTargetState || NextState.bCheckPrerequisitesWhenActivatingChildDirectly);
-
 	// Check if the events are accepted.
 	TArray<FStateTreeSharedEvent, TInlineAllocator<FStateTreeEventQueue::MaxActiveEvents>> StateSelectionEvents;
 	if (NextState.EventDataIndex.IsValid())
@@ -7042,31 +7090,8 @@ bool FStateTreeExecutionContext::SelectStateInternal(
 		StateSelectionEvents.Add(FStateTreeSharedEvent());
 	}
 
-	bool bShouldCreateNewState = true;
-	const bool bIsNextTargetStateInActiveStates = !bNewFrameCreated
-		&& InternalArgs.MissingActiveStates.Num() > 0
-		&& InternalArgs.MissingActiveStates[0].GetStateHandle() == NextStateHandle
-		&& InternalArgs.MissingActiveStates[0].GetFrameID() == NextFrame->FrameID;
-	if (EnumHasAnyFlags(Args.SelectionRules, EStateTreeStateSelectionRules::ReselectedStateCreatesNewStates | EStateTreeStateSelectionRules::CompletedTransitionStatesCreateNewStates))
-	{
-		bShouldCreateNewState = false;
-		if (EnumHasAllFlags(Args.SelectionRules, EStateTreeStateSelectionRules::CompletedTransitionStatesCreateNewStates))
-		{
-			// Request a new state if it doesn't match the previous states or if the state is completed.
-			bShouldCreateNewState = !bIsNextTargetStateInActiveStates || NextFrame->ActiveTasksStatus.GetStatus(NextState).IsCompleted();
-		}
-		if (EnumHasAnyFlags(Args.SelectionRules, EStateTreeStateSelectionRules::ReselectedStateCreatesNewStates))
-		{
-			bShouldCreateNewState = bShouldCreateNewState || bIsTargetState || !bIsNextTargetStateInActiveStates;
-		}
-	}
-	else
-	{
-		bShouldCreateNewState = !bIsNextTargetStateInActiveStates;
-	}
-
 	// Add the state to the selected list. It will be removed if the selection fails.
-	if (bShouldCreateNewState)
+	if (bCreateNewState)
 	{
 		OutSelectionResult->SelectedStates.Emplace(NextFrame->FrameID, FActiveStateID(Storage.GenerateUniqueId()), NextStateHandle);
 		OnExitScope.bStateAdded = true;
@@ -7155,7 +7180,7 @@ bool FStateTreeExecutionContext::SelectStateInternal(
 		}
 
 		const FSelectStateInternalArguments NewInternalArgs = FSelectStateInternalArguments{
-			.MissingActiveStates = bShouldCreateNewState ? TArrayView<const FActiveState>() : InternalArgs.MissingActiveStates.Mid(1),
+			.MissingActiveStates = bCreateNewState ? TArrayView<const FActiveState>() : InternalArgs.MissingActiveStates.Mid(1),
 			.MissingSourceFrameID = FActiveFrameID(),
 			.MissingSourceStates = TArrayView<const FActiveState>(),
 			.MissingStatesToReachTarget = bIsTargetState ? TArrayView<const FStateHandleContext>() : InternalArgs.MissingStatesToReachTarget.Mid(1)
@@ -7163,12 +7188,12 @@ bool FStateTreeExecutionContext::SelectStateInternal(
 		if (NextState.Type == EStateTreeStateType::Linked)
 		{
 			// MissingStatesToReachTarget can include a linked state and the frame needs to be constructed (if needed).
-			OnExitScope.bSucceededToSelectState = SelectStateInternal_Linked(Args, NewInternalArgs, OutSelectionResult, NextFrame->StateTree, NextState, bShouldCreateNewState);
+			OnExitScope.bSucceededToSelectState = SelectStateInternal_Linked(Args, NewInternalArgs, OutSelectionResult, NextFrame->StateTree, NextState, bCreateNewState);
 		}
 		else if (NextState.Type == EStateTreeStateType::LinkedAsset)
 		{
 			// MissingStatesToReachTarget can include a linked asset state and the frame needs to be constructed (if needed).
-			OnExitScope.bSucceededToSelectState = SelectStateInternal_LinkedAsset(Args, NewInternalArgs, OutSelectionResult, NextFrame->StateTree, NextState, NextLinkedStateAsset, bShouldCreateNewState);
+			OnExitScope.bSucceededToSelectState = SelectStateInternal_LinkedAsset(Args, NewInternalArgs, OutSelectionResult, NextFrame->StateTree, NextState, NextLinkedStateAsset, bCreateNewState);
 		}
 		else if (!bIsTargetState)
 		{
@@ -7333,7 +7358,7 @@ namespace UE::StateTree::ExecutionContext::Private
 	}
 
 	FStateTreeExecutionFrame* SelectedFrameLinkedFrame(FStateTreeExecutionState& Exec,
-		const bool bShouldCreateNewState,
+		const bool bCreateNewState,
 		TConstArrayView<const UE::StateTree::FActiveState> MatchingActiveStates,
 		const FExecutionFrameHandle& LinkStateFrameHandle,
 		const EStateTreeStateSelectionRules StateSelectionRules
@@ -7342,7 +7367,7 @@ namespace UE::StateTree::ExecutionContext::Private
 		check(LinkStateFrameHandle.IsValid());
 
 		FStateTreeExecutionFrame* Result = nullptr;
-		if (!bShouldCreateNewState)
+		if (!bCreateNewState)
 		{
 			// Get the next frame ID
 			if (ensure(MatchingActiveStates.Num() >= 1))
@@ -7371,7 +7396,7 @@ bool FStateTreeExecutionContext::SelectStateInternal_Linked(
 	const TSharedRef<FSelectStateResult>& OutSelectionResult,
 	TNotNull<const UStateTree*> NextStateTree,
 	const FCompactStateTreeState& TargetState,
-	bool bShouldCreateNewState)
+	bool bCreateNewState)
 {
 	using namespace UE::StateTree;
 	using namespace UE::StateTree::ExecutionContext;
@@ -7435,7 +7460,7 @@ bool FStateTreeExecutionContext::SelectStateInternal_Linked(
 	}
 
 	// Do we have an existing frame.
-	FStateTreeExecutionFrame* SelectedFrame = SelectedFrameLinkedFrame(Exec, bShouldCreateNewState, InternalArgs.MissingActiveStates, LinkStateFrameHandle, Args.SelectionRules);
+	FStateTreeExecutionFrame* SelectedFrame = SelectedFrameLinkedFrame(Exec, bCreateNewState, InternalArgs.MissingActiveStates, LinkStateFrameHandle, Args.SelectionRules);
 	if (SelectedFrame && !SelectedFrame->HasRoot(LinkStateFrameHandle))
 	{
 		if (Args.Behavior == ESelectStateBehavior::Forced)
@@ -7515,7 +7540,7 @@ bool FStateTreeExecutionContext::SelectStateInternal_LinkedAsset(
 	TNotNull<const UStateTree*> NextStateTree,
 	const FCompactStateTreeState& NextState,
 	const UStateTree* NextLinkedStateAsset,
-	bool bShouldCreateNewState)
+	bool bCreateNewState)
 {
 	using namespace UE::StateTree;
 	using namespace UE::StateTree::ExecutionContext;
@@ -7602,7 +7627,7 @@ bool FStateTreeExecutionContext::SelectStateInternal_LinkedAsset(
 	// Do we have an existing frame.
 	// Do not use the transition override selection rules. A transition outside the frame should not impact the current frame.
 	const EStateTreeStateSelectionRules StateSelectionRules = LinkStateFrameHandle.GetStateTree()->GetStateSelectionRules();
-	FStateTreeExecutionFrame* SelectedFrame = SelectedFrameLinkedFrame(Exec, bShouldCreateNewState, InternalArgs.MissingActiveStates, LinkStateFrameHandle, StateSelectionRules);
+	FStateTreeExecutionFrame* SelectedFrame = SelectedFrameLinkedFrame(Exec, bCreateNewState, InternalArgs.MissingActiveStates, LinkStateFrameHandle, StateSelectionRules);
 	if (SelectedFrame && !SelectedFrame->HasRoot(LinkStateFrameHandle))
 	{
 		if (Args.Behavior == ESelectStateBehavior::Forced)
diff --git a/Engine/Plugins/Runtime/StateTree/Source/StateTreeTestSuite/Private/StateTreeLinkedStateTest.cpp b/Engine/Plugins/Runtime/StateTree/Source/StateTreeTestSuite/Private/StateTreeLinkedStateTest.cpp
--- a/Engine/Plugins/Runtime/StateTree/Source/StateTreeTestSuite/Private/StateTreeLinkedStateTest.cpp
+++ b/Engine/Plugins/Runtime/StateTree/Source/StateTreeTestSuite/Private/StateTreeLinkedStateTest.cpp
@@ -1818,6 +1818,265 @@ struct FStateTreeTest_Linked_FinishGlobalTasks : FStateTreeTestBase
 };
 IMPLEMENT_STATE_TREE_INSTANT_TEST(FStateTreeTest_Linked_FinishGlobalTasks, "System.StateTree.LinkedAsset.FinishGlobalTasks");
 
+struct FStateTreeTest_OverrideLinkedAsset : FStateTreeTestBase
+{
+	// test override while inside Tree1.Root.State1.Tree2.Root.State1, the override changes.
+	virtual bool InstantTest() override
+	{
+		FStateTreeCompilerLog Log;
+
+		// Tree1
+		//  Root1 (tick => override then request new transition to A1)
+		//   State1 (linked to Tree2)
+		// Tree2 GlobalTask2 => 2 ticks
+		//  Root
+			// State1 => 1 tick to completion
+			// State2 => print
+		// Tree3 GlobalTask2 => 2 ticks
+		//  Root Task => 1 tick
+
+		// Asset 3
+		UStateTree& StateTree3 = NewStateTree();
+		{
+			UStateTreeEditorData& EditorData = *Cast<UStateTreeEditorData>(StateTree3.EditorData);
+			GetRootPropertyBag(EditorData).AddProperty("B", EPropertyBagPropertyType::Int32);
+
+			UStateTreeState& Root = EditorData.AddSubTree(FName(TEXT("Root")));
+			TStateTreeEditorNode<FTestTask_PrintValue>& RootTask1 = Root.AddTask<FTestTask_PrintValue>(FName(TEXT("Tree3RootTask1")));
+			RootTask1.GetInstanceData().Value = -3;
+
+			EditorData.AddPropertyBinding(FPropertyBindingPath(EditorData.GetRootParametersGuid(), "B"), FPropertyBindingPath(RootTask1.ID, TEXT("Value")));
+
+			{
+				FStateTreeCompiler Compiler(Log);
+				const bool bResult = Compiler.Compile(StateTree3);
+				AITEST_TRUE(TEXT("StateTree3 should get compiled"), bResult);
+			}
+		}
+
+		// Asset 2
+		UStateTree& StateTree2 = NewStateTree();
+		{
+			UStateTreeEditorData& EditorData = *Cast<UStateTreeEditorData>(StateTree2.EditorData);
+			GetRootPropertyBag(EditorData).AddProperty("A", EPropertyBagPropertyType::Int32);
+
+			UStateTreeState& Root2 = EditorData.AddSubTree(FName(TEXT("Root")));
+
+			UStateTreeState& State1 = Root2.AddChildState("State1");
+			TStateTreeEditorNode<FTestTask_Stand>& State1Task1 = State1.AddTask<FTestTask_Stand>(FName(TEXT("Tree2State1Task1")));
+			TStateTreeEditorNode<FTestTask_PrintValue>& State1Task2 = State1.AddTask<FTestTask_PrintValue>(FName(TEXT("Tree2State1Task2")));
+			State1.AddTransition(EStateTreeTransitionTrigger::OnStateCompleted, EStateTreeTransitionType::NextState);
+
+			EditorData.AddPropertyBinding(FPropertyBindingPath(EditorData.GetRootParametersGuid(), "A"), FPropertyBindingPath(State1Task2.ID, TEXT("Value")));
+
+			UStateTreeState& State2 = Root2.AddChildState("State2");
+			TStateTreeEditorNode<FTestTask_PrintValue>& State2Task1 = State2.AddTask<FTestTask_PrintValue>(FName(TEXT("Tree2State2Task1")));
+
+			EditorData.AddPropertyBinding(FPropertyBindingPath(EditorData.GetRootParametersGuid(), "A"), FPropertyBindingPath(State2Task1.ID, TEXT("Value")));
+
+			{
+				FStateTreeCompiler Compiler(Log);
+				const bool bResult = Compiler.Compile(StateTree2);
+				AITEST_TRUE(TEXT("StateTree2 should get compiled"), bResult);
+			}
+		}
+
+		// Main asset
+		UStateTree& StateTree = NewStateTree();
+		{
+			UStateTreeEditorData& EditorData = *Cast<UStateTreeEditorData>(StateTree.EditorData);
+
+			GetRootPropertyBag(EditorData).AddProperty("Int1Param", EPropertyBagPropertyType::Int32);
+
+			UStateTreeState& Root = EditorData.AddSubTree(FName(TEXT("Root1")));
+			UStateTreeState& State1 = Root.AddChildState(FName(TEXT("State1")), EStateTreeStateType::LinkedAsset);
+			State1.Tag = GetTestTag1();
+			State1.SetLinkedStateAsset(&StateTree2);
+
+			EditorData.AddPropertyBinding(FPropertyBindingPath(EditorData.GetRootParametersGuid(), "Int1Param"), FPropertyBindingPath(State1.Parameters.ID, "A"));
+
+			{
+				FStateTreeCompiler Compiler(Log);
+				const bool bResult = Compiler.Compile(StateTree);
+				AITEST_TRUE(TEXT("StateTree should get compiled"), bResult);
+			}
+		}
+
+		{
+			FStateTreeInstanceData InstanceData;
+			FTestStateTreeExecutionContext Exec(StateTree, StateTree, InstanceData);
+			const bool bInitSucceeded = Exec.IsValid();
+			AITEST_TRUE(TEXT("StateTree should init"), bInitSucceeded);
+
+			FInstancedPropertyBag GlobalParameters = StateTree.GetDefaultParameters();
+			GlobalParameters.SetValueInt32("Int1Param", 11);
+			{
+				EStateTreeRunStatus Status = Exec.Start(GlobalParameters.GetValue());
+				AITEST_EQUAL("Start should complete with Running", Status, EStateTreeRunStatus::Running);
+				AITEST_TRUE(TEXT("StateTree should tick Tree2State1Task1"), Exec.Expect("Tree2State1Task1", TEXT("EnterState")));
+				AITEST_TRUE(TEXT("StateTree should tick Tree2State1Task2"), Exec.Expect("Tree2State1Task2", TEXT("EnterState11")));
+				Exec.LogClear();
+			}
+			// Set the overrides
+			{
+				FStateTreeReference StateTreeReference;
+				StateTreeReference.SetStateTree(&StateTree3);
+				FStateTreeReferenceOverrides Overrides;
+				Overrides.AddOverride(GetTestTag1(), StateTreeReference);
+				Exec.SetLinkedStateTreeOverrides(Overrides);
+			}
+			{
+				// Tree2State1Task1 completes and transition to Tree2State2Task1
+				EStateTreeRunStatus Status = Exec.Tick(0.1f);
+				AITEST_EQUAL("Tick should complete with Running", Status, EStateTreeRunStatus::Running);
+				AITEST_TRUE(TEXT("StateTree should tick Tree2State2Task1"), Exec.Expect("Tree2State2Task1", TEXT("EnterState11")));
+				Exec.LogClear();
+			}
+			{
+				Exec.Stop();
+				Exec.LogClear();
+			}
+			//@TODO Crash when the override (tree3) has a different global parameter struct than the compiled version.
+			//{
+			//	GlobalParameters.SetValueInt32("Int1Param", 33);
+			//	EStateTreeRunStatus Status = Exec.Start(GlobalParameters.GetValue());
+			//	AITEST_EQUAL("Start should complete with Running", Status, EStateTreeRunStatus::Running);
+			//	AITEST_TRUE(TEXT("StateTree should tick Tree3RootTask1"), Exec.Expect("Tree3RootTask1", TEXT("EnterState33")));
+			//	Exec.LogClear();
+			//}
+			Exec.Stop();
+		}
+
+		return true;
+	}
+};
+IMPLEMENT_STATE_TREE_INSTANT_TEST(FStateTreeTest_OverrideLinkedAsset, "System.StateTree.LinkedAsset.Override");
+
+struct FStateTreeTest_ReenterLinkedAssetWithDifferentAsset : FStateTreeTestBase
+{
+	virtual bool InstantTest() override
+	{
+		//Main Tree
+		//	Root
+		//		StateMain [LinkedAsset with tag] -> StateMain (OnTick)
+		//Tree A
+		//	Root (TaskA)
+		//Tree B
+		//	Root (TaskB)
+		// StateMain starts with Tree A set and gets entered.
+		// StateMain's Linked Asset is overridden to Tree B and transition on StateMain to self is attempted.
+		// StateMain should be reentered with the Tree A asset or Tree B depending on the selection rules.
+
+		// Asset A
+		UStateTree& StateTreeA = NewStateTree();
+		{
+			UStateTreeEditorData& EditorDataA = *Cast<UStateTreeEditorData>(StateTreeA.EditorData);
+			GetRootPropertyBag(EditorDataA).AddProperty("A", EPropertyBagPropertyType::Bool);
+			UStateTreeState& RootA = EditorDataA.AddSubTree(FName(TEXT("RootA")));
+			RootA.AddTask<FTestTask_Stand>(FName(TEXT("TaskA")));
+		}
+		// Asset B
+		UStateTree& StateTreeB = NewStateTree();
+		{
+			UStateTreeEditorData& EditorDataB = *Cast<UStateTreeEditorData>(StateTreeB.EditorData);
+			GetRootPropertyBag(EditorDataB).AddProperty("B", EPropertyBagPropertyType::Int32);
+			UStateTreeState& RootB = EditorDataB.AddSubTree(FName(TEXT("RootB")));
+			RootB.AddTask<FTestTask_Stand>(FName(TEXT("TaskB")));
+		}
+		// Main asset
+		const FGameplayTag StateTag = GetTestTag1();
+		UStateTree& StateTree = NewStateTree();
+		{
+			UStateTreeEditorData& EditorData = *Cast<UStateTreeEditorData>(StateTree.EditorData);
+			UStateTreeState& Root = EditorData.AddSubTree(FName(TEXT("Root")));
+			UStateTreeState& StateMain = Root.AddChildState(FName(TEXT("StateMain")), EStateTreeStateType::LinkedAsset);
+			StateMain.AddTransition(EStateTreeTransitionTrigger::OnTick, EStateTreeTransitionType::GotoState, &StateMain);
+			StateMain.Tag = StateTag;
+		}
+		const TCHAR* EnterStateStr(TEXT("EnterState"));
+		const TCHAR* ExitStateStr(TEXT("ExitState"));
+
+		constexpr int32 MaxRules = 4;
+		auto MakeStateSelectionRule = [](int32 Index)
+			{
+				EStateTreeStateSelectionRules Rule = EStateTreeStateSelectionRules::None;
+				if ((Index % 2) == 1)
+				{
+					Rule |= EStateTreeStateSelectionRules::CompletedTransitionStatesCreateNewStates;
+					Rule |= EStateTreeStateSelectionRules::CompletedStateBeforeTransitionSourceFailsTransition;
+				}
+				if (Index >= 2)
+				{
+					Rule |= EStateTreeStateSelectionRules::ReselectedStateCreatesNewStates;
+				}
+				return Rule;
+			};
+
+
+		for (int32 Index = 0; Index < MaxRules; ++Index)
+		{
+			const EStateTreeStateSelectionRules StateSelectionRules = MakeStateSelectionRule(Index);
+			auto CompileTree = [StateSelectionRules](UStateTree& StateTree)
+				{
+					StateTree.ResetCompiled();
+					UStateTreeEditorData* EditorData = CastChecked<UStateTreeEditorData>(StateTree.EditorData);
+					UStateTreeTestSchema* Schema = CastChecked<UStateTreeTestSchema>(EditorData->Schema);
+					Schema->SetStateSelectionRules(StateSelectionRules);
+
+					FStateTreeCompilerLog Log;
+					FStateTreeCompiler Compiler(Log);
+					return Compiler.Compile(&StateTree);
+				};
+
+			AITEST_TRUE("StateTree2 should get compiled", CompileTree(StateTreeA));
+			AITEST_TRUE("StateTree1 should get compiled", CompileTree(StateTreeB));
+			AITEST_TRUE("StateTree3 should get compiled", CompileTree(StateTree));
+
+			{
+				EStateTreeRunStatus Status = EStateTreeRunStatus::Unset;
+				FStateTreeInstanceData InstanceData;
+				FTestStateTreeExecutionContext Exec(StateTree, StateTree, InstanceData);
+				{
+					FStateTreeReferenceOverrides Overrides;
+					FStateTreeReference Ref;
+					Ref.SetStateTree(&StateTreeA);
+					Overrides.AddOverride(StateTag, Ref);
+					Exec.SetLinkedStateTreeOverrides(Overrides);
+				}
+				const bool bInitSucceeded = Exec.IsValid();
+				AITEST_TRUE(TEXT("StateTree should init"), bInitSucceeded);
+				Status = Exec.Start();
+				AITEST_EQUAL(TEXT("Start should complete with Running"), Status, EStateTreeRunStatus::Running);
+				AITEST_TRUE(TEXT("StateTree should enter TaskA"), Exec.Expect("TaskA", EnterStateStr));
+				Exec.LogClear();
+				Exec.Tick(0.1f);
+				AITEST_TRUE(TEXT("StateTree should enter TaskA"), Exec.Expect("TaskA", EnterStateStr));
+				Exec.LogClear();
+				{
+					FStateTreeReferenceOverrides Overrides;
+					FStateTreeReference Ref;
+					Ref.SetStateTree(&StateTreeB);
+					Overrides.AddOverride(StateTag, Ref);
+					Exec.SetLinkedStateTreeOverrides(Overrides);
+				}
+				Exec.Tick(0.1f);
+				if (EnumHasAllFlags(StateSelectionRules, EStateTreeStateSelectionRules::ReselectedStateCreatesNewStates))
+				{
+					AITEST_TRUE(TEXT("StateTree should enter TaskB"), Exec.Expect("TaskB", EnterStateStr));
+				}
+				else
+				{
+					AITEST_TRUE(TEXT("StateTree should enter TaskA"), Exec.Expect("TaskA", EnterStateStr));
+				}
+				Exec.Stop();
+			}
+		}
+		return true;
+	}
+};
+IMPLEMENT_STATE_TREE_INSTANT_TEST(FStateTreeTest_ReenterLinkedAssetWithDifferentAsset, "System.StateTree.LinkedAsset.ReenterWithDifferentAsset");
+
+
 } // namespace UE::StateTree::Tests
 
 UE_ENABLE_OPTIMIZATION_SHIP
