[StateTree] Access ExecutionRuntimeData from async task.

#Synced-CL 49067435

diff --git a/Engine/Plugins/Runtime/StateTree/Source/StateTreeModule/Private/StateTreeAsyncExecutionContext.cpp b/Engine/Plugins/Runtime/StateTree/Source/StateTreeModule/Private/StateTreeAsyncExecutionContext.cpp
--- a/Engine/Plugins/Runtime/StateTree/Source/StateTreeModule/Private/StateTreeAsyncExecutionContext.cpp
+++ b/Engine/Plugins/Runtime/StateTree/Source/StateTreeModule/Private/StateTreeAsyncExecutionContext.cpp
@@ -320,14 +320,30 @@ FStateTreeDataView TStateTreeStrongExecutionContext<bWithWriteAccess>::GetInstan
 	Async::FActivePathInfo ActivePath = GetActivePathInfo();
 	if (ActivePath.IsValid())
 	{
-		const FStateTreeNodeBase& Node = ActivePath.GetNode();
-		FStateTreeDataView InstanceDataView = InstanceData::GetDataViewOrTemporary(
+		return InstanceData::GetDataViewOrTemporary(
 			*Storage,
 			/* SharedInstanceStorage */ nullptr,
 			*ActivePath.Frame,
 			ActivePath.GetNode().InstanceDataHandle);
+	}
+
+	return {};
+}
+
+template <bool bWithWriteAccess>
+FStateTreeDataView TStateTreeStrongExecutionContext<bWithWriteAccess>::GetExecutionRuntimeDataPtrInternal() const
+{
+	using namespace UE::StateTree;
 
-		return InstanceDataView;
+	Async::FActivePathInfo ActivePath = GetActivePathInfo();
+	if (ActivePath.IsValid())
+	{
+		const FStateTreeDataHandle DataHandle = FStateTreeDataHandle(EStateTreeDataSourceType::ExecutionRuntimeDataAny, ActivePath.GetNode().ExecutionRuntimeTemplateIndex.Get());
+		return InstanceData::GetDataViewOrTemporary(
+			*Storage,
+			/* SharedInstanceStorage */ nullptr,
+			*ActivePath.Frame,
+			DataHandle);
 	}
 
 	return {};
diff --git a/Engine/Plugins/Runtime/StateTree/Source/StateTreeModule/Private/StateTreeExecutionContext.cpp b/Engine/Plugins/Runtime/StateTree/Source/StateTreeModule/Private/StateTreeExecutionContext.cpp
--- a/Engine/Plugins/Runtime/StateTree/Source/StateTreeModule/Private/StateTreeExecutionContext.cpp
+++ b/Engine/Plugins/Runtime/StateTree/Source/StateTreeModule/Private/StateTreeExecutionContext.cpp
@@ -2746,22 +2746,8 @@ FStateTreeDataView FStateTreeExecutionContext::GetExecutionRuntimeDataView() con
 		return FStateTreeDataView();
 	}
 
-	const UE::StateTree::InstanceData::FInstanceContainer& Container = Storage.GetExecutionRuntimeData();
-	const int32 ActiveIndex = CurrentlyProcessedFrame->ExecutionRuntimeIndexBase.Get() + CurrentNode->ExecutionRuntimeTemplateIndex.Get();
-
-	if (!ensureMsgf(Container.IsValidIndex(ActiveIndex), TEXT("Invalid execution runtime data index.")))
-	{
-		return FStateTreeDataView();
-	}
-
-	if (Container.IsObject(ActiveIndex))
-	{
-		return Container.GetMutableObject(ActiveIndex);
-	}
-	else
-	{
-		return const_cast<UE::StateTree::InstanceData::FInstanceContainer&>(Container).GetMutableStruct(ActiveIndex);
-	}
+	const FStateTreeDataHandle DataHandle = FStateTreeDataHandle(EStateTreeDataSourceType::ExecutionRuntimeDataAny, CurrentNode->ExecutionRuntimeTemplateIndex.Get());
+	return UE::StateTree::InstanceData::GetDataView(Storage, CurrentlyProcessedSharedInstanceStorage, *CurrentlyProcessedFrame, DataHandle);
 }
 
 EStateTreeRunStatus FStateTreeExecutionContext::ForceTransition(const FRecordedStateTreeTransitionResult& RecordedTransition)
diff --git a/Engine/Plugins/Runtime/StateTree/Source/StateTreeModule/Private/StateTreeInstanceData.cpp b/Engine/Plugins/Runtime/StateTree/Source/StateTreeModule/Private/StateTreeInstanceData.cpp
--- a/Engine/Plugins/Runtime/StateTree/Source/StateTreeModule/Private/StateTreeInstanceData.cpp
+++ b/Engine/Plugins/Runtime/StateTree/Source/StateTreeModule/Private/StateTreeInstanceData.cpp
@@ -192,6 +192,16 @@ namespace UE::StateTree::InstanceData::Private
 		case EStateTreeDataSourceType::SharedInstanceDataObject:
 			return true;
 
+		case EStateTreeDataSourceType::ExecutionRuntimeData:
+		case EStateTreeDataSourceType::ExecutionRuntimeDataObject:
+		case EStateTreeDataSourceType::ExecutionRuntimeDataAny:
+		{
+			const int32 ActiveIndex = CurrentFrame.ExecutionRuntimeIndexBase.Get() + Handle.GetIndex();
+			return Handle.HasValidIndex()
+				&& CurrentFrame.ExecutionRuntimeIndexBase.IsValid() 
+				&& InstanceStorage.GetExecutionRuntimeData().IsValidIndex(ActiveIndex);
+		}
+
 		case EStateTreeDataSourceType::GlobalParameterData:
 		{
 			// Defined as State Parameter in one of its parent frame
@@ -334,6 +344,30 @@ namespace UE::StateTree::InstanceData
 			check(SharedInstanceStorage);
 			return SharedInstanceStorage->GetMutableObject(Handle.GetIndex());
 
+		case EStateTreeDataSourceType::ExecutionRuntimeData:
+		{
+			const int32 ActiveIndex = CurrentFrame.ExecutionRuntimeIndexBase.Get() + Handle.GetIndex();
+			return InstanceStorage.GetExecutionRuntimeData().GetMutableStruct(ActiveIndex);
+		}
+		case EStateTreeDataSourceType::ExecutionRuntimeDataObject:
+		{
+			const int32 ActiveIndex = CurrentFrame.ExecutionRuntimeIndexBase.Get() + Handle.GetIndex();
+			return InstanceStorage.GetExecutionRuntimeData().GetMutableObject(ActiveIndex);
+		}
+		case EStateTreeDataSourceType::ExecutionRuntimeDataAny:
+		{
+			const int32 ActiveIndex = CurrentFrame.ExecutionRuntimeIndexBase.Get() + Handle.GetIndex();
+			UE::StateTree::InstanceData::FInstanceContainer& Container = InstanceStorage.GetExecutionRuntimeData();
+			if (Container.IsObject(ActiveIndex))
+			{
+				return Container.GetMutableObject(ActiveIndex);
+			}
+			else
+			{
+				return Container.GetMutableStruct(ActiveIndex);
+			}
+		}
+
 		case EStateTreeDataSourceType::GlobalParameterData:
 		{
 			if (CurrentFrame.GlobalParameterDataHandle.HasValidIndex())
diff --git a/Engine/Plugins/Runtime/StateTree/Source/StateTreeModule/Private/StateTreeTypes.cpp b/Engine/Plugins/Runtime/StateTree/Source/StateTreeModule/Private/StateTreeTypes.cpp
--- a/Engine/Plugins/Runtime/StateTree/Source/StateTreeModule/Private/StateTreeTypes.cpp
+++ b/Engine/Plugins/Runtime/StateTree/Source/StateTreeModule/Private/StateTreeTypes.cpp
@@ -174,6 +174,8 @@ FStateTreeDataHandle FStateTreeDataHandle::ToObjectSource() const
 		return FStateTreeDataHandle(EStateTreeDataSourceType::EvaluationScopeInstanceDataObject, Index, StateHandle);
 	case EStateTreeDataSourceType::ExecutionRuntimeData:
 		return FStateTreeDataHandle(EStateTreeDataSourceType::ExecutionRuntimeDataObject, Index, StateHandle);
+	case EStateTreeDataSourceType::ExecutionRuntimeDataAny:
+		return FStateTreeDataHandle(EStateTreeDataSourceType::ExecutionRuntimeDataObject, Index, StateHandle);
 	default:
 		return *this;
 	}
@@ -193,6 +195,8 @@ FStateTreeDataHandle FStateTreeDataHandle::ToStructSource() const
 		return FStateTreeDataHandle(EStateTreeDataSourceType::EvaluationScopeInstanceData, Index, StateHandle);
 	case EStateTreeDataSourceType::ExecutionRuntimeDataObject:
 		return FStateTreeDataHandle(EStateTreeDataSourceType::ExecutionRuntimeData, Index, StateHandle);
+	case EStateTreeDataSourceType::ExecutionRuntimeDataAny:
+		return FStateTreeDataHandle(EStateTreeDataSourceType::ExecutionRuntimeData, Index, StateHandle);
 	default:
 		return *this;
 	}
diff --git a/Engine/Plugins/Runtime/StateTree/Source/StateTreeModule/Public/StateTreeAsyncExecutionContext.h b/Engine/Plugins/Runtime/StateTree/Source/StateTreeModule/Public/StateTreeAsyncExecutionContext.h
--- a/Engine/Plugins/Runtime/StateTree/Source/StateTreeModule/Public/StateTreeAsyncExecutionContext.h
+++ b/Engine/Plugins/Runtime/StateTree/Source/StateTreeModule/Public/StateTreeAsyncExecutionContext.h
@@ -162,6 +162,25 @@ struct TStateTreeStrongExecutionContext
 	template<typename T>
 	std::conditional_t<bWithWriteAccess, T*, const T*> GetInstanceDataPtr() const&& = delete;
 
+	/**
+	 * Get the Execution Runtime Data of recorded node. Only callable on lvalue because it would have lost the access track.
+	 * @return nullptr if the frame or state containing the node is no longer active.
+	 */
+	template<typename T>
+	std::conditional_t<bWithWriteAccess, T*, const T*> GetExecutionRuntimeDataPtr() const&
+	{
+		FStateTreeDataView DataView = GetExecutionRuntimeDataPtrInternal();
+		if (DataView.IsValid() && ensure(DataView.GetStruct() == T::StaticStruct()))
+		{
+			return static_cast<T*>(DataView.GetMutableMemory());
+		}
+
+		return nullptr;
+	}
+
+	template<typename T>
+	std::conditional_t<bWithWriteAccess, T*, const T*> GetExecutionRuntimeDataPtr() const&& = delete;
+
 	/**
 	 * Get the info of active frame, state and node this Context is based on
 	 * @return an invalid FActivePathInfo if the frame or state containing the node is no longer active.
@@ -186,6 +205,7 @@ private:
 	}
 
 	UE_API FStateTreeDataView GetInstanceDataPtrInternal() const;
+	UE_API FStateTreeDataView GetExecutionRuntimeDataPtrInternal() const;
 
 	UE_API static void ScheduleNextTick(TNotNull<UObject*> Owner, TNotNull<const UStateTree*> RootStateTree, FStateTreeInstanceStorage& Storage, UE::StateTree::ETickReason Reason) requires bWithWriteAccess;
 
diff --git a/Engine/Plugins/Runtime/StateTree/Source/StateTreeModule/Public/StateTreeTypes.h b/Engine/Plugins/Runtime/StateTree/Source/StateTreeModule/Public/StateTreeTypes.h
--- a/Engine/Plugins/Runtime/StateTree/Source/StateTreeModule/Public/StateTreeTypes.h
+++ b/Engine/Plugins/Runtime/StateTree/Source/StateTreeModule/Public/StateTreeTypes.h
@@ -333,6 +333,9 @@ enum class EStateTreeDataSourceType : uint8
 	/** Node instance data for the duration of the execution. */
 	ExecutionRuntimeDataObject,
 
+	/** Node instance data for the duration of the execution. Discover at runtime if it's a struct or an object. */
+	ExecutionRuntimeDataAny,
+
 	/** Context Data, Tree Parameters */
 	ContextData,
 
@@ -996,6 +999,13 @@ struct FCompactStateTreeState
 	/** If set to true, the state has a custom tick rate. */
 	UPROPERTY()
 	uint8 bHasCustomTickRate : 1 = false;
+
+	/**
+	 * If set to true, linked asset can be override at runtime.
+	 * Bindings may prevent a linked asset to be override.
+	 */
+	UPROPERTY()
+	uint8 bLinkedAssetOverrideAllowed : 1 = true;
 };
 
 USTRUCT()
