[StateTree] Backout part of CL 49165271. It contained code from another test changelist.

#rnx

#Synced-CL 49204654

diff --git a/Engine/Plugins/Runtime/StateTree/Source/StateTreeModule/Public/StateTreeTypes.h b/Engine/Plugins/Runtime/StateTree/Source/StateTreeModule/Public/StateTreeTypes.h
--- a/Engine/Plugins/Runtime/StateTree/Source/StateTreeModule/Public/StateTreeTypes.h
+++ b/Engine/Plugins/Runtime/StateTree/Source/StateTreeModule/Public/StateTreeTypes.h
@@ -999,13 +999,6 @@ struct FCompactStateTreeState
 	/** If set to true, the state has a custom tick rate. */
 	UPROPERTY()
 	uint8 bHasCustomTickRate : 1 = false;
-
-	/**
-	 * If set to true, linked asset can be override at runtime.
-	 * Bindings may prevent a linked asset to be override.
-	 */
-	UPROPERTY()
-	uint8 bLinkedAssetOverrideAllowed : 1 = true;
 };
 
 USTRUCT()
