[StateTree] Hide utility consideration when they are not usable. The parent need to have a seleciton behavior that uses them.

#Synced-CL 49204654

diff --git a/Engine/Plugins/Runtime/StateTree/Source/StateTreeEditorModule/Private/Customizations/StateTreeStateDetails.cpp b/Engine/Plugins/Runtime/StateTree/Source/StateTreeEditorModule/Private/Customizations/StateTreeStateDetails.cpp
--- a/Engine/Plugins/Runtime/StateTree/Source/StateTreeEditorModule/Private/Customizations/StateTreeStateDetails.cpp
+++ b/Engine/Plugins/Runtime/StateTree/Source/StateTreeEditorModule/Private/Customizations/StateTreeStateDetails.cpp
@@ -60,22 +60,22 @@ void FStateTreeStateDetails::CustomizeDetails(IDetailLayoutBuilder& DetailBuilde
 	}
 
 	// Find StateTreeState associated with this panel.
-	TWeakObjectPtr<UStateTreeState> WeakState;
+	UStateTreeState* State = nullptr;
 	for (TWeakObjectPtr<UObject>& WeakObject : Objects)
 	{
 		if (UObject* Object = WeakObject.Get())
 		{
-			if (UStateTreeState* State = Cast<UStateTreeState>(Object))
+			if (UStateTreeState* StateTreeState = Cast<UStateTreeState>(Object))
 			{
-				WeakState = State;
+				State = StateTreeState;
 				break;
 			}
 		}
 	}
+	WeakState = State;
 
 	const UStateTreeSchema* Schema = EditorData ? EditorData->Schema.Get() : nullptr;
 	const FString SchemaPath = Schema ? Schema->GetClass()->GetPathName() : FString();
-	TWeakObjectPtr<UStateTreeEditorData> WeakEditorData = EditorData;
 
 	TSharedPtr<FStateTreeViewModel> ViewModel;
 	if (UStateTreeEditingSubsystem* StateTreeEditingSubsystem = EditorData != nullptr ? GEditor->GetEditorSubsystem<UStateTreeEditingSubsystem>() : nullptr)
@@ -360,25 +360,31 @@ void FStateTreeStateDetails::CustomizeDetails(IDetailLayoutBuilder& DetailBuilde
 	ConsiderationsProperty->MarkHiddenByCustomization();
 	if (Schema && Schema->AllowUtilityConsiderations())
 	{
-		const FName UtilityCategoryName(TEXT("Selection Utility"));
-		int32 NumConsiderations = WeakState.IsValid() ? WeakState->Considerations.Num() : 0;
-		FText UtilityCategoryText = EnumHasAllFlags(GetDefault<UStateTreeEditorUserSettings>()->GetDetailsViewDisplayCount(), EStateTreeEditorUserSettingsNodeType::Utility)
-			? FText::Format(LOCTEXT("StateDetailsSelectionUtilityCounted", "Selection Utility ({0})"), NumConsiderations)
-			: LOCTEXT("StateDetailsSelectionUtility", "Selection Utility");
+		const bool bDisplayUtilityConsideration = State
+			&& State->Parent
+			&& UE::StateTree::IsSelectionBehaviorUsingUtility(State->Parent->SelectionBehavior);
+		if (bDisplayUtilityConsideration)
+		{
+			const FName UtilityCategoryName(TEXT("Selection Utility"));
+			int32 NumConsiderations = WeakState.IsValid() ? WeakState->Considerations.Num() : 0;
+			FText UtilityCategoryText = EnumHasAllFlags(GetDefault<UStateTreeEditorUserSettings>()->GetDetailsViewDisplayCount(), EStateTreeEditorUserSettingsNodeType::Utility)
+				? FText::Format(LOCTEXT("StateDetailsSelectionUtilityCounted", "Selection Utility ({0})"), NumConsiderations)
+				: LOCTEXT("StateDetailsSelectionUtility", "Selection Utility");
 
-		IDetailCategoryBuilder& UtilityConsiderationsCategory = UE::StateTreeEditor::EditorNodeUtils::MakeArrayCategory(
-			DetailBuilder,
-			ConsiderationsProperty,
-			UtilityCategoryName,
-			UtilityCategoryText,
-			FName("StateTreeEditor.Utility"),
-			UE::StateTree::Colors::Orange,
-			UE::StateTree::Colors::Orange.WithAlpha(192),
-			LOCTEXT("UtilityAddTooltip", "Add new Utility Consideration"),
-			/*SortOrder*/3);
-
-		// Weight
-		UtilityConsiderationsCategory.AddProperty(WeightProperty);
+			IDetailCategoryBuilder& UtilityConsiderationsCategory = UE::StateTreeEditor::EditorNodeUtils::MakeArrayCategory(
+				DetailBuilder,
+				ConsiderationsProperty,
+				UtilityCategoryName,
+				UtilityCategoryText,
+				FName("StateTreeEditor.Utility"),
+				UE::StateTree::Colors::Orange,
+				UE::StateTree::Colors::Orange.WithAlpha(192),
+				LOCTEXT("UtilityAddTooltip", "Add new Utility Consideration"),
+				/*SortOrder*/3);
+
+			// Weight
+			UtilityConsiderationsCategory.AddProperty(WeightProperty);
+		}
 	}
 
 	// Tasks
diff --git a/Engine/Plugins/Runtime/StateTree/Source/StateTreeEditorModule/Private/Customizations/StateTreeStateDetails.h b/Engine/Plugins/Runtime/StateTree/Source/StateTreeEditorModule/Private/Customizations/StateTreeStateDetails.h
--- a/Engine/Plugins/Runtime/StateTree/Source/StateTreeEditorModule/Private/Customizations/StateTreeStateDetails.h
+++ b/Engine/Plugins/Runtime/StateTree/Source/StateTreeEditorModule/Private/Customizations/StateTreeStateDetails.h
@@ -4,6 +4,8 @@
 
 #include "IDetailCustomization.h"
 #include "EditorUndoClient.h"
+#include "StateTreeEditorData.h"
+#include "StateTreeState.h"
 
 class IDetailLayoutBuilder;
 class IDetailChildrenBuilder;
@@ -22,7 +24,7 @@ public:
 	/** Makes a new instance of this detail layout class for a specific detail view requesting it */
 	static TSharedRef<IDetailCustomization> MakeInstance();
 
-	/** IDetailCustomization interface */
+	//~ IDetailCustomization interface
 	virtual void PendingDelete() override;
 	virtual void CustomizeDetails(IDetailLayoutBuilder& DetailBuilder) override;
 
@@ -32,7 +34,10 @@ protected:
 	virtual void PostUndo(bool bSuccess) override;
 	virtual void PostRedo(bool bSuccess) override;
 
+protected:
 	TWeakPtr<IPropertyUtilities> WeakPropertyUtilities;
+	TWeakObjectPtr<UStateTreeEditorData> WeakEditorData;
+	TWeakObjectPtr<UStateTreeState> WeakState;
 
 	bool bIsPendingDelete = false;
 };
diff --git a/Engine/Plugins/Runtime/StateTree/Source/StateTreeEditorModule/Private/StateTreeCompiler.cpp b/Engine/Plugins/Runtime/StateTree/Source/StateTreeEditorModule/Private/StateTreeCompiler.cpp
--- a/Engine/Plugins/Runtime/StateTree/Source/StateTreeEditorModule/Private/StateTreeCompiler.cpp
+++ b/Engine/Plugins/Runtime/StateTree/Source/StateTreeEditorModule/Private/StateTreeCompiler.cpp
@@ -2282,9 +2282,8 @@ bool FStateTreeCompiler::CreateConsiderations(UStateTreeState& State, const FStr
 {
 	if (State.Considerations.Num() != 0)
 	{
-		if (!State.Parent
-			|| (State.Parent->SelectionBehavior != EStateTreeStateSelectionBehavior::TrySelectChildrenWithHighestUtility
-				&& State.Parent->SelectionBehavior != EStateTreeStateSelectionBehavior::TrySelectChildrenAtRandomWeightedByUtility))
+		if (State.Parent == nullptr
+			|| !UE::StateTree::IsSelectionBehaviorUsingUtility(State.Parent->SelectionBehavior))
 		{
 			Log.Reportf(EMessageSeverity::Warning, TEXT("State's Utility Considerations data are compiled but they don't have effect."
 					"The Utility Considerations are used only when parent State's Selection Behavior is:"
diff --git a/Engine/Plugins/Runtime/StateTree/Source/StateTreeModule/Public/StateTreeTypes.h b/Engine/Plugins/Runtime/StateTree/Source/StateTreeModule/Public/StateTreeTypes.h
--- a/Engine/Plugins/Runtime/StateTree/Source/StateTreeModule/Public/StateTreeTypes.h
+++ b/Engine/Plugins/Runtime/StateTree/Source/StateTreeModule/Public/StateTreeTypes.h
@@ -186,6 +186,21 @@ enum class EStateTreeStateSelectionBehavior : uint8
 	TrySelectChildrenBasedOnRelativeUtility UE_DEPRECATED(all, "Use TrySelectChildrenAtRandomWeightedByUtility Instead") = TrySelectChildrenAtRandomWeightedByUtility UMETA(Hidden)
 };
 
+namespace UE::StateTree
+{
+	static bool IsSelectionBehaviorUsingUtility(const EStateTreeStateSelectionBehavior Value)
+	{
+		switch(Value)
+		{
+			case EStateTreeStateSelectionBehavior::TrySelectChildrenWithHighestUtility:
+			case EStateTreeStateSelectionBehavior::TrySelectChildrenAtRandomWeightedByUtility:
+				return true;
+			default:
+				return false;
+		}
+	}
+}
+
 
 /** Transitions trigger. */
 UENUM()
