[StateTree] Do not create an instance data storage for shared instance data when no shared data exists. It reduces the runtime memory footprint. Each thread that access state tree increments a counter. That counter is static and is valid for all state tree assets. That means each state tree asset uses a ?number of threads? amount of storage, even if the state tree asset is not intended for threading.

#Synced-CL 49204654

diff --git a/Engine/Plugins/Runtime/StateTree/Source/StateTreeModule/Private/Debugger/StateTreeRuntimeValidation.cpp b/Engine/Plugins/Runtime/StateTree/Source/StateTreeModule/Private/Debugger/StateTreeRuntimeValidation.cpp
--- a/Engine/Plugins/Runtime/StateTree/Source/StateTreeModule/Private/Debugger/StateTreeRuntimeValidation.cpp
+++ b/Engine/Plugins/Runtime/StateTree/Source/StateTreeModule/Private/Debugger/StateTreeRuntimeValidation.cpp
@@ -238,7 +238,7 @@ void FRuntimeValidationInstanceData::ValidateInstanceData(const UStateTree* NewS
 			}
 
 			TSharedPtr<FStateTreeInstanceData> SharedPtr = NewStateTree->GetSharedInstanceData();
-			if (ensureMsgf(SharedPtr, TEXT("The shared instance data is invalid")))
+			if (SharedPtr)
 			{
 				if (!SharedPtr->GetStorage().AreAllInstancesValid())
 				{
diff --git a/Engine/Plugins/Runtime/StateTree/Source/StateTreeModule/Private/StateTree.cpp b/Engine/Plugins/Runtime/StateTree/Source/StateTreeModule/Private/StateTree.cpp
--- a/Engine/Plugins/Runtime/StateTree/Source/StateTreeModule/Private/StateTree.cpp
+++ b/Engine/Plugins/Runtime/StateTree/Source/StateTreeModule/Private/StateTree.cpp
@@ -139,6 +139,11 @@ static int32 GetThreadIndexForSharedInstanceData()
 
 TSharedPtr<FStateTreeInstanceData> UStateTree::GetSharedInstanceData() const
 {
+	if (SharedInstanceData.Num() == 0)
+	{
+		return TSharedPtr<FStateTreeInstanceData>();
+	}
+
 	int32 ThreadIndex = GetThreadIndexForSharedInstanceData();
 
 	// If shared instance data for this thread exists, return it.
diff --git a/Engine/Plugins/Runtime/StateTree/Source/StateTreeModule/Private/StateTreeExecutionContext.cpp b/Engine/Plugins/Runtime/StateTree/Source/StateTreeModule/Private/StateTreeExecutionContext.cpp
--- a/Engine/Plugins/Runtime/StateTree/Source/StateTreeModule/Private/StateTreeExecutionContext.cpp
+++ b/Engine/Plugins/Runtime/StateTree/Source/StateTreeModule/Private/StateTreeExecutionContext.cpp
@@ -995,7 +995,9 @@ TConstArrayView<FStateTreeExecutionFrame> FStateTreeExecutionContext::FSelectSta
 FStateTreeExecutionContext::FCurrentlyProcessedFrameScope::FCurrentlyProcessedFrameScope(FStateTreeExecutionContext& InContext, const FStateTreeExecutionFrame* CurrentParentFrame, const FStateTreeExecutionFrame& CurrentFrame) : Context(InContext)
 {
 	check(CurrentFrame.StateTree);
-	FStateTreeInstanceStorage* SharedInstanceDataStorage = &CurrentFrame.StateTree->GetSharedInstanceData()->GetMutableStorage();
+
+	TSharedPtr<FStateTreeInstanceData> SharedPtr = CurrentFrame.StateTree->GetSharedInstanceData();
+	FStateTreeInstanceStorage* SharedInstanceDataStorage = SharedPtr ? &(SharedPtr->GetMutableStorage()) : nullptr;
 
 	SavedFrame = Context.CurrentlyProcessedFrame;
 	SavedParentFrame = Context.CurrentlyProcessedParentFrame;
