[StateTree] Prevent linked state override when bindings could use invalid data.

#Synced-CL 49284918

diff --git a/Engine/Plugins/Runtime/StateTree/Source/StateTreeEditorModule/Private/Customizations/StateTreeStateDetails.cpp b/Engine/Plugins/Runtime/StateTree/Source/StateTreeEditorModule/Private/Customizations/StateTreeStateDetails.cpp
--- a/Engine/Plugins/Runtime/StateTree/Source/StateTreeEditorModule/Private/Customizations/StateTreeStateDetails.cpp
+++ b/Engine/Plugins/Runtime/StateTree/Source/StateTreeEditorModule/Private/Customizations/StateTreeStateDetails.cpp
@@ -18,6 +18,7 @@
 #include "StateTreeSchema.h"
 #include "StateTreeStateParametersDetails.h"
 #include "Debugger/StateTreeDebuggerUIExtensions.h"
+#include "SWarningOrErrorBox.h"
 #include "Widgets/SBoxPanel.h"
 #include "Widgets/Images/SImage.h"
 #include "Widgets/Text/STextBlock.h"
@@ -76,6 +77,7 @@ void FStateTreeStateDetails::CustomizeDetails(IDetailLayoutBuilder& DetailBuilde
 
 	const UStateTreeSchema* Schema = EditorData ? EditorData->Schema.Get() : nullptr;
 	const FString SchemaPath = Schema ? Schema->GetClass()->GetPathName() : FString();
+	WeakEditorData = EditorData;
 
 	TSharedPtr<FStateTreeViewModel> ViewModel;
 	if (UStateTreeEditingSubsystem* StateTreeEditingSubsystem = EditorData != nullptr ? GEditor->GetEditorSubsystem<UStateTreeEditingSubsystem>() : nullptr)
@@ -87,6 +89,7 @@ void FStateTreeStateDetails::CustomizeDetails(IDetailLayoutBuilder& DetailBuilde
 	const TSharedPtr<IPropertyHandle> NameProperty = DetailBuilder.GetProperty(GET_MEMBER_NAME_CHECKED(UStateTreeState, Name));
 	const TSharedPtr<IPropertyHandle> TagProperty = DetailBuilder.GetProperty(GET_MEMBER_NAME_CHECKED(UStateTreeState, Tag));
 	const TSharedPtr<IPropertyHandle> ColorRefProperty = DetailBuilder.GetProperty(GET_MEMBER_NAME_CHECKED(UStateTreeState, ColorRef));
+	const TSharedPtr<IPropertyHandle> DescriptionProperty = DetailBuilder.GetProperty(GET_MEMBER_NAME_CHECKED(UStateTreeState, Description));
 	const TSharedPtr<IPropertyHandle> EnabledProperty = DetailBuilder.GetProperty(GET_MEMBER_NAME_CHECKED(UStateTreeState, bEnabled));
 	const TSharedPtr<IPropertyHandle> TasksProperty = DetailBuilder.GetProperty(GET_MEMBER_NAME_CHECKED(UStateTreeState, Tasks));
 	const TSharedPtr<IPropertyHandle> SingleTaskProperty = DetailBuilder.GetProperty(GET_MEMBER_NAME_CHECKED(UStateTreeState, SingleTask));
@@ -183,6 +186,10 @@ void FStateTreeStateDetails::CustomizeDetails(IDetailLayoutBuilder& DetailBuilde
 		StateCategory.AddProperty(CustomTickRateProperty);
 	}
 
+	// Description
+	DescriptionProperty->MarkHiddenByCustomization();
+	StateCategory.AddProperty(DescriptionProperty);
+
 	// Type
 	TypeProperty->MarkHiddenByCustomization();
 	if (Schema)
@@ -251,6 +258,18 @@ void FStateTreeStateDetails::CustomizeDetails(IDetailLayoutBuilder& DetailBuilde
 				return !SchemaPath.IsEmpty() && !InAssetData.TagsAndValues.ContainsKeyValue(UE::StateTree::SchemaTag, SchemaPath);
 			})
 		];
+
+		const FButtonStyle* const CloseButtonStyle = &FAppStyle::Get().GetWidgetStyle<FDockTabStyle>("Docking.Tab").CloseButtonStyle;
+		StateCategory.AddCustomRow(LOCTEXT("WarningBindingWithLinkedAsset", "WarningBindingWithLinkedAsset"))
+			.Visibility(MakeAttributeSP(this, &FStateTreeStateDetails::HandleLinkedAssetVisibility))
+			.WholeRowContent()
+			[
+				SNew(SWarningOrErrorBox)
+					.MessageStyle(EMessageStyle::Warning)
+					.IconSize(FVector2D(16.0f, 16.0f))
+					.Padding(FMargin(0.0f))
+					.Message(LOCTEXT("WarningBindingWithLinkedAssetMess", "Override will be disabled because of the binding(s)."))
+			];
 	}
 
 	// Parameters category
@@ -465,6 +484,24 @@ void FStateTreeStateDetails::CustomizeDetails(IDetailLayoutBuilder& DetailBuilde
 	}
 }
 
+EVisibility FStateTreeStateDetails::HandleLinkedAssetVisibility() const
+{
+	const UStateTreeEditorData* EditorData = WeakEditorData.Get();
+	const UStateTreeState* State = WeakState.Get();
+
+	if (EditorData
+		&& State
+		&& State->Type == EStateTreeStateType::LinkedAsset
+		&& State->Tag.IsValid())
+	{
+		TArray<const FPropertyBindingBinding*> Bindings;
+		EditorData->GetEditorPropertyBindings()->GetBindingsFor(State->Parameters.ID, Bindings);
+		const bool bHasBinding = Bindings.Num() > 0;
+		return bHasBinding ? EVisibility::Visible : EVisibility::Collapsed;
+	}
+	return EVisibility::Collapsed;
+}
+
 void FStateTreeStateDetails::PostUndo(bool bSuccess)
 {
 	// Prevent redundant calls to force refresh which doubles the number of DetailCustomization 
diff --git a/Engine/Plugins/Runtime/StateTree/Source/StateTreeEditorModule/Private/Customizations/StateTreeStateDetails.h b/Engine/Plugins/Runtime/StateTree/Source/StateTreeEditorModule/Private/Customizations/StateTreeStateDetails.h
--- a/Engine/Plugins/Runtime/StateTree/Source/StateTreeEditorModule/Private/Customizations/StateTreeStateDetails.h
+++ b/Engine/Plugins/Runtime/StateTree/Source/StateTreeEditorModule/Private/Customizations/StateTreeStateDetails.h
@@ -14,6 +14,7 @@ class IPropertyHandle;
 class IPropertyUtilities;
 class SWidget;
 class SComboButton;
+struct EVisibility;
 
 class FStateTreeStateDetails : public IDetailCustomization, FSelfRegisteringEditorUndoClient
 {
@@ -34,6 +35,9 @@ protected:
 	virtual void PostUndo(bool bSuccess) override;
 	virtual void PostRedo(bool bSuccess) override;
 
+private:
+	EVisibility HandleLinkedAssetVisibility() const;
+
 protected:
 	TWeakPtr<IPropertyUtilities> WeakPropertyUtilities;
 	TWeakObjectPtr<UStateTreeEditorData> WeakEditorData;
diff --git a/Engine/Plugins/Runtime/StateTree/Source/StateTreeModule/Private/StateTreeExecutionContext.cpp b/Engine/Plugins/Runtime/StateTree/Source/StateTreeModule/Private/StateTreeExecutionContext.cpp
--- a/Engine/Plugins/Runtime/StateTree/Source/StateTreeModule/Private/StateTreeExecutionContext.cpp
+++ b/Engine/Plugins/Runtime/StateTree/Source/StateTreeModule/Private/StateTreeExecutionContext.cpp
@@ -6964,7 +6964,9 @@ bool FStateTreeExecutionContext::SelectStateInternal(
 			}
 		}
 
-		if (bCreateNewState)
+		// LinkedAsset that has bindings cannot be overridden by another asset. The bindings will write invalid data.
+		const bool bInvalidOverride = NextState.ParameterBindingsBatch.IsValid();
+		if (bCreateNewState && !bInvalidOverride)
 		{
 			if (const FStateTreeReference* Override = GetLinkedStateTreeOverrideForTag(NextState.Tag))
 			{
