[StateTreeDebugger] added support for breakpoints in statically linked assets (i.e., not dynamic overrides)
- asset no longer get dirty when adding/removing breakpoint
- select the first available instance in the debugger view when analysis starts
#Synced-CL 49287773

diff --git a/Engine/Plugins/Runtime/StateTree/Source/StateTreeEditorModule/Private/Debugger/SStateTreeDebuggerView.cpp b/Engine/Plugins/Runtime/StateTree/Source/StateTreeEditorModule/Private/Debugger/SStateTreeDebuggerView.cpp
--- a/Engine/Plugins/Runtime/StateTree/Source/StateTreeEditorModule/Private/Debugger/SStateTreeDebuggerView.cpp
+++ b/Engine/Plugins/Runtime/StateTree/Source/StateTreeEditorModule/Private/Debugger/SStateTreeDebuggerView.cpp
@@ -24,6 +24,7 @@
 #include "StateTreeModule.h"
 #include "StateTreeViewModel.h"
 #include "SWarningOrErrorBox.h"
+#include "Toolkits/ToolkitManager.h"
 #include "Widgets/Docking/SDockTab.h"
 #include "Widgets/Input/SComboButton.h"
 #include "Widgets/Layout/SScrollBox.h"
@@ -178,7 +179,10 @@ void SStateTreeDebuggerView::OnPIEStarted(const bool bIsSimulating)
 
 void SStateTreeDebuggerView::OnPIEStopped(const bool bIsSimulating)
 {
-	StopRecording();
+	if (UStateTreeEditorSettings::Get().bEnableLegacyDebuggerWindow)
+	{
+		StopRecording();
+	}
 
 	check(Debugger);
 	Debugger->StopSessionAnalysis();
@@ -560,10 +564,13 @@ void SStateTreeDebuggerView::ConstructView(const TSharedRef<FUICommandList>& InC
 				FMenuBuilder MenuBuilder(/*bInShouldCloseWindowAfterMenuSelection*/true, /*InCommandList*/nullptr);
 				for (const TSharedRef<const UE::StateTreeDebugger::FInstanceDescriptor>& InstanceDescriptor : InstanceDescriptors)
 				{
-					if (InstanceDescriptor.Get().StateTree != Debugger->GetAsset())
+					// Add instance to the list if the asset associated to the view is the instance's main asset, or one of its linked assets.
+					if (InstanceDescriptor.Get().StateTree != Debugger->GetAsset()
+						&& !InstanceDescriptor.Get().LinkedStateTrees.Contains(Debugger->GetAsset()))
 					{
 						continue;
 					}
+
 					const FText Label = Debugger->DescribeInstance(InstanceDescriptor.Get());
 
 					FUIAction ItemAction(FExecuteAction::CreateSPLambda(Debugger.ToSharedRef(), [Debugger, InstanceId = InstanceDescriptor.Get().Id]
@@ -1007,6 +1014,12 @@ void SStateTreeDebuggerView::OnBreakpointHit(const FStateTreeInstanceDebugId Ins
 				return Breakpoint.IsMatchingEvent(Event);
 			});
 	}
+
+	// Bring attention to the Editor associated to the breakpoint
+	if (const TSharedPtr<IToolkit> FoundAssetEditor = FToolkitManager::Get().FindEditorForAsset(StateTree.Get()))
+	{
+		FoundAssetEditor->BringToolkitToFront();
+	}
 }
 
 void SStateTreeDebuggerView::OnNewSession()
diff --git a/Engine/Plugins/Runtime/StateTree/Source/StateTreeEditorModule/Private/StateTreeViewModel.cpp b/Engine/Plugins/Runtime/StateTree/Source/StateTreeEditorModule/Private/StateTreeViewModel.cpp
--- a/Engine/Plugins/Runtime/StateTree/Source/StateTreeEditorModule/Private/StateTreeViewModel.cpp
+++ b/Engine/Plugins/Runtime/StateTree/Source/StateTreeEditorModule/Private/StateTreeViewModel.cpp
@@ -687,26 +687,26 @@ void FStateTreeViewModel::HandleEnableStateBreakpoint(EStateTreeBreakpointType T
 	if (CanAddStateBreakpoint(Type))
 	{
 		const FScopedTransaction Transaction(LOCTEXT("AddStateBreakpoint", "Add State Breakpoint(s)"));
-		EditorData->Modify();
+		EditorData->Modify(/*bAlwaysMarkDirty*/false);
 		for (int Index = 0; Index < ValidatedSelectedStates.Num(); ++Index)
 		{
 			const UStateTreeState* SelectedState = ValidatedSelectedStates[Index];
 			if (HasBreakpoint[Index] == false && SelectedState != nullptr)
 			{
-				EditorData->AddBreakpoint(SelectedState->ID, Type);	
+				EditorData->AddBreakpoint(SelectedState->ID, Type);
 			}
 		}
 	}
 	else if (CanRemoveStateBreakpoint(Type))
 	{
 		const FScopedTransaction Transaction(LOCTEXT("RemoveStateBreakpoint", "Remove State Breakpoint(s)"));
-		EditorData->Modify();
+		EditorData->Modify(/*bAlwaysMarkDirty*/false);
 		for (int Index = 0; Index < ValidatedSelectedStates.Num(); ++Index)
 		{
 			const UStateTreeState* SelectedState = ValidatedSelectedStates[Index];
 			if (HasBreakpoint[Index] && SelectedState != nullptr)
 			{
-				EditorData->RemoveBreakpoint(SelectedState->ID, Type);	
+				EditorData->RemoveBreakpoint(SelectedState->ID, Type);
 			}
 		}
 	}
diff --git a/Engine/Plugins/Runtime/StateTree/Source/StateTreeModule/Private/Debugger/StateTreeDebugger.cpp b/Engine/Plugins/Runtime/StateTree/Source/StateTreeModule/Private/Debugger/StateTreeDebugger.cpp
--- a/Engine/Plugins/Runtime/StateTree/Source/StateTreeModule/Private/Debugger/StateTreeDebugger.cpp
+++ b/Engine/Plugins/Runtime/StateTree/Source/StateTreeModule/Private/Debugger/StateTreeDebugger.cpp
@@ -127,6 +127,25 @@ void FStateTreeDebugger::Tick(const float DeltaTime)
 		}
 	}
 
+	// QOL: try to select an instance by default
+	if (GetSelectedInstanceId().IsInvalid() && IsAnalysisSessionActive())
+	{
+		TArray<const TSharedRef<const UE::StateTreeDebugger::FInstanceDescriptor>> InstanceDescriptors;
+		GetSessionInstanceDescriptors(InstanceDescriptors);
+
+		for (const TSharedRef<const UE::StateTreeDebugger::FInstanceDescriptor>& InstanceDescriptor : InstanceDescriptors)
+		{
+			const UE::StateTreeDebugger::FInstanceDescriptor& Descriptor = InstanceDescriptor.Get();
+			// Add instance to the list if the asset associated to the view is the instance's main asset, or one of its linked assets.
+			if (Descriptor.StateTree != GetAsset()
+				&& !Descriptor.LinkedStateTrees.Contains(GetAsset()))
+			{
+				continue;
+			}
+			SelectInstance(Descriptor.Id);
+		}
+	}
+
 	if (StateTreeAsset.IsValid()
 		&& IsAnalysisSessionActive()
 		&& !IsAnalysisSessionPaused())
@@ -155,6 +174,7 @@ void FStateTreeDebugger::StopSessionAnalysis()
 		AnalysisSession->Stop(/*WaitOnAnalysis*/true);
 	}
 
+	CurrentFrameStateTreeAsset.Reset();
 	bSessionAnalysisActive = false;
 	bSessionAnalysisPaused = false;
 	LastProcessedRecordedWorldTime = 0;
@@ -205,6 +225,18 @@ void FStateTreeDebugger::SelectInstance(const FStateTreeInstanceDebugId Instance
 	{
 		SelectedInstanceId = InstanceId;
 
+		// The debugger can be associated to a StateTree asset used as a linked StateTree in active instances.
+		// Those active instances are available for the user to select and in that case we need to fetch the main
+		// StateTree (i.e., associated to the runtime instance) to read the proper event timeline.
+		InstanceRootAsset = StateTreeAsset;
+		if (const TSharedPtr<const UE::StateTreeDebugger::FInstanceDescriptor> FoundDescriptor = GetDescriptor(SelectedInstanceId))
+		{
+			if (const UStateTree* FoundStateTree = FoundDescriptor->StateTree.Get())
+			{
+				InstanceRootAsset = FoundStateTree;
+			}
+		}
+
 		// Update event collection for newly debugged instance
 		SetScrubStateCollection(GetMutableEventCollection(InstanceId));
 	}
@@ -753,7 +785,7 @@ void FStateTreeDebugger::ReadTrace(const uint64 FrameIndex)
 				, FrameProvider
 				, *TargetFrame
 				, /*ITraceReader*/this
-				, UE::StateTreeDebugger::FTraceFilter(StateTreeAsset.Get())
+				, UE::StateTreeDebugger::FTraceFilter(InstanceRootAsset.Get())
 				, LastTraceReadTime);
 		}
 	}
@@ -770,7 +802,7 @@ void FStateTreeDebugger::ReadTrace(const double ScrubTime)
 			*Session
 			, ScrubTime
 			, this
-			, UE::StateTreeDebugger::FTraceFilter(StateTreeAsset.Get())
+			, UE::StateTreeDebugger::FTraceFilter(InstanceRootAsset.Get())
 			, LastTraceReadTime);
 
 		SendNotifications();
@@ -811,8 +843,15 @@ void FStateTreeDebugger::SendNotifications()
 
 bool FStateTreeDebugger::EvaluateBreakpoints(const FStateTreeInstanceDebugId InstanceId, const FStateTreeTraceEventVariantType& Event)
 {
+	// Keep track of the StateTree associated to the last frame update so breakpoints are tested against the right asset.
+	if (const FStateTreeTraceInstanceFrameEvent* FrameEvent = Event.TryGet<FStateTreeTraceInstanceFrameEvent>())
+	{
+		CurrentFrameStateTreeAsset = FrameEvent->WeakStateTree;
+	}
+
 	if (StateTreeAsset == nullptr // asset is required to properly match state handles
-		|| HitBreakpoint.IsSet() // Only consider first hit breakpoint in the frame
+		|| CurrentFrameStateTreeAsset != StateTreeAsset // ignore events from a different asset (required for linked statetree)
+		|| HitBreakpoint.IsSet() // only consider first hit breakpoint in the frame
 		|| Breakpoints.IsEmpty()
 		|| (SelectedInstanceId.IsValid() && InstanceId != SelectedInstanceId)) // ignore events not for the selected instances
 	{
diff --git a/Engine/Plugins/Runtime/StateTree/Source/StateTreeModule/Private/Debugger/StateTreeTraceProvider.cpp b/Engine/Plugins/Runtime/StateTree/Source/StateTreeModule/Private/Debugger/StateTreeTraceProvider.cpp
--- a/Engine/Plugins/Runtime/StateTree/Source/StateTreeModule/Private/Debugger/StateTreeTraceProvider.cpp
+++ b/Engine/Plugins/Runtime/StateTree/Source/StateTreeModule/Private/Debugger/StateTreeTraceProvider.cpp
@@ -1,4 +1,4 @@
-// Copyright Epic Games, Inc. All Rights Reserved.
+﻿// Copyright Epic Games, Inc. All Rights Reserved.
 
 #if WITH_STATETREE_TRACE_DEBUGGER
 
@@ -85,6 +85,8 @@ void FStateTreeTraceProvider::AppendInstanceEvent(
 	const double InWorldRecordingTime,
 	const EStateTreeTraceEventType InEventType)
 {
+	using namespace UE::StateTreeDebugger;
+
 	if (InEventType == EStateTreeTraceEventType::Push)
 	{
 		TWeakObjectPtr<const UStateTree> WeakStateTree;
@@ -92,8 +94,8 @@ void FStateTreeTraceProvider::AppendInstanceEvent(
 		{
 			if (const UStateTree* StateTree = WeakStateTree.Get())
 			{
-				const TSharedRef<UE::StateTreeDebugger::FInstanceDescriptor>* Descriptor = Descriptors.FindByPredicate(
-					[InInstanceId](const TSharedRef<const UE::StateTreeDebugger::FInstanceDescriptor>& Descriptor)
+				const TSharedRef<FInstanceDescriptor>* Descriptor = Descriptors.FindByPredicate(
+					[InInstanceId](const TSharedRef<const FInstanceDescriptor>& Descriptor)
 					{
 						return Descriptor.Get().Id == InInstanceId;
 					});
@@ -102,11 +104,26 @@ void FStateTreeTraceProvider::AppendInstanceEvent(
 				// In that case we can reuse the existing entries
 				if (Descriptor == nullptr)
 				{
-					Descriptors.Emplace(MakeShared<UE::StateTreeDebugger::FInstanceDescriptor>(
+					TSharedRef<FInstanceDescriptor>& NewDescriptor = Descriptors.Emplace_GetRef(MakeShared<FInstanceDescriptor>(
 						StateTree
 						, InInstanceId
 						, InInstanceName
-						, TRange<double>(InWorldRecordingTime, UE::StateTreeDebugger::FInstanceDescriptor::ActiveInstanceEndTime)));
+						, TRange<double>(InWorldRecordingTime, FInstanceDescriptor::ActiveInstanceEndTime)));
+
+					// Build list of statically linked assets
+					// This can then be used when processing events and breakpoints set in linked assets
+					auto GatherLinkedStateTrees = [](const auto& Self, const UStateTree* InStateTree, TArray<TWeakObjectPtr<const UStateTree>>& OutAssets)->void
+						{
+							for (const FCompactStateTreeState& State : InStateTree->GetStates())
+							{
+								if (State.LinkedAsset)
+								{
+									OutAssets.Add(State.LinkedAsset);
+									Self(Self, State.LinkedAsset, OutAssets);
+								}
+							}
+						};
+					GatherLinkedStateTrees(GatherLinkedStateTrees, StateTree, NewDescriptor.Get().LinkedStateTrees);
 
 					check(InstanceIdToDebuggerEntryTimelines.Find(InInstanceId) == nullptr);
 					InstanceIdToDebuggerEntryTimelines.Add(InInstanceId, EventsTimelines.Num());
diff --git a/Engine/Plugins/Runtime/StateTree/Source/StateTreeModule/Public/Debugger/StateTreeDebugger.h b/Engine/Plugins/Runtime/StateTree/Source/StateTreeModule/Public/Debugger/StateTreeDebugger.h
--- a/Engine/Plugins/Runtime/StateTree/Source/StateTreeModule/Public/Debugger/StateTreeDebugger.h
+++ b/Engine/Plugins/Runtime/StateTree/Source/StateTreeModule/Public/Debugger/StateTreeDebugger.h
@@ -146,6 +146,7 @@ struct FStateTreeDebugger : FTickableGameObject, UE::StateTreeDebugger::ITraceRe
 	void SetAsset(const UStateTree* Asset)
 	{
 		StateTreeAsset = Asset;
+		InstanceRootAsset = Asset;
 	}
 
 	/** Forces a single refresh to latest state. Useful when simulation is paused. */
@@ -444,6 +445,15 @@ private:
 	/** Module used to access the store client and analysis sessions .*/
 	IStateTreeModule& StateTreeModule;
 
+	/**
+	 * The StateTree asset associated to the selected instance.
+	 * This can be different from 'StateTreeAsset' if the debugger is associated to a linked asset.
+	 */
+	TWeakObjectPtr<const UStateTree> InstanceRootAsset;
+
+	/** The StateTree asset associated to the last frame event processed by the debugger */
+	TWeakObjectPtr<const UStateTree> CurrentFrameStateTreeAsset;
+
 	/** The StateTree asset associated to this debugger. All instances will be using this asset. */
 	TWeakObjectPtr<const UStateTree> StateTreeAsset;
 
diff --git a/Engine/Plugins/Runtime/StateTree/Source/StateTreeModule/Public/Debugger/StateTreeDebuggerTypes.h b/Engine/Plugins/Runtime/StateTree/Source/StateTreeModule/Public/Debugger/StateTreeDebuggerTypes.h
--- a/Engine/Plugins/Runtime/StateTree/Source/StateTreeModule/Public/Debugger/StateTreeDebuggerTypes.h
+++ b/Engine/Plugins/Runtime/StateTree/Source/StateTreeModule/Public/Debugger/StateTreeDebuggerTypes.h
@@ -101,6 +101,7 @@ struct FInstanceDescriptor : TSharedFromThis<FInstanceDescriptor>
 	static constexpr double ActiveInstanceEndTime = std::numeric_limits<double>::max();
 	TRange<double> Lifetime = TRange<double>(0);
 	TWeakObjectPtr<const UStateTree> StateTree = nullptr;
+	TArray<TWeakObjectPtr<const UStateTree>> LinkedStateTrees;
 	FString Name;
 	FStateTreeInstanceDebugId Id = FStateTreeInstanceDebugId::Invalid;
 };
