[StateTree] Fix unit test System.StateTree.LinkedAsset.Override. The bug was fixed with CL 49345530.

#Synced-CL 49284918

diff --git a/Engine/Plugins/Runtime/StateTree/Source/StateTreeTestSuite/Private/StateTreeLinkedStateTest.cpp b/Engine/Plugins/Runtime/StateTree/Source/StateTreeTestSuite/Private/StateTreeLinkedStateTest.cpp
--- a/Engine/Plugins/Runtime/StateTree/Source/StateTreeTestSuite/Private/StateTreeLinkedStateTest.cpp
+++ b/Engine/Plugins/Runtime/StateTree/Source/StateTreeTestSuite/Private/StateTreeLinkedStateTest.cpp
@@ -1841,7 +1841,7 @@ struct FStateTreeTest_OverrideLinkedAsset : FStateTreeTestBase
 			UStateTreeEditorData& EditorData = *Cast<UStateTreeEditorData>(StateTree3.EditorData);
 			GetRootPropertyBag(EditorData).AddProperty("B", EPropertyBagPropertyType::Int32);
 
-			UStateTreeState& Root = EditorData.AddSubTree(FName(TEXT("Root")));
+			UStateTreeState& Root = EditorData.AddSubTree(FName(TEXT("Tree3Root")));
 			TStateTreeEditorNode<FTestTask_PrintValue>& RootTask1 = Root.AddTask<FTestTask_PrintValue>(FName(TEXT("Tree3RootTask1")));
 			RootTask1.GetInstanceData().Value = -3;
 
@@ -1860,16 +1860,16 @@ struct FStateTreeTest_OverrideLinkedAsset : FStateTreeTestBase
 			UStateTreeEditorData& EditorData = *Cast<UStateTreeEditorData>(StateTree2.EditorData);
 			GetRootPropertyBag(EditorData).AddProperty("A", EPropertyBagPropertyType::Int32);
 
-			UStateTreeState& Root2 = EditorData.AddSubTree(FName(TEXT("Root")));
+			UStateTreeState& Root2 = EditorData.AddSubTree(FName(TEXT("Tree2Root")));
 
-			UStateTreeState& State1 = Root2.AddChildState("State1");
+			UStateTreeState& State1 = Root2.AddChildState("Tree2State1");
 			TStateTreeEditorNode<FTestTask_Stand>& State1Task1 = State1.AddTask<FTestTask_Stand>(FName(TEXT("Tree2State1Task1")));
 			TStateTreeEditorNode<FTestTask_PrintValue>& State1Task2 = State1.AddTask<FTestTask_PrintValue>(FName(TEXT("Tree2State1Task2")));
 			State1.AddTransition(EStateTreeTransitionTrigger::OnStateCompleted, EStateTreeTransitionType::NextState);
 
 			EditorData.AddPropertyBinding(FPropertyBindingPath(EditorData.GetRootParametersGuid(), "A"), FPropertyBindingPath(State1Task2.ID, TEXT("Value")));
 
-			UStateTreeState& State2 = Root2.AddChildState("State2");
+			UStateTreeState& State2 = Root2.AddChildState("Tree2State2");
 			TStateTreeEditorNode<FTestTask_PrintValue>& State2Task1 = State2.AddTask<FTestTask_PrintValue>(FName(TEXT("Tree2State2Task1")));
 
 			EditorData.AddPropertyBinding(FPropertyBindingPath(EditorData.GetRootParametersGuid(), "A"), FPropertyBindingPath(State2Task1.ID, TEXT("Value")));
@@ -1888,8 +1888,8 @@ struct FStateTreeTest_OverrideLinkedAsset : FStateTreeTestBase
 
 			GetRootPropertyBag(EditorData).AddProperty("Int1Param", EPropertyBagPropertyType::Int32);
 
-			UStateTreeState& Root = EditorData.AddSubTree(FName(TEXT("Root1")));
-			UStateTreeState& State1 = Root.AddChildState(FName(TEXT("State1")), EStateTreeStateType::LinkedAsset);
+			UStateTreeState& Root = EditorData.AddSubTree(FName(TEXT("Tree1Root1")));
+			UStateTreeState& State1 = Root.AddChildState(FName(TEXT("Tree1State1")), EStateTreeStateType::LinkedAsset);
 			State1.Tag = GetTestTag1();
 			State1.SetLinkedStateAsset(&StateTree2);
 
@@ -1927,6 +1927,7 @@ struct FStateTreeTest_OverrideLinkedAsset : FStateTreeTestBase
 			}
 			{
 				// Tree2State1Task1 completes and transition to Tree2State2Task1
+				//The Tree1State1 didn't complete.
 				EStateTreeRunStatus Status = Exec.Tick(0.1f);
 				AITEST_EQUAL("Tick should complete with Running", Status, EStateTreeRunStatus::Running);
 				AITEST_TRUE(TEXT("StateTree should tick Tree2State2Task1"), Exec.Expect("Tree2State2Task1", TEXT("EnterState11")));
@@ -1936,14 +1937,14 @@ struct FStateTreeTest_OverrideLinkedAsset : FStateTreeTestBase
 				Exec.Stop();
 				Exec.LogClear();
 			}
-			//@TODO Crash when the override (tree3) has a different global parameter struct than the compiled version.
-			//{
-			//	GlobalParameters.SetValueInt32("Int1Param", 33);
-			//	EStateTreeRunStatus Status = Exec.Start(GlobalParameters.GetValue());
-			//	AITEST_EQUAL("Start should complete with Running", Status, EStateTreeRunStatus::Running);
-			//	AITEST_TRUE(TEXT("StateTree should tick Tree3RootTask1"), Exec.Expect("Tree3RootTask1", TEXT("EnterState33")));
-			//	Exec.LogClear();
-			//}
+			// Change the override. The override should not work because there bindings.
+			{
+				GlobalParameters.SetValueInt32("Int1Param", 33);
+				EStateTreeRunStatus Status = Exec.Start(GlobalParameters.GetValue());
+				AITEST_EQUAL("Start should complete with Running", Status, EStateTreeRunStatus::Running);
+				AITEST_TRUE(TEXT("StateTree should tick Tree2State1Task2"), Exec.Expect("Tree2State1Task2", TEXT("EnterState33")));
+				Exec.LogClear();
+			}
 			Exec.Stop();
 		}
 
