[StateTree] False positive circular reference while compiling state tree asset async. This occurs when A refers to B and C. And B also refer to C.

#Synced-CL 49284918

diff --git a/Engine/Plugins/Runtime/StateTree/Source/StateTreeEditorModule/Private/StateTreeCompilerManager.cpp b/Engine/Plugins/Runtime/StateTree/Source/StateTreeEditorModule/Private/StateTreeCompilerManager.cpp
--- a/Engine/Plugins/Runtime/StateTree/Source/StateTreeEditorModule/Private/StateTreeCompilerManager.cpp
+++ b/Engine/Plugins/Runtime/StateTree/Source/StateTreeEditorModule/Private/StateTreeCompilerManager.cpp
@@ -506,10 +506,11 @@ int32 FCompilerManagerImpl::CompileWithQueueImpl_AnyThread(TNotNull<UStateTree*>
 				const int32 DirtyAssetsIndex = DirtyAssets_AnyThread.IndexOfByKey(TObjectKey<UStateTree>(Dependency));
 				if (DirtyAssetsIndex != INDEX_NONE)
 				{
+					DirtyAssets_AnyThread.RemoveAtSwap(DirtyAssetsIndex, EAllowShrinking::No);
+
 					FStateTreeCompilerLog NewLog;
 					ProcessedAssets.Add(Dependency);
 					CompiledStateTreeCount += CompileWithQueueImpl_AnyThread(Dependency, NewLog, ProcessedAssets);
-					DirtyAssets_AnyThread.RemoveAtSwap(DirtyAssetsIndex);
 				}
 			}
 		}
