[StateTree] Added FStateTreeExecutionContext::GetCurrentlyProcessedStateSelection.
#Synced-CL 49284918

diff --git a/Engine/Plugins/Runtime/StateTree/Source/StateTreeModule/Private/StateTreeExecutionContext.cpp b/Engine/Plugins/Runtime/StateTree/Source/StateTreeModule/Private/StateTreeExecutionContext.cpp
--- a/Engine/Plugins/Runtime/StateTree/Source/StateTreeModule/Private/StateTreeExecutionContext.cpp
+++ b/Engine/Plugins/Runtime/StateTree/Source/StateTreeModule/Private/StateTreeExecutionContext.cpp
@@ -4149,6 +4149,13 @@ EStateTreeRunStatus FStateTreeExecutionContext::TickEvaluatorsAndGlobalTasksForF
 	return Result;
 }
 
+TSharedPtr<FStateTreeExecutionContext::FSelectStateResult> FStateTreeExecutionContext::GetCurrentlyProcessedStateSelection() const
+{
+	// A transition requests can be in progress or already succeeded.
+	//Use the current transition request first (for linked frame).
+	return CurrentlyProcessedTemporaryStorage ? CurrentlyProcessedTemporaryStorage : RequestedTransition ? RequestedTransition->Selection : nullptr;
+}
+
 void FStateTreeExecutionContext::MoveGlobalTemporariesToInstances(FStateTreeExecutionFrame& Frame)
 {
 	TNotNull<const UStateTree*> CurrentStateTree = Frame.StateTree;
diff --git a/Engine/Plugins/Runtime/StateTree/Source/StateTreeModule/Public/StateTreeExecutionContext.h b/Engine/Plugins/Runtime/StateTree/Source/StateTreeModule/Public/StateTreeExecutionContext.h
--- a/Engine/Plugins/Runtime/StateTree/Source/StateTreeModule/Public/StateTreeExecutionContext.h
+++ b/Engine/Plugins/Runtime/StateTree/Source/StateTreeModule/Public/StateTreeExecutionContext.h
@@ -611,12 +611,10 @@ PRAGMA_ENABLE_DEPRECATION_WARNINGS
 		return CurrentlyProcessedParentFrame;
 	}
 
-	/** @return in progress transition or the latest requested transition. */
+	/** @return the currently processed temporary storage */
 	TSharedPtr<UE::StateTree::ExecutionContext::ITemporaryStorage> GetCurrentlyProcessedTemporaryStorage() const
 	{
-		// A transition requests can be in progress or already succeeded.
-		//Use the current transition request first (for linked frame).
-		return CurrentlyProcessedTemporaryStorage ? CurrentlyProcessedTemporaryStorage : RequestedTransition ? RequestedTransition->Selection : nullptr;
+		return GetCurrentlyProcessedStateSelection();
 	}
 
 	/** @return Array view to named external data descriptors associated with this context. Note: Init() must be called before calling this method. */
@@ -1007,6 +1005,9 @@ protected:
 	 */
 	UE_API EStateTreeRunStatus TickEvaluatorsAndGlobalTasksForFrame(float DeltaTime, bool bTickGlobalTasks, int32 FrameIndex, const FStateTreeExecutionFrame* CurrentParentFrame, const TNotNull<FStateTreeExecutionFrame*> CurrentFrame);
 
+	/** @return in progress transition or the latest requested transition. */
+	UE_API TSharedPtr<FSelectStateResult> GetCurrentlyProcessedStateSelection() const;
+
 private:
 	/** Move the temporary data for each evaluators and tasks to instance data. */
 	void MoveGlobalTemporariesToInstances(FStateTreeExecutionFrame& Frame);
