[StateTree] Implement drag and drop for state tree state nodes
Extending the existing drag and drop structure to support dragging state's nodes around. By default, dragging a node moves it to another state (literally moving) keeping the bindings valid if any other node references it's properties. By holding CONTROL keyboard key while dragging, the behavior alternates to copy instead move.
#Synced-CL 49562939

diff --git a/Engine/Plugins/Runtime/StateTree/Source/StateTreeEditorModule/Private/Customizations/Widgets/SStateTreeContextMenuButton.cpp b/Engine/Plugins/Runtime/StateTree/Source/StateTreeEditorModule/Private/Customizations/Widgets/SStateTreeContextMenuButton.cpp
--- a/Engine/Plugins/Runtime/StateTree/Source/StateTreeEditorModule/Private/Customizations/Widgets/SStateTreeContextMenuButton.cpp
+++ b/Engine/Plugins/Runtime/StateTree/Source/StateTreeEditorModule/Private/Customizations/Widgets/SStateTreeContextMenuButton.cpp
@@ -5,6 +5,7 @@
 #include "Framework/Application/SlateApplication.h"
 #include "Framework/MultiBox/MultiBoxBuilder.h"
 #include "Widgets/Input/SMenuAnchor.h"
+#include "StateTreeDragDrop.h"
 #include "StateTreeState.h"
 #include "StateTreeViewModel.h"
 
@@ -26,9 +27,41 @@ void SStateTreeContextMenuButton::Construct(const FArguments& InArgs, const TSha
 		}
 	}
 
-	SButton::FArguments ButtonArgs;
+	TSharedRef<SStackBox> ButtonContent = SNew(SStackBox)
+		.Orientation(EOrientation::Orient_Horizontal);
+	if (InArgs._bAllowDragAndDrop)
+	{
+		ButtonContent->AddSlot()
+			.Padding(0.0f)
+			.HAlign(HAlign_Center)
+			.VAlign(VAlign_Center)
+			.MaxSize(16.0f)
+			[
+				SNew(SImage)
+				.Image(FCoreStyle::Get().GetBrush("VerticalBoxDragIndicatorShort"))
+				.Cursor(EMouseCursor::GrabHand)
+				.Visibility_Lambda([Weak = AsWeak()]()
+					{
+						if (const TSharedPtr<SWidget> Button = Weak.Pin())
+						{
+							return Button->IsHovered() ? EVisibility::Visible : EVisibility::Hidden;
+						}
+						return EVisibility::Hidden;
+					})
+			];
+	}
+	ButtonContent->AddSlot()
+		.Padding(0.0f)
+		.HAlign(HAlign_Center)
+		.VAlign(VAlign_Center)
+		.AutoSize()
+		[
+			InArgs._Content.Widget
+		];
 
+	SButton::FArguments ButtonArgs;
 	ButtonArgs
+	.AllowDragDrop(InArgs._bAllowDragAndDrop)
 	.OnClicked_Lambda([this]()
 	{
 		StateTreeViewModel->BringNodeToFocus(OwnerStateWeak.Get(), NodeID);
@@ -41,7 +74,7 @@ void SStateTreeContextMenuButton::Construct(const FArguments& InArgs, const TSha
 		.Placement(MenuPlacement_BelowAnchor)
 		.OnGetMenuContent(FOnGetContent::CreateSP(this, &SStateTreeContextMenuButton::MakeContextMenu))
 		[
-			InArgs._Content.Widget
+			ButtonContent
 		]
 	];
 
@@ -80,6 +113,12 @@ FReply SStateTreeContextMenuButton::OnMouseButtonUp(const FGeometry& MyGeometry,
 	return SButton::OnMouseButtonUp(MyGeometry, MouseEvent);
 }
 
+FReply SStateTreeContextMenuButton::OnDragDetected(const FGeometry& MyGeometry, const FPointerEvent& MouseEvent)
+{
+	const bool bIsControlDown = MouseEvent.GetModifierKeys().IsControlDown();
+	return FReply::Handled().BeginDragDrop(FStateTreeSelectedDragDrop::New(StateTreeViewModel, NodeID, OwnerStateWeak, bIsControlDown));
+}
+
 TSharedRef<SWidget> SStateTreeContextMenuButton::MakeContextMenu() const
 {
 	FMenuBuilder MenuBuilder(/*ShouldCloseWindowAfterMenuSelection*/true, /*CommandList*/nullptr);
diff --git a/Engine/Plugins/Runtime/StateTree/Source/StateTreeEditorModule/Private/Customizations/Widgets/SStateTreeContextMenuButton.h b/Engine/Plugins/Runtime/StateTree/Source/StateTreeEditorModule/Private/Customizations/Widgets/SStateTreeContextMenuButton.h
--- a/Engine/Plugins/Runtime/StateTree/Source/StateTreeEditorModule/Private/Customizations/Widgets/SStateTreeContextMenuButton.h
+++ b/Engine/Plugins/Runtime/StateTree/Source/StateTreeEditorModule/Private/Customizations/Widgets/SStateTreeContextMenuButton.h
@@ -17,7 +17,8 @@ public:
 	SLATE_BEGIN_ARGS(SStateTreeContextMenuButton)
 		: _Content()
 		, _ButtonStyle(&FCoreStyle::Get().GetWidgetStyle< FButtonStyle >("Button"))
-		, _ContentPadding(FMargin(4.0, 2.0))
+		, _ContentPadding(FMargin(4.0f, 2.0f))
+		, _bAllowDragAndDrop(false)
 		{
 		}
 
@@ -30,12 +31,16 @@ public:
 		/** Spacing between button's border and the content. */
 		SLATE_ATTRIBUTE(FMargin, ContentPadding)
 
+		/** True if the button should detect drag on mouse down. */
+		SLATE_ARGUMENT(bool, bAllowDragAndDrop)
+
 	SLATE_END_ARGS()
 
 	void Construct(const FArguments& InArgs, const TSharedRef<FStateTreeViewModel>& InStateTreeViewModel, const TWeakObjectPtr<UStateTreeState> InOwnerState, const FGuid& InSourceID, bool InbIsTransition = false);
 
 	//~ SWidget interface
 	virtual FReply OnMouseButtonUp(const FGeometry& MyGeometry, const FPointerEvent& MouseEvent) override;
+	virtual FReply OnDragDetected(const FGeometry& MyGeometry, const FPointerEvent& MouseEvent) override;
 	//~ SWidget overrides
 
 	TSharedRef<SWidget> MakeContextMenu() const;
diff --git a/Engine/Plugins/Runtime/StateTree/Source/StateTreeEditorModule/Private/SStateTreeViewRow.cpp b/Engine/Plugins/Runtime/StateTree/Source/StateTreeEditorModule/Private/SStateTreeViewRow.cpp
--- a/Engine/Plugins/Runtime/StateTree/Source/StateTreeEditorModule/Private/SStateTreeViewRow.cpp
+++ b/Engine/Plugins/Runtime/StateTree/Source/StateTreeEditorModule/Private/SStateTreeViewRow.cpp
@@ -468,9 +468,10 @@ TSharedRef<SWidget> SStateTreeViewRow::MakeTasksWidget(const TSharedPtr<SScrollB
 				};
 
 			TasksBox->AddSlot()
-				.Padding(FMargin(0, 0, 6, 0))
+				.Padding(FMargin(0.0f))
 				[
 					SNew(SStateTreeContextMenuButton, StateTreeViewModel.ToSharedRef(), WeakState, TaskId)
+					.bAllowDragAndDrop(true)
 					.ButtonStyle(FAppStyle::Get(), "SimpleButton")
 					.ContentPadding(FMargin(0, 0))
 					[
@@ -747,8 +748,9 @@ TSharedRef<SWidget> SStateTreeViewRow::MakeConditionsWidget(const TSharedPtr<SSc
 									[
 
 										SNew(SStateTreeContextMenuButton, StateTreeViewModel.ToSharedRef(), WeakState, ConditionId)
+										.bAllowDragAndDrop(true)
 										.ButtonStyle(FAppStyle::Get(), "SimpleButton")
-										.ContentPadding(FMargin(2.f, 0.f))
+										.ContentPadding(FMargin(0.0f))
 										[
 											SNew(SHorizontalBox)
 										
@@ -2244,6 +2246,18 @@ TOptional<EItemDropZone> SStateTreeViewRow::HandleCanAcceptDrop(const FDragDropE
 	if (DragDropOperation.IsValid())
 	{
 		DragDropOperation->SetCanDrop(true);
+		
+		if (DragDropOperation->IsDraggingNode())
+		{
+			const bool bIsSameState = DragDropOperation->OriginState == TargetState;
+			DragDropOperation->SetCanDrop(!bIsSameState);
+			if (bIsSameState || DropZone != EItemDropZone::OntoItem)
+			{
+				return TOptional<EItemDropZone>();
+			}
+			
+			return DropZone;
+		}
 
 		// Cannot drop on selection or child of selection.
 		if (StateTreeViewModel && StateTreeViewModel->IsChildOfSelection(TargetState.Get()))
@@ -2265,17 +2279,32 @@ FReply SStateTreeViewRow::HandleAcceptDrop(const FDragDropEvent& DragDropEvent,
 	{
 		if (StateTreeViewModel)
 		{
-			if (DropZone == EItemDropZone::AboveItem)
-			{
-				StateTreeViewModel->MoveSelectedStatesBefore(TargetState.Get());
-			}
-			else if (DropZone == EItemDropZone::BelowItem)
+			if (DragDropOperation->IsDraggingNode())
 			{
-				StateTreeViewModel->MoveSelectedStatesAfter(TargetState.Get());
+				if (DragDropOperation->bIsControlKeyDown)
+				{
+					StateTreeViewModel->CopyNode(DragDropOperation->OriginState.Get(), DragDropOperation->NodeID);
+					StateTreeViewModel->PasteNodesToStates(MakeConstArrayView({TargetState.Get()}));
+				}
+				else
+				{
+					StateTreeViewModel->MoveNode(DragDropOperation->OriginState, TargetState, DragDropOperation->NodeID);
+				}
 			}
 			else
 			{
-				StateTreeViewModel->MoveSelectedStatesInto(TargetState.Get());
+				if (DropZone == EItemDropZone::AboveItem)
+				{
+					StateTreeViewModel->MoveSelectedStatesBefore(TargetState.Get());
+				}
+				else if (DropZone == EItemDropZone::BelowItem)
+				{
+					StateTreeViewModel->MoveSelectedStatesAfter(TargetState.Get());
+				}
+				else
+				{
+					StateTreeViewModel->MoveSelectedStatesInto(TargetState.Get());
+				}
 			}
 
 			return FReply::Handled();
diff --git a/Engine/Plugins/Runtime/StateTree/Source/StateTreeEditorModule/Private/StateTreeDragDrop.cpp b/Engine/Plugins/Runtime/StateTree/Source/StateTreeEditorModule/Private/StateTreeDragDrop.cpp
--- a/Engine/Plugins/Runtime/StateTree/Source/StateTreeEditorModule/Private/StateTreeDragDrop.cpp
+++ b/Engine/Plugins/Runtime/StateTree/Source/StateTreeEditorModule/Private/StateTreeDragDrop.cpp
@@ -1,4 +1,4 @@
-﻿// Copyright Epic Games, Inc. All Rights Reserved.
+// Copyright Epic Games, Inc. All Rights Reserved.
 
 #include "StateTreeDragDrop.h"
 
@@ -8,28 +8,52 @@
 #include "Widgets/Layout/SBorder.h"
 #include "Widgets/Text/STextBlock.h"
 
+#define LOCTEXT_NAMESPACE "FStateTreeSelectedDragDrop"
+
 TSharedPtr<SWidget> FStateTreeSelectedDragDrop::GetDefaultDecorator() const
 {
 	// Display all dragged nodes, enable the ones that can be moved into the current target node. 
 	const TSharedRef<SVerticalBox> Box = SNew(SVerticalBox);
+	
+	TWeakPtr<const FStateTreeSelectedDragDrop> Weak = SharedThis(this);
+	auto AddTextToBox = [Weak, &Box](const FText& Text)
+	{
+		Box->AddSlot()
+		.Padding(FMargin(4.0f, 2.0f))
+		[
+			SNew(STextBlock)
+			.Text(Text)
+			.IsEnabled_Lambda([Weak]()
+			{
+				TSharedPtr<const FStateTreeSelectedDragDrop> DragDrop = Weak.Pin();
+				return DragDrop && DragDrop->bCanDrop;
+			})
+		];
+	};
+
+	if (IsDraggingNode())
+	{
+		const FStateTreeEditorNode* Node = OriginState->GetNodeByID(NodeID);
+		FText NodeDescription;
+		if (Node)
+		{
+			NodeDescription = ViewModel->GetNodeDescription(*Node, EStateTreeNodeFormatting::Text);
+		}
 
-	if (ViewModel)
+		FFormatNamedArguments Args;
+		Args.Add(TEXT("ActionName"), bIsControlKeyDown ? LOCTEXT("StateTreeDragDropCopyAction", "Copy") : LOCTEXT("StateTreeDragDropMoveAction", "Move"));
+		Args.Add(TEXT("NodeDescription"), NodeDescription);
+		FText ActionText = FText::Format(INVTEXT("{ActionName} {NodeDescription}"), Args);
+		AddTextToBox(ActionText);
+	}
+	else if (ViewModel)
 	{
 		TArray<UStateTreeState*> SelectedStates;
 		ViewModel->GetSelectedStates(SelectedStates);
 
 		for (UStateTreeState* State : SelectedStates)
 		{
-			Box->AddSlot()
-			.Padding(FMargin(4, 2))
-			[
-				SNew(STextBlock)
-				.Text(FText::FromName(State->Name))
-				.IsEnabled_Lambda([this]()
-				{
-					return bCanDrop;
-				})
-			];
+			AddTextToBox(FText::FromName(State->Name));
 		}
 	}
 
@@ -39,3 +63,5 @@ TSharedPtr<SWidget> FStateTreeSelectedDragDrop::GetDefaultDecorator() const
 			Box
 		];
 }
+
+#undef LOCTEXT_NAMESPACE
\ No newline at end of file
diff --git a/Engine/Plugins/Runtime/StateTree/Source/StateTreeEditorModule/Private/StateTreeDragDrop.h b/Engine/Plugins/Runtime/StateTree/Source/StateTreeEditorModule/Private/StateTreeDragDrop.h
--- a/Engine/Plugins/Runtime/StateTree/Source/StateTreeEditorModule/Private/StateTreeDragDrop.h
+++ b/Engine/Plugins/Runtime/StateTree/Source/StateTreeEditorModule/Private/StateTreeDragDrop.h
@@ -4,6 +4,8 @@
 
 #include "DragAndDrop/DecoratedDragDropOp.h"
 
+struct FStateTreeEditorNode;
+class UStateTreeState;
 class FStateTreeViewModel;
 
 class FStateTreeSelectedDragDrop : public FDecoratedDragDropOp
@@ -19,6 +21,18 @@ public:
 
 		return Operation;
 	}
+	
+	static TSharedRef<FStateTreeSelectedDragDrop> New(TSharedPtr<FStateTreeViewModel> InViewModel, FGuid NodeID, TWeakObjectPtr<UStateTreeState> OriginState, bool bIsControlKeyDown)
+	{
+		TSharedRef<FStateTreeSelectedDragDrop> Operation = MakeShared<FStateTreeSelectedDragDrop>();
+		Operation->ViewModel = InViewModel;
+		Operation->OriginState = OriginState;
+		Operation->NodeID = NodeID;
+		Operation->bIsControlKeyDown = bIsControlKeyDown;
+		Operation->Construct();
+
+		return Operation;
+	}
 
 	virtual TSharedPtr<SWidget> GetDefaultDecorator() const override;
 
@@ -26,7 +40,15 @@ public:
 	{
 		bCanDrop = bState;
 	}
+	
+	bool IsDraggingNode() const
+	{
+		return ViewModel.IsValid() && OriginState.IsValid() && NodeID.IsValid();
+	}
 
 	TSharedPtr<FStateTreeViewModel> ViewModel;
+	TWeakObjectPtr<UStateTreeState> OriginState;
+	FGuid NodeID;
+	bool bIsControlKeyDown = false;
 	bool bCanDrop = false;
 };
\ No newline at end of file
diff --git a/Engine/Plugins/Runtime/StateTree/Source/StateTreeEditorModule/Private/StateTreeViewModel.cpp b/Engine/Plugins/Runtime/StateTree/Source/StateTreeEditorModule/Private/StateTreeViewModel.cpp
--- a/Engine/Plugins/Runtime/StateTree/Source/StateTreeEditorModule/Private/StateTreeViewModel.cpp
+++ b/Engine/Plugins/Runtime/StateTree/Source/StateTreeEditorModule/Private/StateTreeViewModel.cpp
@@ -1238,25 +1238,29 @@ void FStateTreeViewModel::PasteNode(TWeakObjectPtr<UStateTreeState> State, const
 }
 
 void FStateTreeViewModel::PasteNodesToSelectedStates()
+{
+	TArray<UStateTreeState*> CurrentlySelectedStates;
+	GetSelectedStates(CurrentlySelectedStates);
+	PasteNodesToStates(CurrentlySelectedStates);
+}
+
+void FStateTreeViewModel::PasteNodesToStates(TConstArrayView<UStateTreeState*> States)
 {
 	if (UStateTreeEditorData* EditorData = TreeDataWeak.Get())
 	{
-		TArray<UStateTreeState*> CurrentlySelectedStates;
-		GetSelectedStates(CurrentlySelectedStates);
-
-		if (CurrentlySelectedStates.Num())
+		if (!States.IsEmpty())
 		{
 			using namespace UE::StateTreeEditor;
 			TArray<FClipboardEditorData, TInlineAllocator<4>> Clipboards;
-			Clipboards.AddDefaulted(CurrentlySelectedStates.Num());
+			Clipboards.AddDefaulted(States.Num());
 
-			for (int32 Idx = 0; Idx < CurrentlySelectedStates.Num(); ++Idx)
+			for (int32 Idx = 0; Idx < States.Num(); ++Idx)
 			{
 				FClipboardEditorData& Clipboard = Clipboards[Idx];
 
 				// any editor nodes and transitions
 				constexpr const UScriptStruct* TargetType = nullptr;
-				ImportTextAsClipboardEditorData(TargetType, EditorData, CurrentlySelectedStates[Idx], Clipboard);
+				ImportTextAsClipboardEditorData(TargetType, EditorData, States[Idx], Clipboard);
 
 				if (!Clipboard.IsValid())
 				{
@@ -1267,9 +1271,9 @@ void FStateTreeViewModel::PasteNodesToSelectedStates()
 			FScopedTransaction ScopedTransaction(LOCTEXT("PasteNodesToSelectedStatesTransaction", "Paste Nodes To Selected States"));
 			EditorData->Modify();
 
-			for (int32 Idx = 0; Idx < CurrentlySelectedStates.Num(); ++Idx)
+			for (int32 Idx = 0; Idx < States.Num(); ++Idx)
 			{
-				UStateTreeState* SelectedState = CurrentlySelectedStates[Idx];
+				UStateTreeState* SelectedState = States[Idx];
 
 				SelectedState->Modify();
 
@@ -1279,6 +1283,83 @@ void FStateTreeViewModel::PasteNodesToSelectedStates()
 	}
 }
 
+void FStateTreeViewModel::MoveNode(TWeakObjectPtr<UStateTreeState> From, TWeakObjectPtr<UStateTreeState> To, const FGuid& ID)
+{
+	if (UStateTreeEditorData* EditorData = TreeDataWeak.Get())
+	{
+		UStateTreeState* FromState = From.Get();
+		UStateTreeState* ToState = To.Get();
+		if (FromState && ToState)
+		{
+			FScopedTransaction ScopedTransaction(LOCTEXT("MoveNodeTransaction", "Move Node To New State"));
+			const UE::StateTree::Editor::Private::FStateTreeStateNodeEditorHandle StateNodeHandle(EditorData, FromState, ID);
+			if (!StateNodeHandle.IsValid())
+			{
+				return;
+			}
+			
+			FromState->Modify();
+			ToState->Modify();
+			StateNodeHandle.EditorData->Modify();
+			
+			if (StateNodeHandle.IsTransition())
+			{
+				FStateTreeTransition& Transition = StateNodeHandle.GetTransitionArray()[StateNodeHandle.GetNodeIndex()];
+				UE::StateTree::PropertyHelpers::ModifyStateInPreAndPostEdit(
+					LOCTEXT("MoveNodeToState", "Move Node To State"),
+					ToState,
+					EditorData,
+					TEXT("Transitions"),
+					[&Transition](TNotNull<UStateTreeState*> InOwnerState, TNotNull<UStateTreeEditorData*> InEditorData, const FStateTreeEditPropertyPath& InPropertyPath)
+					{
+						InOwnerState->Transitions.Add(MoveTemp(Transition));
+					},
+					INDEX_NONE,
+					EPropertyChangeType::ValueSet);
+				DeleteNode(From, ID);
+			}
+			else
+			{
+				TArray<FStateTreeEditorNode>* TargetArray = nullptr;
+				const UScriptStruct* NodeType = StateNodeHandle.GetEditorNode().GetNode().GetScriptStruct();
+				FString RelativeNodePath;
+				if (NodeType->IsChildOf<FStateTreeTaskBase>())
+				{
+					TargetArray = &ToState->Tasks;
+					RelativeNodePath = "Tasks";
+				}
+				else if (NodeType->IsChildOf<FStateTreeConditionBase>())
+				{
+					TargetArray = &ToState->EnterConditions;
+					RelativeNodePath = "EnterConditions";
+				}
+				else if (NodeType->IsChildOf<FStateTreeConsiderationBase>())
+				{
+					TargetArray = &ToState->Considerations;
+					RelativeNodePath = "Considerations";
+				}
+				
+				FStateTreeEditorNode& EditorNode = StateNodeHandle.GetEditorNodeArray()[StateNodeHandle.GetNodeIndex()];
+				UE::StateTree::PropertyHelpers::ModifyStateInPreAndPostEdit(
+					LOCTEXT("MoveNodeToState", "Move Node To State"),
+					ToState,
+					EditorData,
+					RelativeNodePath,
+					[&TargetArray, &EditorNode](TNotNull<UStateTreeState*> InOwnerState, TNotNull<UStateTreeEditorData*> InEditorData, const FStateTreeEditPropertyPath& InPropertyPath)
+					{
+						TargetArray->Add(MoveTemp(EditorNode));
+					},
+					INDEX_NONE,
+					EPropertyChangeType::ValueSet);
+				DeleteNode(From, ID);
+			}
+			
+			OnStateNodesChanged.Broadcast(FromState);
+			OnStateNodesChanged.Broadcast(ToState);
+		}
+	}
+}
+
 void FStateTreeViewModel::DuplicateNode(TWeakObjectPtr<UStateTreeState> State, const FGuid& ID)
 {
 	if (UStateTreeEditorData* EditorData = TreeDataWeak.Get())
diff --git a/Engine/Plugins/Runtime/StateTree/Source/StateTreeEditorModule/Public/StateTreeState.h b/Engine/Plugins/Runtime/StateTree/Source/StateTreeEditorModule/Public/StateTreeState.h
--- a/Engine/Plugins/Runtime/StateTree/Source/StateTreeEditorModule/Public/StateTreeState.h
+++ b/Engine/Plugins/Runtime/StateTree/Source/StateTreeEditorModule/Public/StateTreeState.h
@@ -351,6 +351,38 @@ public:
 				return Other.ID == TransitionID;
 			});
 	}
+	
+	const FStateTreeEditorNode* GetNodeByID(const FGuid NodeID) const
+	{
+		for (const FStateTreeEditorNode& EnterCondition : EnterConditions)
+		{
+			if (EnterCondition.ID == NodeID)
+			{
+				return &EnterCondition;
+			}
+		}
+		for (const FStateTreeEditorNode& Task : Tasks)
+		{
+			if (Task.ID == NodeID)
+			{
+				return &Task;
+			}
+		}
+		for (const FStateTreeEditorNode& Consideration : Considerations)
+		{
+			if (Consideration.ID == NodeID)
+			{
+				return &Consideration;
+			}
+		}
+		
+		if (SingleTask.ID == NodeID)
+		{
+			return &SingleTask;
+		}
+		
+		return nullptr;
+	}
 
 	FGuid GetEventID() const
 	{
diff --git a/Engine/Plugins/Runtime/StateTree/Source/StateTreeEditorModule/Public/StateTreeViewModel.h b/Engine/Plugins/Runtime/StateTree/Source/StateTreeEditorModule/Public/StateTreeViewModel.h
--- a/Engine/Plugins/Runtime/StateTree/Source/StateTreeEditorModule/Public/StateTreeViewModel.h
+++ b/Engine/Plugins/Runtime/StateTree/Source/StateTreeEditorModule/Public/StateTreeViewModel.h
@@ -189,6 +189,8 @@ public:
 	UE_API void CopyAllNodes(TWeakObjectPtr<UStateTreeState> State, const FGuid& ID);
 	UE_API void PasteNode(TWeakObjectPtr<UStateTreeState> State, const FGuid& ID);
 	UE_API void PasteNodesToSelectedStates();
+	UE_API void PasteNodesToStates(TConstArrayView<UStateTreeState*> States);
+	UE_API void MoveNode(TWeakObjectPtr<UStateTreeState> From, TWeakObjectPtr<UStateTreeState> To, const FGuid& ID);
 	UE_API void DuplicateNode(TWeakObjectPtr<UStateTreeState> State, const FGuid& ID);
 
 	// Force to update the view externally.
