[StateTree] Clean up the build.cs. Remove duplicated initialization.

#rnx
#Synced-CL 49562939

diff --git a/Engine/Plugins/Runtime/StateTree/Source/StateTreeModule/StateTreeModule.Build.cs b/Engine/Plugins/Runtime/StateTree/Source/StateTreeModule/StateTreeModule.Build.cs
--- a/Engine/Plugins/Runtime/StateTree/Source/StateTreeModule/StateTreeModule.Build.cs
+++ b/Engine/Plugins/Runtime/StateTree/Source/StateTreeModule/StateTreeModule.Build.cs
@@ -117,7 +117,6 @@ namespace UnrealBuildTool.Rules
 
 			if (IsStateTreeDebuggerSupported(Target))
 			{
-				PublicDefinitions.Add("WITH_STATETREE_TRACE_DEBUGGER=1");
 				PublicDependencyModuleNames.AddRange(
 					new[]
 					{
