[StateTree] Rename internal property name to match the UBT-generated name. Those functions use the native version. The BP version is not affected.
#Synced-CL 49649089

diff --git a/Engine/Plugins/Runtime/StateTree/Source/StateTreeModule/Private/StateTreePropertyRefHelpers.cpp b/Engine/Plugins/Runtime/StateTree/Source/StateTreeModule/Private/StateTreePropertyRefHelpers.cpp
--- a/Engine/Plugins/Runtime/StateTree/Source/StateTreeModule/Private/StateTreePropertyRefHelpers.cpp
+++ b/Engine/Plugins/Runtime/StateTree/Source/StateTreeModule/Private/StateTreePropertyRefHelpers.cpp
@@ -23,9 +23,9 @@ namespace UE::StateTree::PropertyRefHelpers
 	static const FLazyName Int64Name = "int64";
 	static const FLazyName FloatName = "float";
 	static const FLazyName DoubleName = "double";
-	static const FLazyName NameName = "Name";
-	static const FLazyName StringName = "String";
-	static const FLazyName TextName = "Text";
+	static const FLazyName NameName = "FName";
+	static const FLazyName StringName = "FString";
+	static const FLazyName TextName = "FText";
 	const FName IsRefToArrayName = "IsRefToArray";
 	const FName CanRefToArrayName = "CanRefToArray";
 	const FName RefTypeName = "RefType";
