[StateTree] Save the user settings when edited from the UX in the config file. Remove "none" from the drop-down list.
#Synced-CL 49649089

diff --git a/Engine/Plugins/Runtime/StateTree/Source/StateTreeEditorModule/Private/SStateTreeView.cpp b/Engine/Plugins/Runtime/StateTree/Source/StateTreeEditorModule/Private/SStateTreeView.cpp
--- a/Engine/Plugins/Runtime/StateTree/Source/StateTreeEditorModule/Private/SStateTreeView.cpp
+++ b/Engine/Plugins/Runtime/StateTree/Source/StateTreeEditorModule/Private/SStateTreeView.cpp
@@ -625,32 +625,64 @@ TSharedRef<SWidget> SStateTreeView::HandleGenerateSettingsMenu()
 				{
 					DetailLayout.HideCategory("OtherStuff");
 					{
-						IDetailCategoryBuilder& CategoryBuilder = DetailLayout.EditCategory("State View");
-						TArray<TSharedRef<IPropertyHandle>> AllProperties;
-						CategoryBuilder.GetDefaultProperties(AllProperties);
+						auto FindProperty = [](TArrayView<TSharedRef<IPropertyHandle>> AllProperties, FName PropertyToFind)
+							{
+								return AllProperties.FindByPredicate(
+									[PropertyToFind](TSharedRef<IPropertyHandle>& Other)
+									{
+										return Other->GetProperty()->GetFName() == PropertyToFind;
+									});
+							};
+
+						{
+							IDetailCategoryBuilder& CategoryBuilder = DetailLayout.EditCategory("State View");
+							TArray<TSharedRef<IPropertyHandle>> AllProperties;
+							CategoryBuilder.GetDefaultProperties(AllProperties);
 
-						const FName PropertyToFind = "StatesViewDisplayNodeType";
-						TSharedRef<IPropertyHandle>* FoundProperty = AllProperties.FindByPredicate([PropertyToFind](TSharedRef<IPropertyHandle>& Other)
+							TSharedRef<IPropertyHandle>* FoundProperty = FindProperty(AllProperties, "StatesViewDisplayNodeType");
+							if (ensure(FoundProperty))
 							{
-								return Other->GetProperty()->GetFName() == PropertyToFind;
-							});
-						if (ensure(FoundProperty))
+								CategoryBuilder.AddProperty(*FoundProperty).CustomWidget()
+									.NameContent()
+									[
+										(*FoundProperty)->CreatePropertyNameWidget()
+									]
+									.ValueContent()
+									[
+										SNew(SEnumComboBox, StaticEnum<EStateTreeEditorUserSettingsNodeType>())
+											.OnEnumSelectionChanged(SEnumComboBox::FOnEnumSelectionChanged::CreateLambda([](int32 NewValue, ESelectInfo::Type)
+												{
+													GetMutableDefault<UStateTreeEditorUserSettings>()->SetStatesViewDisplayNodeType((EStateTreeEditorUserSettingsNodeType)NewValue);
+												}))
+											.CurrentValue(MakeAttributeLambda([]() { return (int32)GetDefault<UStateTreeEditorUserSettings>()->GetStatesViewDisplayNodeType(); }))
+											.Font(IDetailLayoutBuilder::GetDetailFont())
+									];
+							}
+						}
 						{
-							CategoryBuilder.AddProperty(*FoundProperty).CustomWidget()
-							.NameContent()
-							[
-								(*FoundProperty)->CreatePropertyNameWidget()
-							]
-							.ValueContent()
-							[
-								SNew(SEnumComboBox, StaticEnum<EStateTreeEditorUserSettingsNodeType>())
-									.OnEnumSelectionChanged(SEnumComboBox::FOnEnumSelectionChanged::CreateLambda([](int32 NewValue, ESelectInfo::Type)
-										{
-											GetMutableDefault<UStateTreeEditorUserSettings>()->SetStatesViewDisplayNodeType((EStateTreeEditorUserSettingsNodeType)NewValue);
-										}))
-									.CurrentValue(MakeAttributeLambda([]() { return (int32)GetDefault<UStateTreeEditorUserSettings>()->GetStatesViewDisplayNodeType(); }))
+							IDetailCategoryBuilder& CategoryBuilder = DetailLayout.EditCategory("Details View");
+							TArray<TSharedRef<IPropertyHandle>> AllProperties;
+							CategoryBuilder.GetDefaultProperties(AllProperties);
+
+							TSharedRef<IPropertyHandle>* FoundProperty = FindProperty(AllProperties, "DetailsViewDisplayCount");
+							if (ensure(FoundProperty))
+							{
+								CategoryBuilder.AddProperty(*FoundProperty).CustomWidget()
+									.NameContent()
+									[
+										(*FoundProperty)->CreatePropertyNameWidget()
+									]
+									.ValueContent()
+									[
+										SNew(SEnumComboBox, StaticEnum<EStateTreeEditorUserSettingsNodeType>())
+											.OnEnumSelectionChanged(SEnumComboBox::FOnEnumSelectionChanged::CreateLambda([](int32 NewValue, ESelectInfo::Type)
+												{
+													GetMutableDefault<UStateTreeEditorUserSettings>()->SetDetailsViewDisplayCount((EStateTreeEditorUserSettingsNodeType)NewValue);
+												}))
+											.CurrentValue(MakeAttributeLambda([]() { return (int32)GetDefault<UStateTreeEditorUserSettings>()->GetDetailsViewDisplayCount(); }))
 											.Font(IDetailLayoutBuilder::GetDetailFont())
-							];
+									];
+							}
 						}
 					}
 				}
diff --git a/Engine/Plugins/Runtime/StateTree/Source/StateTreeEditorModule/Private/StateTreeEditorUserSettings.cpp b/Engine/Plugins/Runtime/StateTree/Source/StateTreeEditorModule/Private/StateTreeEditorUserSettings.cpp
--- a/Engine/Plugins/Runtime/StateTree/Source/StateTreeEditorModule/Private/StateTreeEditorUserSettings.cpp
+++ b/Engine/Plugins/Runtime/StateTree/Source/StateTreeEditorModule/Private/StateTreeEditorUserSettings.cpp
@@ -40,6 +40,7 @@ void UStateTreeEditorUserSettings::PostEditChangeChainProperty(struct FPropertyC
 		|| PropertyName == GET_MEMBER_NAME_CHECKED(UStateTreeEditorUserSettings, DetailsViewDisplayCount))
 	{
 		OnSettingsChanged.Broadcast();
+		SaveConfig();
 	}
 
 	Super::PostEditChangeChainProperty(PropertyChangedEvent);
diff --git a/Engine/Plugins/Runtime/StateTree/Source/StateTreeEditorModule/Private/StateTreeEditorUserSettings.h b/Engine/Plugins/Runtime/StateTree/Source/StateTreeEditorModule/Private/StateTreeEditorUserSettings.h
--- a/Engine/Plugins/Runtime/StateTree/Source/StateTreeEditorModule/Private/StateTreeEditorUserSettings.h
+++ b/Engine/Plugins/Runtime/StateTree/Source/StateTreeEditorModule/Private/StateTreeEditorUserSettings.h
@@ -8,7 +8,7 @@
 UENUM(meta=(Bitflags, UseEnumValuesAsMaskValuesInEditor = "true"))
 enum class EStateTreeEditorUserSettingsNodeType : uint8
 {
-	None = 0 << 0,
+	None = 0 << 0 UMETA(Hidden),
 	Condition = 1 << 0,
 	Task = 1 << 1,
 	Transition = 1 << 2,
@@ -82,6 +82,6 @@ private:
 	EStateTreeEditorUserSettingsNodeType DetailsViewDisplayCount = EStateTreeEditorUserSettingsNodeType::All;
 
 	/** Display the debugger migration warning. */
-	UPROPERTY(config, EditAnywhere, Category = "Debugger")
+	UPROPERTY(config)
 	bool bDisplayMigrationToRewindDebuggerWarning = true;
 };
