[StateTree] Fixed an issue where transition delegates weren't triggered if the delegate was broadcasted during the transition tick.
#Synced-CL 49649089

diff --git a/Engine/Plugins/Runtime/StateTree/Source/StateTreeModule/Private/StateTreeExecutionContext.cpp b/Engine/Plugins/Runtime/StateTree/Source/StateTreeModule/Private/StateTreeExecutionContext.cpp
--- a/Engine/Plugins/Runtime/StateTree/Source/StateTreeModule/Private/StateTreeExecutionContext.cpp
+++ b/Engine/Plugins/Runtime/StateTree/Source/StateTreeModule/Private/StateTreeExecutionContext.cpp
@@ -157,7 +157,7 @@ namespace UE::StateTree::ExecutionContext::Private
 
 	static constexpr int32 ExpectedAmountTransitionEvent = 8;
 	using FSharedEventInlineArray = TArray<FStateTreeSharedEvent, TInlineAllocator<ExpectedAmountTransitionEvent, FNonconcurrentLinearArrayAllocator>>;
-	void GetTriggerTransitionEvent(const FCompactStateTransition& Transition, FStateTreeInstanceStorage& Storage, const FStateTreeSharedEvent& TransitionEvent, const TArrayView<const FStateTreeSharedEvent> EventsToProcess, FSharedEventInlineArray& OutTransitionEvents)
+	void GetTriggerTransitionEvent(const FCompactStateTransition& Transition, TConstArrayView<FStateTreeDelegateDispatcher> BroadcastedDelegates, const FStateTreeSharedEvent& TransitionEvent, const TArrayView<const FStateTreeSharedEvent> EventsToProcess, FSharedEventInlineArray& OutTransitionEvents)
 	{
 		if (Transition.Trigger == EStateTreeTransitionTrigger::OnEvent)
 		{
@@ -186,7 +186,7 @@ namespace UE::StateTree::ExecutionContext::Private
 		}
 		else if (EnumHasAnyFlags(Transition.Trigger, EStateTreeTransitionTrigger::OnDelegate))
 		{
-			if (Storage.IsDelegateBroadcasted(Transition.RequiredDelegateDispatcher))
+			if (BroadcastedDelegates.Contains(Transition.RequiredDelegateDispatcher))
 			{
 				// Dummy event to make sure we iterate the transition loop once.
 				OutTransitionEvents.Emplace(FStateTreeSharedEvent());
@@ -1952,7 +1952,7 @@ void FStateTreeExecutionContext::TickTriggerTransitionsInternal()
 		}
 
 		// Stop as soon as have found a running state.
-		if (Exec.LastTickStatus == EStateTreeRunStatus::Running)
+		if (Exec.LastTickStatus == EStateTreeRunStatus::Running && !Storage.HasBroadcastedDelegates())
 		{
 			break;
 		}
@@ -5602,6 +5602,10 @@ bool FStateTreeExecutionContext::TriggerTransitions()
 
 	RequestedTransition.Reset();
 
+	check(CurrentlyProcessedBroadcastedDelegates.IsEmpty());
+	CurrentlyProcessedBroadcastedDelegates = Storage.StealBroadcastedDelegates();
+	ON_SCOPE_EXIT{ CurrentlyProcessedBroadcastedDelegates.Empty(); };
+
 	//
 	// Process transition requests
 	//
@@ -5683,7 +5687,7 @@ bool FStateTreeExecutionContext::TriggerTransitions()
 	{
 		// Re-cache bHasEvents, RequestTransition above can create new events.
 		const bool bHasEvents = EventQueue && EventQueue->HasEvents();
-		const bool bHasBroadcastedDelegates = Storage.HasBroadcastedDelegates();
+		const bool bHasBroadcastedDelegates = !CurrentlyProcessedBroadcastedDelegates.IsEmpty();
 
 		// Transition() can TriggerTransitions() in a loop when a sub-frame completes.
 		//We do not want to evaluate the transition from that sub-frame.
@@ -5868,7 +5872,7 @@ bool FStateTreeExecutionContext::TriggerTransitions()
 				}
 
 				UE::StateTree::ExecutionContext::Private::FSharedEventInlineArray TransitionEvents;
-				UE::StateTree::ExecutionContext::Private::GetTriggerTransitionEvent(Transition, Storage, FStateTreeSharedEvent(), GetEventsToProcessView(), TransitionEvents);
+				UE::StateTree::ExecutionContext::Private::GetTriggerTransitionEvent(Transition, CurrentlyProcessedBroadcastedDelegates, FStateTreeSharedEvent(), GetEventsToProcessView(), TransitionEvents);
 
 				for (const FStateTreeSharedEvent& TransitionEvent : TransitionEvents)
 				{
@@ -5963,8 +5967,6 @@ PRAGMA_ENABLE_DEPRECATION_WARNINGS
 		EventQueue->Reset();
 	}
 
-	Storage.ResetBroadcastedDelegates();
-
 	//
 	// Check state completion transitions.
 	//
@@ -8050,7 +8052,7 @@ bool FStateTreeExecutionContext::SelectStateInternal_TryFollowTransitions(
 		}
 
 		FSharedEventInlineArray TransitionEvents;
-		GetTriggerTransitionEvent(Transition, Storage, Args.TransitionEvent, GetEventsToProcessView(), TransitionEvents);
+		GetTriggerTransitionEvent(Transition, CurrentlyProcessedBroadcastedDelegates, Args.TransitionEvent, GetEventsToProcessView(), TransitionEvents);
 
 		for (const FStateTreeSharedEvent& SelectedStateTransitionEvent : TransitionEvents)
 		{
diff --git a/Engine/Plugins/Runtime/StateTree/Source/StateTreeModule/Private/StateTreeInstanceData.cpp b/Engine/Plugins/Runtime/StateTree/Source/StateTreeModule/Private/StateTreeInstanceData.cpp
--- a/Engine/Plugins/Runtime/StateTree/Source/StateTreeModule/Private/StateTreeInstanceData.cpp
+++ b/Engine/Plugins/Runtime/StateTree/Source/StateTreeModule/Private/StateTreeInstanceData.cpp
@@ -563,9 +563,9 @@ bool FStateTreeInstanceStorage::IsDelegateBroadcasted(const FStateTreeDelegateDi
 	return BroadcastedDelegates.Contains(Dispatcher);
 }
 
-void FStateTreeInstanceStorage::ResetBroadcastedDelegates()
+TArray<FStateTreeDelegateDispatcher> FStateTreeInstanceStorage::StealBroadcastedDelegates()
 {
-	BroadcastedDelegates.Empty();
+	return MoveTemp(BroadcastedDelegates);
 }
 
 bool FStateTreeInstanceStorage::HasBroadcastedDelegates() const
diff --git a/Engine/Plugins/Runtime/StateTree/Source/StateTreeModule/Public/StateTreeExecutionContext.h b/Engine/Plugins/Runtime/StateTree/Source/StateTreeModule/Public/StateTreeExecutionContext.h
--- a/Engine/Plugins/Runtime/StateTree/Source/StateTreeModule/Public/StateTreeExecutionContext.h
+++ b/Engine/Plugins/Runtime/StateTree/Source/StateTreeModule/Public/StateTreeExecutionContext.h
@@ -1590,6 +1590,9 @@ PRAGMA_ENABLE_DEPRECATION_WARNINGS
 	/** When set, start the transitions loop from TriggerTransitionsFromFrameIndex. */
 	TOptional<int32> TriggerTransitionsFromFrameIndex;
 
+	/** Current broadcasted delegates we're processing. */
+	TArray<FStateTreeDelegateDispatcher> CurrentlyProcessedBroadcastedDelegates;
+
 	/** Current frame we're processing. */
 	const FStateTreeExecutionFrame* CurrentlyProcessedParentFrame = nullptr; 
 	const FStateTreeExecutionFrame* CurrentlyProcessedFrame = nullptr; 
diff --git a/Engine/Plugins/Runtime/StateTree/Source/StateTreeModule/Public/StateTreeInstanceData.h b/Engine/Plugins/Runtime/StateTree/Source/StateTreeModule/Public/StateTreeInstanceData.h
--- a/Engine/Plugins/Runtime/StateTree/Source/StateTreeModule/Public/StateTreeInstanceData.h
+++ b/Engine/Plugins/Runtime/StateTree/Source/StateTreeModule/Public/StateTreeInstanceData.h
@@ -195,6 +195,9 @@ struct FStateTreeInstanceStorage
 	/** Resets the list of broadcasted delegates. */
 	UE_API void ResetBroadcastedDelegates();
 
+	/** Resets the list of broadcasted delegates and returns it's copy. */
+	UE_API TArray<FStateTreeDelegateDispatcher> StealBroadcastedDelegates();
+
 	/** @return true if there's any broadcasted delegates. */
 	UE_API bool HasBroadcastedDelegates() const;
 
diff --git a/Engine/Plugins/Runtime/StateTree/Source/StateTreeTestSuite/Private/StateTreeDelegateTest.cpp b/Engine/Plugins/Runtime/StateTree/Source/StateTreeTestSuite/Private/StateTreeDelegateTest.cpp
--- a/Engine/Plugins/Runtime/StateTree/Source/StateTreeTestSuite/Private/StateTreeDelegateTest.cpp
+++ b/Engine/Plugins/Runtime/StateTree/Source/StateTreeTestSuite/Private/StateTreeDelegateTest.cpp
@@ -170,6 +170,51 @@ struct FStateTreeTest_Delegate_Transitions : FStateTreeTestBase
 };
 IMPLEMENT_STATE_TREE_INSTANT_TEST(FStateTreeTest_Delegate_Transitions, "System.StateTree.Delegate.Transitions");
 
+struct FStateTreeTest_Delegate_BroadcastingInTransitionsTick : FStateTreeTestBase
+{
+	virtual bool InstantTest() override
+	{
+		UStateTree& StateTree = NewStateTree();
+		UStateTreeEditorData& EditorData = *Cast<UStateTreeEditorData>(StateTree.EditorData);
+		UStateTreeState& Root = EditorData.AddSubTree(FName(TEXT("Root")));
+
+		UStateTreeState& StateA = Root.AddChildState("A");
+		UStateTreeState& StateB = Root.AddChildState("B");
+
+		FStateTreeTransition& TransitionAToA = StateA.AddTransition(EStateTreeTransitionTrigger::OnDelegate, EStateTreeTransitionType::GotoState, &StateA);
+		FStateTreeTransition& TransitionRootToB = Root.AddTransition(EStateTreeTransitionTrigger::OnDelegate, EStateTreeTransitionType::GotoState, &StateB);
+
+		TStateTreeEditorNode<FTestTask_BroadcastDelegate>& DispatcherTask0 = Root.AddTask<FTestTask_BroadcastDelegate>(FName(TEXT("DispatcherTask0")));
+		TStateTreeEditorNode<FTestTask_BroadcastDelegate>& DispatcherTask1 = StateA.AddTask<FTestTask_BroadcastDelegate>(FName(TEXT("DispatcherTask1")));
+
+		EditorData.AddPropertyBinding(FPropertyBindingPath(DispatcherTask0.ID, GET_MEMBER_NAME_CHECKED(FTestTask_BroadcastDelegate_InstanceData, OnTickDelegate)), FPropertyBindingPath(TransitionAToA.ID, GET_MEMBER_NAME_CHECKED(FStateTreeTransition, DelegateListener)));
+		EditorData.AddPropertyBinding(FPropertyBindingPath(DispatcherTask1.ID, GET_MEMBER_NAME_CHECKED(FTestTask_BroadcastDelegate_InstanceData, OnExitDelegate)), FPropertyBindingPath(TransitionRootToB.ID, GET_MEMBER_NAME_CHECKED(FStateTreeTransition, DelegateListener)));
+
+		FStateTreeCompilerLog Log;
+		FStateTreeCompiler Compiler(Log);
+		const bool bResult = Compiler.Compile(StateTree);
+		AITEST_TRUE("StateTree should get compiled", bResult);
+
+		FStateTreeInstanceData InstanceData;
+		FTestStateTreeExecutionContext Exec(StateTree, StateTree, InstanceData);
+		const bool bInitSucceeded = Exec.IsValid();
+		AITEST_TRUE("StateTree should init", bInitSucceeded);
+
+		Exec.Start();
+		AITEST_TRUE(TEXT("StateTree Active States should be in Root/A"), Exec.ExpectInActiveStates(Root.Name, StateA.Name));
+		Exec.LogClear();
+
+		Exec.Tick(0.1f);
+		AITEST_TRUE(TEXT("StateTree Active States should be in Root/B"), Exec.ExpectInActiveStates(Root.Name, StateB.Name));
+		Exec.LogClear();
+
+		Exec.Stop();
+
+		return true;
+	}
+};
+IMPLEMENT_STATE_TREE_INSTANT_TEST(FStateTreeTest_Delegate_BroadcastingInTransitionsTick, "System.StateTree.Delegate.BroadcastingInTransitionsTick");
+
 struct FStateTreeTest_Delegate_Rebroadcasting : FStateTreeTestBase
 {
 	virtual bool InstantTest() override
