[StateTree] Create UE::StateTree::EComparisonOperator to replace EGenericAICheck from the AIModule. This is an effort to make StateTree not depend on AI modules and make it generic for everyone to use.
#Synced-CL 49649089

diff --git a/Engine/Plugins/Runtime/StateTree/Source/StateTreeModule/Private/Conditions/StateTreeCommonConditions.cpp b/Engine/Plugins/Runtime/StateTree/Source/StateTreeModule/Private/Conditions/StateTreeCommonConditions.cpp
--- a/Engine/Plugins/Runtime/StateTree/Source/StateTreeModule/Private/Conditions/StateTreeCommonConditions.cpp
+++ b/Engine/Plugins/Runtime/StateTree/Source/StateTreeModule/Private/Conditions/StateTreeCommonConditions.cpp
@@ -2,6 +2,7 @@
 
 #include "Conditions/StateTreeCommonConditions.h"
 
+#include "AITypes.h"
 #include "Conditions/StateTreeConditionHelpers.h"
 #include "StateTreeExecutionContext.h"
 #include "UObject/EnumProperty.h"
@@ -11,10 +12,43 @@
 
 #define LOCTEXT_NAMESPACE "StateTree"
 
+namespace UE::StateTree::Conditions
+{
+	// Deprecated
+	EComparisonOperator GenericAICheckToComparisonOperator(EGenericAICheck Operator)
+	{
+		switch (Operator)
+		{
+		case EGenericAICheck::Equal:
+			return EComparisonOperator::Equal;
+		case EGenericAICheck::NotEqual:
+			return EComparisonOperator::NotEqual;
+		case EGenericAICheck::Less:
+			return EComparisonOperator::Less;
+		case EGenericAICheck::LessOrEqual:
+			return EComparisonOperator::LessOrEqual;
+		case EGenericAICheck::Greater:
+			return EComparisonOperator::Greater;
+		case EGenericAICheck::GreaterOrEqual:
+			return EComparisonOperator::GreaterOrEqual;
+		default:
+			ensureMsgf(false, TEXT("Unhandled operator %d"), Operator);
+			return EComparisonOperator::Equal;
+		}
+	}
+}
+
 //----------------------------------------------------------------------//
 //  FStateTreeCompareIntCondition
 //----------------------------------------------------------------------//
 
+// Deprecated
+PRAGMA_DISABLE_DEPRECATION_WARNINGS
+FStateTreeCompareIntCondition::FStateTreeCompareIntCondition(const EGenericAICheck InOperator, const EStateTreeCompare InInverts)
+	: FStateTreeCompareIntCondition(UE::StateTree::Conditions::GenericAICheckToComparisonOperator(InOperator), InInverts)
+{}
+PRAGMA_ENABLE_DEPRECATION_WARNINGS
+
 bool FStateTreeCompareIntCondition::TestCondition(FStateTreeExecutionContext& Context) const
 {
 	const FInstanceDataType& InstanceData = Context.GetInstanceData(*this);
@@ -63,6 +97,14 @@ FText FStateTreeCompareIntCondition::GetDescription(const FGuid& ID, FStateTreeD
 //  FStateTreeCompareFloatCondition
 //----------------------------------------------------------------------//
 
+// Deprecated
+PRAGMA_DISABLE_DEPRECATION_WARNINGS
+FStateTreeCompareFloatCondition::FStateTreeCompareFloatCondition(const EGenericAICheck InOperator, const EStateTreeCompare InInverts)
+	: FStateTreeCompareFloatCondition(UE::StateTree::Conditions::GenericAICheckToComparisonOperator(InOperator), InInverts)
+{
+}
+PRAGMA_ENABLE_DEPRECATION_WARNINGS
+
 bool FStateTreeCompareFloatCondition::TestCondition(FStateTreeExecutionContext& Context) const
 {
 	const FInstanceDataType& InstanceData = Context.GetInstanceData(*this);
@@ -268,6 +310,14 @@ void FStateTreeCompareEnumCondition::OnBindingChanged(const FGuid& ID, FStateTre
 //  FStateTreeCompareDistanceCondition
 //----------------------------------------------------------------------//
 
+// Deprecated
+PRAGMA_DISABLE_DEPRECATION_WARNINGS
+FStateTreeCompareDistanceCondition::FStateTreeCompareDistanceCondition(const EGenericAICheck InOperator, const EStateTreeCompare InInverts)
+	: FStateTreeCompareDistanceCondition(UE::StateTree::Conditions::GenericAICheckToComparisonOperator(InOperator), InInverts)
+{
+}
+PRAGMA_ENABLE_DEPRECATION_WARNINGS
+
 bool FStateTreeCompareDistanceCondition::TestCondition(FStateTreeExecutionContext& Context) const
 {
 	const FInstanceDataType& InstanceData = Context.GetInstanceData(*this);
diff --git a/Engine/Plugins/Runtime/StateTree/Source/StateTreeModule/Private/StateTreeNodeDescriptionHelpers.cpp b/Engine/Plugins/Runtime/StateTree/Source/StateTreeModule/Private/StateTreeNodeDescriptionHelpers.cpp
--- a/Engine/Plugins/Runtime/StateTree/Source/StateTreeModule/Private/StateTreeNodeDescriptionHelpers.cpp
+++ b/Engine/Plugins/Runtime/StateTree/Source/StateTreeModule/Private/StateTreeNodeDescriptionHelpers.cpp
@@ -1,6 +1,7 @@
 // Copyright Epic Games, Inc. All Rights Reserved.
 
 #include "StateTreeNodeDescriptionHelpers.h"
+#include "Conditions/StateTreeConditionHelpers.h"
 #include "GameplayTagContainer.h"
 #include "StateTreeNodeBase.h"
 
@@ -10,23 +11,23 @@
 namespace UE::StateTree::DescHelpers
 {
 
-FText GetOperatorText(const EGenericAICheck Operator, EStateTreeNodeFormatting Formatting)
+FText GetOperatorText(const EComparisonOperator Operator, EStateTreeNodeFormatting Formatting)
 {
 	if (Formatting == EStateTreeNodeFormatting::RichText)
 	{
 		switch (Operator)
 		{
-		case EGenericAICheck::Equal:
+		case EComparisonOperator::Equal:
 			return FText::FromString(TEXT("<s>==</>"));
-		case EGenericAICheck::NotEqual:
+		case EComparisonOperator::NotEqual:
 			return FText::FromString(TEXT("<s>!=</>"));
-		case EGenericAICheck::Less:
+		case EComparisonOperator::Less:
 			return FText::FromString(TEXT("<s>&lt;</>"));
-		case EGenericAICheck::LessOrEqual:
+		case EComparisonOperator::LessOrEqual:
 			return FText::FromString(TEXT("<s>&lt;=</>"));
-		case EGenericAICheck::Greater:
+		case EComparisonOperator::Greater:
 			return FText::FromString(TEXT("<s>&gt;</>"));
-		case EGenericAICheck::GreaterOrEqual:
+		case EComparisonOperator::GreaterOrEqual:
 			return FText::FromString(TEXT("<s>&gt;=</>"));
 		default:
 			break;
@@ -36,17 +37,17 @@ FText GetOperatorText(const EGenericAICheck Operator, EStateTreeNodeFormatting F
 	{
 		switch (Operator)
 		{
-		case EGenericAICheck::Equal:
+		case EComparisonOperator::Equal:
 			return FText::FromString(TEXT("=="));
-		case EGenericAICheck::NotEqual:
+		case EComparisonOperator::NotEqual:
 			return FText::FromString(TEXT("!="));
-		case EGenericAICheck::Less:
+		case EComparisonOperator::Less:
 			return FText::FromString(TEXT("<"));
-		case EGenericAICheck::LessOrEqual:
+		case EComparisonOperator::LessOrEqual:
 			return FText::FromString(TEXT("<="));
-		case EGenericAICheck::Greater:
+		case EComparisonOperator::Greater:
 			return FText::FromString(TEXT(">"));
-		case EGenericAICheck::GreaterOrEqual:
+		case EComparisonOperator::GreaterOrEqual:
 			return FText::FromString(TEXT(">="));
 		default:
 			break;
@@ -56,6 +57,14 @@ FText GetOperatorText(const EGenericAICheck Operator, EStateTreeNodeFormatting F
 	return FText::FromString(TEXT("??"));
 }
 
+// Deprecated
+FText GetOperatorText(const EGenericAICheck Operator, EStateTreeNodeFormatting Formatting)
+{
+PRAGMA_DISABLE_DEPRECATION_WARNINGS
+	return GetOperatorText(UE::StateTree::Conditions::GenericAICheckToComparisonOperator(Operator), Formatting);
+PRAGMA_ENABLE_DEPRECATION_WARNINGS
+}
+
 FText GetInvertText(bool bInvert, EStateTreeNodeFormatting Formatting)
 {
 	FText InvertText = FText::GetEmpty();
diff --git a/Engine/Plugins/Runtime/StateTree/Source/StateTreeModule/Public/Conditions/StateTreeCommonConditions.h b/Engine/Plugins/Runtime/StateTree/Source/StateTreeModule/Public/Conditions/StateTreeCommonConditions.h
--- a/Engine/Plugins/Runtime/StateTree/Source/StateTreeModule/Public/Conditions/StateTreeCommonConditions.h
+++ b/Engine/Plugins/Runtime/StateTree/Source/StateTreeModule/Public/Conditions/StateTreeCommonConditions.h
@@ -2,13 +2,17 @@
 
 #pragma once
 
+#if UE_ENABLE_INCLUDE_ORDER_DEPRECATED_IN_5_8
 #include "AITypes.h"
+#endif // UE_ENABLE_INCLUDE_ORDER_DEPRECATED_IN_5_8
+
 #include "StateTreeConditionBase.h"
 #include "StateTreeAnyEnum.h"
 #include "StateTreeCommonConditions.generated.h"
 
 #define UE_API STATETREEMODULE_API
 
+enum class EGenericAICheck : uint8;
 struct FStateTreeDataView;
 
 USTRUCT()
@@ -35,11 +39,14 @@ struct FStateTreeCompareIntCondition : public FStateTreeConditionCommonBase
 	using FInstanceDataType = FStateTreeCompareIntConditionInstanceData;
 
 	FStateTreeCompareIntCondition() = default;
-	explicit FStateTreeCompareIntCondition(const EGenericAICheck InOperator, const EStateTreeCompare InInverts = EStateTreeCompare::Default)
+	explicit FStateTreeCompareIntCondition(const UE::StateTree::EComparisonOperator InOperator, const EStateTreeCompare InInverts = EStateTreeCompare::Default)
 		: bInvert(InInverts == EStateTreeCompare::Invert)
 		, Operator(InOperator)
 	{}
 
+	UE_DEPRECATED(5.8, "Use EComparisonOperator instead.")
+	UE_API explicit FStateTreeCompareIntCondition(const EGenericAICheck InOperator, const EStateTreeCompare InInverts = EStateTreeCompare::Default);
+
 	virtual const UStruct* GetInstanceDataType() const override { return FInstanceDataType::StaticStruct(); }
 	UE_API virtual bool TestCondition(FStateTreeExecutionContext& Context) const override;
 #if WITH_EDITOR
@@ -49,8 +56,9 @@ struct FStateTreeCompareIntCondition : public FStateTreeConditionCommonBase
 	UPROPERTY(EditAnywhere, Category = "Parameter")
 	bool bInvert = false;
 
+	UE_DEPRECATED_FORGAME(5.8, "The Operator type changed from EGenericAICheck to EComparisonOperator.")
 	UPROPERTY(EditAnywhere, Category = "Parameter", meta = (InvalidEnumValues = "IsTrue"))
-	EGenericAICheck Operator = EGenericAICheck::Equal;
+	UE::StateTree::EComparisonOperator Operator = UE::StateTree::EComparisonOperator::Equal;
 };
 
 
@@ -78,11 +86,14 @@ struct FStateTreeCompareFloatCondition : public FStateTreeConditionCommonBase
 	using FInstanceDataType = FStateTreeCompareFloatConditionInstanceData;
 
 	FStateTreeCompareFloatCondition() = default;
-	explicit FStateTreeCompareFloatCondition(const EGenericAICheck InOperator, const EStateTreeCompare InInverts = EStateTreeCompare::Default)
+	explicit FStateTreeCompareFloatCondition(const UE::StateTree::EComparisonOperator InOperator, const EStateTreeCompare InInverts = EStateTreeCompare::Default)
 		: bInvert(InInverts == EStateTreeCompare::Invert)
 		, Operator(InOperator)
 	{}
 
+	UE_DEPRECATED(5.8, "Use EComparisonOperator instead.")
+	UE_API explicit FStateTreeCompareFloatCondition(const EGenericAICheck InOperator, const EStateTreeCompare InInverts = EStateTreeCompare::Default);
+
 	virtual const UStruct* GetInstanceDataType() const override { return FInstanceDataType::StaticStruct(); }
 	UE_API virtual bool TestCondition(FStateTreeExecutionContext& Context) const override;
 #if WITH_EDITOR
@@ -92,8 +103,9 @@ struct FStateTreeCompareFloatCondition : public FStateTreeConditionCommonBase
 	UPROPERTY(EditAnywhere, Category = "Parameter")
 	bool bInvert = false;
 
+	UE_DEPRECATED_FORGAME(5.8, "The Operator type changed from EGenericAICheck to EComparisonOperator.")
 	UPROPERTY(EditAnywhere, Category = "Parameter", meta = (InvalidEnumValues = "IsTrue"))
-	EGenericAICheck Operator = EGenericAICheck::Equal;
+	UE::StateTree::EComparisonOperator Operator = UE::StateTree::EComparisonOperator::Equal;
 };
 
 
@@ -211,10 +223,13 @@ struct FStateTreeCompareDistanceCondition : public FStateTreeConditionCommonBase
 	using FInstanceDataType = FStateTreeCompareDistanceConditionInstanceData;
 
 	FStateTreeCompareDistanceCondition() = default;
-	explicit FStateTreeCompareDistanceCondition(const EGenericAICheck InOperator, const EStateTreeCompare InInverts = EStateTreeCompare::Default)
+	explicit FStateTreeCompareDistanceCondition(const UE::StateTree::EComparisonOperator InOperator, const EStateTreeCompare InInverts = EStateTreeCompare::Default)
 		: bInvert(InInverts == EStateTreeCompare::Invert)
 		, Operator(InOperator)
 	{}
+
+	UE_DEPRECATED(5.8, "Use EComparisonOperator instead.")
+	UE_API explicit FStateTreeCompareDistanceCondition(const EGenericAICheck InOperator, const EStateTreeCompare InInverts = EStateTreeCompare::Default);
 	
 	virtual const UStruct* GetInstanceDataType() const override { return FInstanceDataType::StaticStruct(); }
 	UE_API virtual bool TestCondition(FStateTreeExecutionContext& Context) const override;
@@ -225,8 +240,9 @@ struct FStateTreeCompareDistanceCondition : public FStateTreeConditionCommonBase
 	UPROPERTY(EditAnywhere, Category = "Condition")
 	bool bInvert = false;
 
+	UE_DEPRECATED_FORGAME(5.8, "The Operator type changed from EGenericAICheck to EComparisonOperator.")
 	UPROPERTY(EditAnywhere, Category = "Condition", meta = (InvalidEnumValues = "IsTrue"))
-	EGenericAICheck Operator = EGenericAICheck::Equal;
+	UE::StateTree::EComparisonOperator Operator = UE::StateTree::EComparisonOperator::Equal;
 };
 
 USTRUCT()
diff --git a/Engine/Plugins/Runtime/StateTree/Source/StateTreeModule/Public/Conditions/StateTreeConditionHelpers.h b/Engine/Plugins/Runtime/StateTree/Source/StateTreeModule/Public/Conditions/StateTreeConditionHelpers.h
--- a/Engine/Plugins/Runtime/StateTree/Source/StateTreeModule/Public/Conditions/StateTreeConditionHelpers.h
+++ b/Engine/Plugins/Runtime/StateTree/Source/StateTreeModule/Public/Conditions/StateTreeConditionHelpers.h
@@ -2,30 +2,43 @@
 
 #pragma once
 
+#if UE_ENABLE_INCLUDE_ORDER_DEPRECATED_IN_5_8
 #include "AITypes.h"
+#endif // UE_ENABLE_INCLUDE_ORDER_DEPRECATED_IN_5_8
+
+#include "StateTreeTypes.h"
+
+#define UE_API STATETREEMODULE_API
+
+enum class EGenericAICheck : uint8;
 
 namespace UE::StateTree::Conditions
 {
 	template<typename T>
-	bool CompareNumbers(const T Left, const T Right, const EGenericAICheck Operator)
+	bool CompareNumbers(const T Left, const T Right, const EComparisonOperator Operator)
 	{
 		switch (Operator)
 		{
-		case EGenericAICheck::Equal:
+		case EComparisonOperator::Equal:
 			return Left == Right;
-		case EGenericAICheck::NotEqual:
+		case EComparisonOperator::NotEqual:
 			return Left != Right;
-		case EGenericAICheck::Less:
+		case EComparisonOperator::Less:
 			return Left < Right;
-		case EGenericAICheck::LessOrEqual:
+		case EComparisonOperator::LessOrEqual:
 			return Left <= Right;
-		case EGenericAICheck::Greater:
+		case EComparisonOperator::Greater:
 			return Left > Right;
-		case EGenericAICheck::GreaterOrEqual:
+		case EComparisonOperator::GreaterOrEqual:
 			return Left >= Right;
 		default:
 			ensureMsgf(false, TEXT("Unhandled operator %d"), Operator);
 			return false;
 		}
 	}
+
+	UE_DEPRECATED(5.8, "Use EComparisonOperator instead.")
+	extern UE_API EComparisonOperator GenericAICheckToComparisonOperator(EGenericAICheck Operator);
 } // UE::StateTree::Conditions
+
+#undef UE_API
diff --git a/Engine/Plugins/Runtime/StateTree/Source/StateTreeModule/Public/StateTreeNodeDescriptionHelpers.h b/Engine/Plugins/Runtime/StateTree/Source/StateTreeModule/Public/StateTreeNodeDescriptionHelpers.h
--- a/Engine/Plugins/Runtime/StateTree/Source/StateTreeModule/Public/StateTreeNodeDescriptionHelpers.h
+++ b/Engine/Plugins/Runtime/StateTree/Source/StateTreeModule/Public/StateTreeNodeDescriptionHelpers.h
@@ -2,11 +2,15 @@
 
 #pragma once
 
+#if UE_ENABLE_INCLUDE_ORDER_DEPRECATED_IN_5_8
 #include "AITypes.h"
+#endif // UE_ENABLE_INCLUDE_ORDER_DEPRECATED_IN_5_8
+
 #include "Internationalization/Text.h"
 #include "StateTreeTypes.h"
 #include "StateTreePropertyBindings.h"
 
+enum class EGenericAICheck : uint8;
 enum class EStateTreeNodeFormatting : uint8;
 struct FGameplayTagContainer;
 struct FGameplayTagQuery;
@@ -15,7 +19,11 @@ namespace UE::StateTree::DescHelpers
 {
 #if WITH_EDITOR || WITH_STATETREE_TRACE
 
-/** @return description for a EGenericAICheck. */
+/** @return description for a EComparisonOperator. */
+extern STATETREEMODULE_API FText GetOperatorText(const EComparisonOperator Operator, EStateTreeNodeFormatting Formatting);
+
+UE_DEPRECATED(5.8, "Use GetOperatorText with EComparisonOperator instead.")
+/** @return description for a EComparisonOperator. */
 extern STATETREEMODULE_API FText GetOperatorText(const EGenericAICheck Operator, EStateTreeNodeFormatting Formatting);
 
 /** @return description for condition inversion (returns "Not" plus a space). */
diff --git a/Engine/Plugins/Runtime/StateTree/Source/StateTreeModule/Public/StateTreeTypes.h b/Engine/Plugins/Runtime/StateTree/Source/StateTreeModule/Public/StateTreeTypes.h
--- a/Engine/Plugins/Runtime/StateTree/Source/StateTreeModule/Public/StateTreeTypes.h
+++ b/Engine/Plugins/Runtime/StateTree/Source/StateTreeModule/Public/StateTreeTypes.h
@@ -56,6 +56,17 @@ namespace UE::StateTree
 		extern const STATETREEMODULE_API FColor DarkBronze;
 	} // Colors
 
+	UENUM()
+	enum class EComparisonOperator : uint8
+	{
+		Less,
+		LessOrEqual,
+		Equal,
+		NotEqual,
+		GreaterOrEqual,
+		Greater,
+	};
+
 }; // UE::StateTree
 
 enum class EStateTreeRunStatus : uint8;
diff --git a/Engine/Plugins/Runtime/StateTree/Source/StateTreeModule/StateTreeModule.Build.cs b/Engine/Plugins/Runtime/StateTree/Source/StateTreeModule/StateTreeModule.Build.cs
--- a/Engine/Plugins/Runtime/StateTree/Source/StateTreeModule/StateTreeModule.Build.cs
+++ b/Engine/Plugins/Runtime/StateTree/Source/StateTreeModule/StateTreeModule.Build.cs
@@ -72,7 +72,7 @@ namespace UnrealBuildTool.Rules
 					"CoreUObject",
 					"DeveloperSettings",
 					"Engine",
-					"AIModule",
+					"AIModule", //@TODO Move to Private dependency when AITypes.h deprection is removed.
 					"GameplayTags",
 					"PropertyBindingUtils"
 				}
diff --git a/Engine/Plugins/Runtime/StateTree/Source/StateTreeTestSuite/Private/StateTreeBindingTest.cpp b/Engine/Plugins/Runtime/StateTree/Source/StateTreeTestSuite/Private/StateTreeBindingTest.cpp
--- a/Engine/Plugins/Runtime/StateTree/Source/StateTreeTestSuite/Private/StateTreeBindingTest.cpp
+++ b/Engine/Plugins/Runtime/StateTree/Source/StateTreeTestSuite/Private/StateTreeBindingTest.cpp
@@ -612,7 +612,7 @@ struct FStateTreeTest_PropertyFunctions : FStateTreeTestBase
 
 		// Condition with property function binding.
 		{
-			TStateTreeEditorNode<FStateTreeCompareIntCondition>& EnterCond = Root.AddEnterCondition<FStateTreeCompareIntCondition>(EGenericAICheck::Equal);
+			TStateTreeEditorNode<FStateTreeCompareIntCondition>& EnterCond = Root.AddEnterCondition<FStateTreeCompareIntCondition>(UE::StateTree::EComparisonOperator::Equal);
 			EnterCond.GetInstanceData().Right = 1;
 			EditorData.AddPropertyBinding(CastChecked<UScriptStruct>(FTestPropertyFunction::StaticStruct()), {PathSegmentToFuncResult}, FPropertyBindingPath(EnterCond.ID, TEXT("Left")));
 		}
diff --git a/Engine/Plugins/Runtime/StateTree/Source/StateTreeTestSuite/Private/StateTreeTest.cpp b/Engine/Plugins/Runtime/StateTree/Source/StateTreeTestSuite/Private/StateTreeTest.cpp
--- a/Engine/Plugins/Runtime/StateTree/Source/StateTreeTestSuite/Private/StateTreeTest.cpp
+++ b/Engine/Plugins/Runtime/StateTree/Source/StateTreeTestSuite/Private/StateTreeTest.cpp
@@ -45,7 +45,7 @@ struct FStateTreeTest_MakeAndBakeStateTree : FStateTreeTestBase
 		auto& TaskB1 = StateA.AddTask<FTestTask_B>();
 		EditorData.AddPropertyBinding(EvalA, TEXT("IntA"), TaskB1, TEXT("IntB"));
 
-		auto& IntCond = StateA.AddEnterCondition<FStateTreeCompareIntCondition>(EGenericAICheck::Less);
+		auto& IntCond = StateA.AddEnterCondition<FStateTreeCompareIntCondition>(UE::StateTree::EComparisonOperator::Less);
 		IntCond.GetInstanceData().Right = 2;
 
 		EditorData.AddPropertyBinding(EvalA, TEXT("IntA"), IntCond, TEXT("Left"));
@@ -57,7 +57,7 @@ struct FStateTreeTest_MakeAndBakeStateTree : FStateTreeTestBase
 		EditorData.AddPropertyBinding(EvalA, TEXT("bBoolA"), TaskB2, TEXT("bBoolB"));
 
 		FStateTreeTransition& Trans = StateB.AddTransition({}, EStateTreeTransitionType::GotoState, &Root);
-		auto& TransFloatCond = Trans.AddCondition<FStateTreeCompareFloatCondition>(EGenericAICheck::Less);
+		auto& TransFloatCond = Trans.AddCondition<FStateTreeCompareFloatCondition>(UE::StateTree::EComparisonOperator::Less);
 		TransFloatCond.GetInstanceData().Right = 13.0f;
 		EditorData.AddPropertyBinding(EvalA, TEXT("FloatA"), TransFloatCond, TEXT("Left"));
 
@@ -576,7 +576,7 @@ struct FStateTreeTest_SubTree_Condition : FStateTreeTestBase
 		Task2.GetNode().TicksToCompletion = 1;
 		
 		// Allow to enter State1 if Task1 instance data TicksToCompletion > 0.
-		TStateTreeEditorNode<FStateTreeCompareIntCondition>& IntCond1 = State1.AddEnterCondition<FStateTreeCompareIntCondition>(EGenericAICheck::Greater);
+		TStateTreeEditorNode<FStateTreeCompareIntCondition>& IntCond1 = State1.AddEnterCondition<FStateTreeCompareIntCondition>(UE::StateTree::EComparisonOperator::Greater);
 		EditorData.AddPropertyBinding(SubTask, TEXT("CurrentTick"), IntCond1, TEXT("Left"));
 		IntCond1.GetInstanceData().Right = 0;
 
diff --git a/Engine/Plugins/Runtime/StateTree/Source/StateTreeTestSuite/Private/StateTreeTransitionTest.cpp b/Engine/Plugins/Runtime/StateTree/Source/StateTreeTestSuite/Private/StateTreeTransitionTest.cpp
--- a/Engine/Plugins/Runtime/StateTree/Source/StateTreeTestSuite/Private/StateTreeTransitionTest.cpp
+++ b/Engine/Plugins/Runtime/StateTree/Source/StateTreeTestSuite/Private/StateTreeTransitionTest.cpp
@@ -476,7 +476,7 @@ struct FStateTreeTest_TransitionReactivateTargetState : FStateTreeTestBase
 
 		// Reactivate transtion back to root, by default will not fire due to condition. But should reactivate / re-trigger enter state when activated.
 		FStateTreeTransition& ReactivateRootTransition = State0.AddTransition(EStateTreeTransitionTrigger::OnTick, EStateTreeTransitionType::GotoState, &Root);
-		TStateTreeEditorNode<FStateTreeCompareIntCondition>& TransIntCond = ReactivateRootTransition.AddCondition<FStateTreeCompareIntCondition>(EGenericAICheck::Equal);
+		TStateTreeEditorNode<FStateTreeCompareIntCondition>& TransIntCond = ReactivateRootTransition.AddCondition<FStateTreeCompareIntCondition>(UE::StateTree::EComparisonOperator::Equal);
 		TransIntCond.GetInstanceData().Right = 0;
 		EditorData.AddPropertyBinding(
 			FPropertyBindingPath(EditorData.GetRootParametersGuid(), TEXT("Int")),
@@ -641,7 +641,7 @@ struct FStateTreeTest_PassingTransitionEventToStateSelection : FStateTreeTestBas
 		StateA.bHasRequiredEventToEnter  = true;
 		StateA.RequiredEventToEnter.PayloadStruct = FStateTreeTest_PropertyStructA::StaticStruct();
 		auto& TaskA = StateA.AddTask<FTestTask_Stand>(FName(TEXT("TaskA")));
-		TStateTreeEditorNode<FStateTreeCompareIntCondition>& AIntCond = StateA.AddEnterCondition<FStateTreeCompareIntCondition>(EGenericAICheck::Equal);
+		TStateTreeEditorNode<FStateTreeCompareIntCondition>& AIntCond = StateA.AddEnterCondition<FStateTreeCompareIntCondition>(UE::StateTree::EComparisonOperator::Equal);
 		AIntCond.GetInstanceData().Right = 0;
 		EditorData.AddPropertyBinding(
 			FPropertyBindingPath(StateA.GetEventID(), PathToPayloadMember.GetSegments()),
@@ -658,7 +658,7 @@ struct FStateTreeTest_PassingTransitionEventToStateSelection : FStateTreeTestBas
 			FPropertyBindingPath(StateB.GetEventID(), PathToPayloadMember.GetSegments()),
 			FPropertyBindingPath(TaskB.ID, TEXT("Value")));
 		
-		TStateTreeEditorNode<FStateTreeCompareIntCondition>& BIntCond = StateB.AddEnterCondition<FStateTreeCompareIntCondition>(EGenericAICheck::Equal);
+		TStateTreeEditorNode<FStateTreeCompareIntCondition>& BIntCond = StateB.AddEnterCondition<FStateTreeCompareIntCondition>(UE::StateTree::EComparisonOperator::Equal);
 		BIntCond.GetInstanceData().Right = 1;
 		EditorData.AddPropertyBinding(
 			FPropertyBindingPath(StateB.GetEventID(), PathToPayloadMember.GetSegments()),
@@ -670,7 +670,7 @@ struct FStateTreeTest_PassingTransitionEventToStateSelection : FStateTreeTestBas
 		// Transition from Initial -> StateA
 		FStateTreeTransition& TransA = StateInitial.AddTransition(EStateTreeTransitionTrigger::OnEvent, FGameplayTag(), EStateTreeTransitionType::GotoState, &StateA);
 		TransA.RequiredEvent.PayloadStruct = FStateTreeTest_PropertyStructA::StaticStruct();
-		TStateTreeEditorNode<FStateTreeCompareIntCondition>& TransAIntCond = TransA.AddCondition<FStateTreeCompareIntCondition>(EGenericAICheck::Equal);
+		TStateTreeEditorNode<FStateTreeCompareIntCondition>& TransAIntCond = TransA.AddCondition<FStateTreeCompareIntCondition>(UE::StateTree::EComparisonOperator::Equal);
 		TransAIntCond.GetInstanceData().Right = 1;
 		EditorData.AddPropertyBinding(
 			FPropertyBindingPath(TransA.GetEventID(), PathToPayloadMember.GetSegments()),
@@ -678,7 +678,7 @@ struct FStateTreeTest_PassingTransitionEventToStateSelection : FStateTreeTestBas
 		// Transition from Initial -> StateB
 		FStateTreeTransition& TransB = StateInitial.AddTransition(EStateTreeTransitionTrigger::OnEvent, FGameplayTag(), EStateTreeTransitionType::GotoState, &StateB);
 		TransB.RequiredEvent.PayloadStruct = FStateTreeTest_PropertyStructA::StaticStruct();
-		TStateTreeEditorNode<FStateTreeCompareIntCondition>& TransBIntCond = TransB.AddCondition<FStateTreeCompareIntCondition>(EGenericAICheck::Equal);
+		TStateTreeEditorNode<FStateTreeCompareIntCondition>& TransBIntCond = TransB.AddCondition<FStateTreeCompareIntCondition>(UE::StateTree::EComparisonOperator::Equal);
 		TransBIntCond.GetInstanceData().Right = 1;
 		EditorData.AddPropertyBinding(
 			FPropertyBindingPath(TransB.GetEventID(), PathToPayloadMember.GetSegments()),
@@ -748,7 +748,7 @@ struct FStateTreeTest_FollowTransitions : FStateTreeTestBase
 			{
 				// This transition should be skipped due to the condition
 				FStateTreeTransition& TransA = StateTrans.AddTransition(EStateTreeTransitionTrigger::OnTick, EStateTreeTransitionType::GotoState, &StateA);
-				TStateTreeEditorNode<FStateTreeCompareIntCondition>& TransIntCond = TransA.AddCondition<FStateTreeCompareIntCondition>(EGenericAICheck::Equal);
+				TStateTreeEditorNode<FStateTreeCompareIntCondition>& TransIntCond = TransA.AddCondition<FStateTreeCompareIntCondition>(UE::StateTree::EComparisonOperator::Equal);
 				TransIntCond.GetInstanceData().Right = 0;
 				EditorData.AddPropertyBinding(
 					FPropertyBindingPath(EditorData.GetRootParametersGuid(), TEXT("Int")),
@@ -759,7 +759,7 @@ struct FStateTreeTest_FollowTransitions : FStateTreeTestBase
 				// This transition leads to selection, but will be overridden.
 				FStateTreeTransition& TransB = StateTrans.AddTransition(EStateTreeTransitionTrigger::OnTick, EStateTreeTransitionType::GotoState, &StateB);
 				TransB.Priority = EStateTreeTransitionPriority::Normal;
-				TStateTreeEditorNode<FStateTreeCompareIntCondition>& TransIntCond = TransB.AddCondition<FStateTreeCompareIntCondition>(EGenericAICheck::Equal);
+				TStateTreeEditorNode<FStateTreeCompareIntCondition>& TransIntCond = TransB.AddCondition<FStateTreeCompareIntCondition>(UE::StateTree::EComparisonOperator::Equal);
 				TransIntCond.GetInstanceData().Right = 1;
 				EditorData.AddPropertyBinding(
 					FPropertyBindingPath(EditorData.GetRootParametersGuid(), TEXT("Int")),
@@ -770,7 +770,7 @@ struct FStateTreeTest_FollowTransitions : FStateTreeTestBase
 				// This transition is selected, should override previous one due to priority.
 				FStateTreeTransition& TransC = StateTrans.AddTransition(EStateTreeTransitionTrigger::OnTick, EStateTreeTransitionType::GotoState, &StateC);
 				TransC.Priority = EStateTreeTransitionPriority::High;
-				TStateTreeEditorNode<FStateTreeCompareIntCondition>& TransIntCond = TransC.AddCondition<FStateTreeCompareIntCondition>(EGenericAICheck::Equal);
+				TStateTreeEditorNode<FStateTreeCompareIntCondition>& TransIntCond = TransC.AddCondition<FStateTreeCompareIntCondition>(UE::StateTree::EComparisonOperator::Equal);
 				TransIntCond.GetInstanceData().Right = 1;
 				EditorData.AddPropertyBinding(
 					FPropertyBindingPath(EditorData.GetRootParametersGuid(), TEXT("Int")),
@@ -833,7 +833,7 @@ struct FStateTreeTest_InfiniteLoop : FStateTreeTestBase
 			{
 				// A -> B
 				FStateTreeTransition& Trans = StateA.AddTransition(EStateTreeTransitionTrigger::OnTick, EStateTreeTransitionType::GotoState, &StateB);
-				TStateTreeEditorNode<FStateTreeCompareIntCondition>& TransIntCond = Trans.AddCondition<FStateTreeCompareIntCondition>(EGenericAICheck::Equal);
+				TStateTreeEditorNode<FStateTreeCompareIntCondition>& TransIntCond = Trans.AddCondition<FStateTreeCompareIntCondition>(UE::StateTree::EComparisonOperator::Equal);
 				TransIntCond.GetInstanceData().Right = 1;
 				EditorData.AddPropertyBinding(
 					FPropertyBindingPath(EditorData.GetRootParametersGuid(), TEXT("Int")),
@@ -847,7 +847,7 @@ struct FStateTreeTest_InfiniteLoop : FStateTreeTestBase
 			{
 				// B -> A
 				FStateTreeTransition& Trans = StateB.AddTransition(EStateTreeTransitionTrigger::OnTick, EStateTreeTransitionType::GotoState, &StateA);
-				TStateTreeEditorNode<FStateTreeCompareIntCondition>& TransIntCond = Trans.AddCondition<FStateTreeCompareIntCondition>(EGenericAICheck::Equal);
+				TStateTreeEditorNode<FStateTreeCompareIntCondition>& TransIntCond = Trans.AddCondition<FStateTreeCompareIntCondition>(UE::StateTree::EComparisonOperator::Equal);
 				TransIntCond.GetInstanceData().Right = 1;
 				EditorData.AddPropertyBinding(
 					FPropertyBindingPath(EditorData.GetRootParametersGuid(), TEXT("Int")),
