[StateTree] Allow UserDefinedStruct to be used inside a compiled StateTree parameter. The struct is a PropertyBag when it's an internal structure, not visible outside of the StateTree asset. It is a UserDefinedStruct when the parameters are from a linked asset.
This is part of a bigger change to support UserDefinedStruct inside StateTree.

#Synced-CL 49772122

diff --git a/Engine/Plugins/Runtime/StateTree/Source/StateTreeEditorModule/Private/StateTreeCompiler.cpp b/Engine/Plugins/Runtime/StateTree/Source/StateTreeEditorModule/Private/StateTreeCompiler.cpp
--- a/Engine/Plugins/Runtime/StateTree/Source/StateTreeEditorModule/Private/StateTreeCompiler.cpp
+++ b/Engine/Plugins/Runtime/StateTree/Source/StateTreeEditorModule/Private/StateTreeCompiler.cpp
@@ -451,7 +451,7 @@ bool FStateTreeCompiler::PreCompile(TNotNull<UStateTree*> InStateTree)
 	{
 		return false;
 	}
-	
+
 	// Cleanup existing state
 	StateTree->ResetCompiled();
 
@@ -1076,7 +1076,7 @@ bool FStateTreeCompiler::CreateStateTasksAndParameters()
 			State->Parameters.ID
 		};
 
-		if (!UE::StateTree::Compiler::ValidateNoLevelActorReferences(Log, LinkedParamsDesc, FStateTreeDataView(), FStateTreeDataView(CompactStateTreeParameters.Parameters.GetMutableValue())))
+		if (!UE::StateTree::Compiler::ValidateNoLevelActorReferences(Log, LinkedParamsDesc, FStateTreeDataView(), FStateTreeDataView(CompactStateTreeParameters.GetMutableValue())))
 		{
 			bSucceeded = false;
 			continue;
@@ -1141,7 +1141,7 @@ bool FStateTreeCompiler::CreateStateTasksAndParameters()
 
 			// Only nodes support output bindings
 			constexpr FStateTreeIndex16* OutputBindingsBatch = nullptr;
-			if (!CreateBindingsForStruct(LinkedParamsDesc, FStateTreeDataView(CompactStateTreeParameters.Parameters.GetMutableValue()), PropertyFunctionsBegin, PropertyFunctionsEnd, CompactState.ParameterBindingsBatch, OutputBindingsBatch))
+			if (!CreateBindingsForStruct(LinkedParamsDesc, FStateTreeDataView(CompactStateTreeParameters.GetMutableValue()), PropertyFunctionsBegin, PropertyFunctionsEnd, CompactState.ParameterBindingsBatch, OutputBindingsBatch))
 			{
 				bSucceeded = false;
 				continue;
diff --git a/Engine/Plugins/Runtime/StateTree/Source/StateTreeModule/Private/StateTree.cpp b/Engine/Plugins/Runtime/StateTree/Source/StateTreeModule/Private/StateTree.cpp
--- a/Engine/Plugins/Runtime/StateTree/Source/StateTreeModule/Private/StateTree.cpp
+++ b/Engine/Plugins/Runtime/StateTree/Source/StateTreeModule/Private/StateTree.cpp
@@ -267,7 +267,7 @@ void UStateTree::ResetCompiled()
 	IDToStateMappings.Reset();
 	IDToNodeMappings.Reset();
 	IDToTransitionMappings.Reset();
-	
+
 	EvaluatorsBegin = 0;
 	EvaluatorsNum = 0;
 
@@ -1001,12 +1001,12 @@ bool UStateTree::PatchBindings()
 				UE_LOG(LogStateTree, Error, TEXT("%s: Data for state '%s' is malformed. Please recompile the StateTree asset."), *GetPathName(), *State.Name.ToString());
 				return false;
 			}
-			Desc->Struct = Params.Parameters.GetPropertyBagStruct();
+			Desc->Struct = Params.GetValue().GetScriptStruct();
 
 			if (State.ParameterBindingsBatch.IsValid())
 			{
 				FPropertyBindingCopyInfoBatch& Batch = CopyBatches[State.ParameterBindingsBatch.Get()];
-				Batch.TargetStruct.GetMutable().Struct = Params.Parameters.GetPropertyBagStruct();
+				Batch.TargetStruct.GetMutable().Struct = Params.GetValue().GetScriptStruct();
 			}
 		}
 	}
@@ -1029,7 +1029,7 @@ bool UStateTree::PatchBindings()
 			const FCompactStateTreeParameters& Params = DefaultInstanceData.GetMutableStruct(State.ParameterTemplateIndex.Get()).Get<FCompactStateTreeParameters>();
 			const FCompactStateTreeParameters& LinkedStateParams = DefaultInstanceData.GetMutableStruct(LinkedState.ParameterTemplateIndex.Get()).Get<FCompactStateTreeParameters>();
 
-			if (LinkedStateParams.Parameters.GetPropertyBagStruct() != Params.Parameters.GetPropertyBagStruct())
+			if (LinkedStateParams.GetValue().GetScriptStruct() != Params.GetValue().GetScriptStruct())
 			{
 				UE_LOG(LogStateTree, Error, TEXT("%s: The parameters on state '%s' does not match the linked state parameters in state '%s'. Please recompile the StateTree asset."), *GetPathName(), *State.Name.ToString(), *LinkedState.Name.ToString());
 				return false;
@@ -1041,7 +1041,7 @@ bool UStateTree::PatchBindings()
 			const FInstancedPropertyBag& TargetTreeParameters = State.LinkedAsset->Parameters;
 			const FCompactStateTreeParameters& Params = DefaultInstanceData.GetMutableStruct(State.ParameterTemplateIndex.Get()).Get<FCompactStateTreeParameters>();
 
-			if (TargetTreeParameters.GetPropertyBagStruct() != Params.Parameters.GetPropertyBagStruct())
+			if (TargetTreeParameters.GetPropertyBagStruct() != Params.GetValue().GetScriptStruct())
 			{
 				UE_LOG(LogStateTree, Error, TEXT("%s: The parameters on state '%s' does not match the linked asset parameters '%s'. Please recompile the StateTree asset."),
 					*GetPathName(), *State.Name.ToString(), *State.LinkedAsset->GetPathName());
@@ -1068,10 +1068,10 @@ bool UStateTree::PatchBindings()
 		if (State.ParameterDataHandle.IsValid())
 		{
 			FCompactStateTreeParameters& Params = DefaultInstanceData.GetMutableStruct(State.ParameterTemplateIndex.Get()).Get<FCompactStateTreeParameters>();
-			DataViews.Add(State.ParameterDataHandle, Params.Parameters.GetMutableValue());
+			DataViews.Add(State.ParameterDataHandle, Params.GetMutableValue());
 			if (State.ParameterBindingsBatch.IsValid())
 			{
-				BindingBatchDataView.Add(State.ParameterBindingsBatch, Params.Parameters.GetMutableValue());
+				BindingBatchDataView.Add(State.ParameterBindingsBatch, Params.GetMutableValue());
 			}
 		}
 	}
diff --git a/Engine/Plugins/Runtime/StateTree/Source/StateTreeModule/Private/StateTreeExecutionContext.cpp b/Engine/Plugins/Runtime/StateTree/Source/StateTreeModule/Private/StateTreeExecutionContext.cpp
--- a/Engine/Plugins/Runtime/StateTree/Source/StateTreeModule/Private/StateTreeExecutionContext.cpp
+++ b/Engine/Plugins/Runtime/StateTree/Source/StateTreeModule/Private/StateTreeExecutionContext.cpp
@@ -2455,7 +2455,7 @@ void FStateTreeExecutionContext::UpdateInstanceData(const TSharedPtr<FSelectStat
 						{
 							// Point to the parameter block set by linked state.
 							const FCompactStateTreeParameters* Params = ParamsInstanceData.GetPtr<const FCompactStateTreeParameters>();
-							const UStruct* StateParameterDataStruct = Params ? Params->Parameters.GetPropertyBagStruct() : nullptr;
+							const UStruct* StateParameterDataStruct = Params ? Params->GetValue().GetScriptStruct() : nullptr;
 							check(NextStateParameterDataStruct == StateParameterDataStruct);
 
 							check(ParentFrame);
@@ -2519,7 +2519,7 @@ void FStateTreeExecutionContext::UpdateInstanceData(const TSharedPtr<FSelectStat
 								check(CurrentState->ParameterDataHandle.GetSource() == EStateTreeDataSourceType::StateParameterData);
 								checkf(!NextStateParameterDataHandle.IsValid(), TEXT("NextStateParameterDataIndex not should be set yet when we encounter a linked state."));
 								NextStateParameterDataHandle = CurrentState->ParameterDataHandle;
-								NextStateParameterDataStruct = Params ? Params->Parameters.GetPropertyBagStruct() : nullptr;
+								NextStateParameterDataStruct = Params ? Params->GetValue().GetScriptStruct() : nullptr;
 							}
 						}
 					}
@@ -7013,8 +7013,8 @@ bool FStateTreeExecutionContext::SelectStateInternal(
 				FStateTreeDataView TempStateParametersView = AddTemporaryInstance(*NextFrame, FStateTreeIndex16::Invalid, NextState.ParameterDataHandle, FConstStructView(TBaseStructure<FCompactStateTreeParameters>::Get()));
 				check(TempStateParametersView.IsValid());
 				FCompactStateTreeParameters& StateParams = TempStateParametersView.GetMutable<FCompactStateTreeParameters>();
-				StateParams.Parameters = *NextLinkedStateParameterOverride;
-				NextStateParametersView = FStateTreeDataView(StateParams.Parameters.GetMutableValue());
+				StateParams = FCompactStateTreeParameters(*NextLinkedStateParameterOverride);
+				NextStateParametersView = FStateTreeDataView(StateParams.GetMutableValue());
 			}
 			else
 			{
@@ -7023,7 +7023,7 @@ bool FStateTreeExecutionContext::SelectStateInternal(
 				FStateTreeDataView TempStateParametersView = AddTemporaryInstance(*NextFrame, FStateTreeIndex16::Invalid, NextState.ParameterDataHandle, DefaultStateParamsInstanceData);
 				check(TempStateParametersView.IsValid());
 				FCompactStateTreeParameters& StateParams = TempStateParametersView.GetMutable<FCompactStateTreeParameters>();
-				NextStateParametersView = FStateTreeDataView(StateParams.Parameters.GetMutableValue());
+				NextStateParametersView = FStateTreeDataView(StateParams.GetMutableValue());
 			}
 		}
 
diff --git a/Engine/Plugins/Runtime/StateTree/Source/StateTreeModule/Private/StateTreeInstanceData.cpp b/Engine/Plugins/Runtime/StateTree/Source/StateTreeModule/Private/StateTreeInstanceData.cpp
--- a/Engine/Plugins/Runtime/StateTree/Source/StateTreeModule/Private/StateTreeInstanceData.cpp
+++ b/Engine/Plugins/Runtime/StateTree/Source/StateTreeModule/Private/StateTreeInstanceData.cpp
@@ -271,7 +271,7 @@ namespace UE::StateTree::InstanceData::Private
 
 				if (FCompactStateTreeParameters* Params = InstanceStorage.GetMutableTemporaryStruct(CurrentFrame.GlobalParameterDataFrameID, CurrentFrame.GlobalParameterDataHandle).GetPtr<FCompactStateTreeParameters>())
 				{
-					return Params->Parameters.GetMutableValue();
+					return Params->GetMutableValue();
 				}
 			}
 
@@ -292,7 +292,7 @@ namespace UE::StateTree::InstanceData::Private
 
 			if (FCompactStateTreeParameters* Params = InstanceStorage.GetMutableTemporaryStruct(CurrentFrame.StateParameterDataFrameID, CurrentFrame.StateParameterDataHandle).GetPtr<FCompactStateTreeParameters>())
 			{
-				return Params->Parameters.GetMutableValue();
+				return Params->GetMutableValue();
 			}
 
 			break;
@@ -302,7 +302,7 @@ namespace UE::StateTree::InstanceData::Private
 		{
 			if (FCompactStateTreeParameters* Params = InstanceStorage.GetMutableTemporaryStruct(CurrentFrame.FrameID, Handle).GetPtr<FCompactStateTreeParameters>())
 			{
-				return Params->Parameters.GetMutableValue();
+				return Params->GetMutableValue();
 			}
 
 			break;
@@ -376,7 +376,7 @@ namespace UE::StateTree::InstanceData
 				check(CurrentFrame.GlobalParameterIndexBase.IsValid());
 
 				FCompactStateTreeParameters& StateParams = InstanceStorage.GetMutableStruct(CurrentFrame.GlobalParameterIndexBase.Get() + CurrentFrame.GlobalParameterDataHandle.GetIndex()).Get<FCompactStateTreeParameters>();
-				return StateParams.Parameters.GetMutableValue();
+				return StateParams.GetMutableValue();
 			}
 
 			// Root Global Parameters
@@ -397,13 +397,13 @@ namespace UE::StateTree::InstanceData
 					&& CurrentFrame.FrameID != CurrentFrame.StateParameterDataFrameID));
 
 			FCompactStateTreeParameters& StateParams = InstanceStorage.GetMutableStruct(CurrentFrame.StateParameterIndexBase.Get() + CurrentFrame.StateParameterDataHandle.GetIndex()).Get<FCompactStateTreeParameters>();
-			return StateParams.Parameters.GetMutableValue();
+			return StateParams.GetMutableValue();
 		}
 
 		case EStateTreeDataSourceType::StateParameterData:
 		{
 			FCompactStateTreeParameters& StateParams = InstanceStorage.GetMutableStruct(CurrentFrame.ActiveInstanceIndexBase.Get() + Handle.GetIndex()).Get<FCompactStateTreeParameters>();
-			return StateParams.Parameters.GetMutableValue();
+			return StateParams.GetMutableValue();
 		}
 
 		case EStateTreeDataSourceType::StateEvent:
diff --git a/Engine/Plugins/Runtime/StateTree/Source/StateTreeModule/Public/StateTreeTypes.h b/Engine/Plugins/Runtime/StateTree/Source/StateTreeModule/Public/StateTreeTypes.h
--- a/Engine/Plugins/Runtime/StateTree/Source/StateTreeModule/Public/StateTreeTypes.h
+++ b/Engine/Plugins/Runtime/StateTree/Source/StateTreeModule/Public/StateTreeTypes.h
@@ -1032,15 +1032,48 @@ struct FCompactStateTreeParameters
 {
 	GENERATED_BODY()
 
+	PRAGMA_DISABLE_DEPRECATION_WARNINGS
 	FCompactStateTreeParameters() = default;
+	FCompactStateTreeParameters(const FCompactStateTreeParameters&) = default;
+	FCompactStateTreeParameters(FCompactStateTreeParameters&&) = default;
+	FCompactStateTreeParameters& operator=(const FCompactStateTreeParameters&) = default;
+	FCompactStateTreeParameters& operator=(FCompactStateTreeParameters&&) = default;
+	PRAGMA_ENABLE_DEPRECATION_WARNINGS
+
+	explicit FCompactStateTreeParameters(FInstancedPropertyBag InParameters)
+		: Parameters(MoveTemp(InParameters))
+	{
+	}
 
-	FCompactStateTreeParameters(const FInstancedPropertyBag& InParameters)
-		: Parameters(InParameters)
+	explicit FCompactStateTreeParameters(FInstancedStruct InParameters)
+		: InstancedParameters(MoveTemp(InParameters))
 	{
 	}
-	
+
+public:
+	FConstStructView GetValue() const
+	{
+		PRAGMA_DISABLE_DEPRECATION_WARNINGS
+		return Parameters.IsValid() ? Parameters.GetValue() : FConstStructView(InstancedParameters);
+		PRAGMA_ENABLE_DEPRECATION_WARNINGS
+	}
+
+	FStructView GetMutableValue()
+	{
+		PRAGMA_DISABLE_DEPRECATION_WARNINGS
+		return Parameters.IsValid() ? Parameters.GetMutableValue() : FStructView(InstancedParameters);
+		PRAGMA_ENABLE_DEPRECATION_WARNINGS
+	}
+
+private:
+	UE_DEPRECATED_FORGAME(5.8, "Direct access to Parameter is not valid. Use GetValue or GetMutableValue.")
+	/** To hold a PropertyBag for internal state parameters. */
 	UPROPERTY()
 	FInstancedPropertyBag Parameters;
+
+	/** To hold a UScriptStruct for a linked state tree asset. */
+	UPROPERTY()
+	FInstancedStruct InstancedParameters;
 };
 
 UENUM()
diff --git a/Engine/Plugins/Runtime/StateTree/Source/StateTreeTestSuite/Private/StateTreeBindingTest.cpp b/Engine/Plugins/Runtime/StateTree/Source/StateTreeTestSuite/Private/StateTreeBindingTest.cpp
--- a/Engine/Plugins/Runtime/StateTree/Source/StateTreeTestSuite/Private/StateTreeBindingTest.cpp
+++ b/Engine/Plugins/Runtime/StateTree/Source/StateTreeTestSuite/Private/StateTreeBindingTest.cpp
@@ -1542,10 +1542,15 @@ struct FStateTreeTest_BindingStructRef : FStateTreeTestBase
 						FConstStructView RootParameters = Context.GetInstanceData()->GetStorage().GetStruct(0);
 						if (const FCompactStateTreeParameters* Parameters = RootParameters.GetPtr<const FCompactStateTreeParameters>())
 						{
-							TValueOrError<FStateTreeTest_PropertyStruct*, EPropertyBagResult> PropertyStruct = Parameters->Parameters.GetValueStruct<FStateTreeTest_PropertyStruct>("RootParam");
-							if (PropertyStruct.HasValue())
+							if (Parameters->GetValue().GetScriptStruct() != nullptr)
 							{
-								TaskA_State_PropertyStructPtr = PropertyStruct.GetValue();
+								if (const FStructProperty* StructProperty = CastField<const FStructProperty>(Parameters->GetValue().GetScriptStruct()->FindPropertyByName("RootParam")))
+								{
+									if (StructProperty->Struct == TBaseStructure<FStateTreeTest_PropertyStruct>())
+									{
+										TaskA_State_PropertyStructPtr = StructProperty->ContainerPtrToValuePtr<const FStateTreeTest_PropertyStruct>(Parameters->GetValue().GetMemory());
+									}
+								}
 							}
 						}
 
@@ -1566,10 +1571,23 @@ struct FStateTreeTest_BindingStructRef : FStateTreeTestBase
 						FConstStructView State1Parameters = Context.GetInstanceData()->GetStorage().GetStruct(1);
 						if (const FCompactStateTreeParameters* Parameters = State1Parameters.GetPtr<const FCompactStateTreeParameters>())
 						{
-							TValueOrError<FStateTreeStructRef*, EPropertyBagResult> PropertyStruct = Parameters->Parameters.GetValueStruct<FStateTreeStructRef>("StateParam");
-							if (PropertyStruct.HasValue())
+							if (Parameters->GetValue().GetScriptStruct() != nullptr)
 							{
-								TaskB_State_PropertyStructPtr = PropertyStruct.GetValue()->GetPtr<const FStateTreeTest_PropertyStruct>();
+								if (const FStructProperty* StructProperty = CastField<const FStructProperty>(Parameters->GetValue().GetScriptStruct()->FindPropertyByName("StateParam")))
+								{
+									if (StructProperty->Struct == TBaseStructure<FStateTreeStructRef>())
+									{
+										if (const FStateTreeStructRef* StructRef = StructProperty->ContainerPtrToValuePtr<const FStateTreeStructRef>(Parameters->GetValue().GetMemory()))
+										{
+											TaskB_State_PropertyStructPtr = StructRef->GetPtr<const FStateTreeTest_PropertyStruct>();
+										}
+									}
+								}
+							}
+
+							if (const FStateTreeStructRef* PropertyStruct = Parameters->GetValue().GetPtr<FStateTreeStructRef>())
+							{
+								TaskB_State_PropertyStructPtr = PropertyStruct->GetPtr<const FStateTreeTest_PropertyStruct>();
 							}
 						}
 
