[Crash Fix] Crash creating an exposed State Tree State Link variable to an array

[Added] Guards in the StateTreePropertyHelpers.h file in order to fail if GetProperty() property returns nullptr.

#Synced-CL 49621724
#Synced-CL 49621724

diff --git a/Engine/Plugins/Runtime/StateTree/Source/StateTreeEditorModule/Private/StateTreePropertyHelpers.h b/Engine/Plugins/Runtime/StateTree/Source/StateTreeEditorModule/Private/StateTreePropertyHelpers.h
--- a/Engine/Plugins/Runtime/StateTree/Source/StateTreeEditorModule/Private/StateTreePropertyHelpers.h
+++ b/Engine/Plugins/Runtime/StateTree/Source/StateTreeEditorModule/Private/StateTreePropertyHelpers.h
@@ -48,7 +48,13 @@ bool IsScriptStruct(const TSharedPtr<IPropertyHandle>& PropertyHandle)
 		return false;
 	}
 
-	FStructProperty* StructProperty = CastField<FStructProperty>(PropertyHandle->GetProperty());
+	FProperty* Property = PropertyHandle->GetProperty();
+	if (!Property)
+	{
+		return false;
+	}
+
+	FStructProperty* StructProperty = CastField<FStructProperty>(Property);
 	return StructProperty && StructProperty->Struct->IsA(TBaseStructure<T>::Get()->GetClass());
 }
 /**
@@ -69,7 +75,13 @@ FPropertyAccess::Result GetStructValue(const TSharedPtr<const IPropertyHandle>&
 		return FPropertyAccess::Fail;
 	}
 
-	FStructProperty* StructProperty = CastFieldChecked<FStructProperty>(ValueProperty->GetProperty());
+	FProperty* Property = ValueProperty->GetProperty();
+	if (!Property)
+	{
+		return FPropertyAccess::Fail;
+	}
+
+	FStructProperty* StructProperty = CastFieldChecked<FStructProperty>(Property);
 	check(StructProperty);
 	check(StructProperty->Struct == TBaseStructure<T>::Get());
 
@@ -114,19 +126,23 @@ const T* GetStructPtr(const TSharedPtr<const IPropertyHandle>& ValueProperty)
 		return nullptr;
 	}
 
-	FStructProperty* StructProperty = CastFieldChecked<FStructProperty>(ValueProperty->GetProperty());
-	check(StructProperty);
-	check(StructProperty->Struct == TBaseStructure<T>::Get());
-
-	TArray<const void*> RawData;
-	ValueProperty->AccessRawData(RawData);
-	if (RawData.Num() == 1)
+	if (FProperty* Property = ValueProperty->GetProperty())
 	{
-		return static_cast<const T*>(RawData[0]); 
+		FStructProperty* StructProperty = CastFieldChecked<FStructProperty>(Property);
+		check(StructProperty);
+		check(StructProperty->Struct == TBaseStructure<T>::Get());
+
+		TArray<const void*> RawData;
+		ValueProperty->AccessRawData(RawData);
+		if (RawData.Num() == 1)
+		{
+			return static_cast<const T*>(RawData[0]);
+		}
 	}
 
 	return nullptr;
 }
+
 /**
  * Sets a struct property to specific value, checks type before access. Expects T is struct.
  * @param ValueProperty Handle to property where value is got from.
@@ -140,7 +156,13 @@ FPropertyAccess::Result SetStructValue(const TSharedPtr<IPropertyHandle>& ValueP
 		return FPropertyAccess::Fail;
 	}
 
-	FStructProperty* StructProperty = CastFieldChecked<FStructProperty>(ValueProperty->GetProperty());
+	FProperty* Property = ValueProperty->GetProperty();
+	if (!Property)
+	{
+		return FPropertyAccess::Fail;
+	}
+
+	FStructProperty* StructProperty = CastFieldChecked<FStructProperty>(Property);
 	if (!StructProperty)
 	{
 		return FPropertyAccess::Fail;
