[StateTree] Compile a persistent struct for the global parameter. Referencer, like StateTreeReference, can keep a reference to the persistent struct instead of copying the property bag description. It enables better tracking and propagation, utilizing the UObject override system.
#Synced-CL 49833236

diff --git a/Engine/Plugins/Runtime/StateTree/Source/StateTreeEditorModule/Private/StateTreeCompiler.cpp b/Engine/Plugins/Runtime/StateTree/Source/StateTreeEditorModule/Private/StateTreeCompiler.cpp
--- a/Engine/Plugins/Runtime/StateTree/Source/StateTreeEditorModule/Private/StateTreeCompiler.cpp
+++ b/Engine/Plugins/Runtime/StateTree/Source/StateTreeEditorModule/Private/StateTreeCompiler.cpp
@@ -21,6 +21,9 @@
 #include "StateTreeDelegate.h"
 #include "Customizations/StateTreeEditorNodeUtils.h"
 
+#include "StructUtils/UserDefinedStruct.h"
+#include "StructUtilsEditorUtils.h"
+
 namespace UE::StateTree
 {
 	struct FCompileNodeContext : ICompileNodeContext
@@ -678,10 +681,49 @@ UStateTreeState* FStateTreeCompiler::GetState(const FGuid& StateID) const
 
 bool FStateTreeCompiler::CreateParameters()
 {
-	// Copy parameters from EditorData	
-	StateTree->Parameters = EditorData->GetRootParametersPropertyBag();
+	// Copy parameters from EditorData
 	StateTree->ParameterDataType = EditorData->Schema->GetGlobalParameterDataType();
 
+	if (StateTree->ParameterDataType == EStateTreeParameterDataType::GlobalParameterData)
+	{
+		StateTree->Parameters = EditorData->GetRootParametersPropertyBag();
+
+		// Copy parameters from EditorData
+		const FInstancedPropertyBag& PropertyBag = EditorData->GetRootParametersPropertyBag();
+		if (UUserDefinedStruct* UDS = Cast<UUserDefinedStruct>(StateTree->GlobalParametersStruct))
+		{
+			UE::StructUtils::UpdateUserDefinedStructFromDescs(UDS, PropertyBag);
+		}
+		else
+		{
+			// The linker loads and saves in alphabetic order inside the same package. Add a _ to load/save the UDS before the StateTree.
+			const FName UDSName = *FString::Printf(TEXT("_%s_DefaultParameters"), *StateTree->GetName());
+			TValueOrError<UUserDefinedStruct*, void> NewUDS = UE::StructUtils::CreateUserDefinedStructFromDescs(StateTree->GetPackage(), PropertyBag, UDSName);
+			StateTree->GlobalParametersStruct = NewUDS.HasValue() ? NewUDS.GetValue() : nullptr;
+		}
+
+		if (StateTree->GlobalParametersStruct == nullptr)
+		{
+			Log.Reportf(EMessageSeverity::Error, TEXT("Failed to create the global parameters type."));
+			return false;
+		}
+
+		StateTree->GlobalParameters = StateTree->GlobalParametersStruct;
+
+		// The layout should match because the bindings are expected to work.
+		if (!UE::StructUtils::TheSameLayout(PropertyBag, StateTree->GlobalParameters))
+		{
+			Log.Reportf(EMessageSeverity::Warning, TEXT("The generated user defined struct do not matches with the editor version."));
+			return false;
+		}
+	}
+	else
+	{
+		StateTree->Parameters.Reset();
+		StateTree->GlobalParameters = nullptr;
+		StateTree->GlobalParametersStruct = nullptr;
+	}
+
 	// Mark parameters as binding source
 	const EStateTreeDataSourceType GlobalParameterDataType = UE::StateTree::CastToDataSourceType(StateTree->ParameterDataType);
 	const FStateTreeBindableStructDesc ParametersDesc = {
diff --git a/Engine/Plugins/Runtime/StateTree/Source/StateTreeModule/Private/StateTree.cpp b/Engine/Plugins/Runtime/StateTree/Source/StateTreeModule/Private/StateTree.cpp
--- a/Engine/Plugins/Runtime/StateTree/Source/StateTreeModule/Private/StateTree.cpp
+++ b/Engine/Plugins/Runtime/StateTree/Source/StateTreeModule/Private/StateTree.cpp
@@ -17,6 +17,7 @@
 #include "Misc/DataValidation.h"
 #include "Misc/EnumerateRange.h"
 #include "UObject/AssetRegistryTagsContext.h"
+#include "StructUtils/UserDefinedStruct.h"
 #include "StateTreePropertyFunctionBase.h"
 #include "AutoRTFM.h"
 
@@ -36,6 +37,15 @@ FCustomVersionRegistration GRegisterStateTreeCustomVersion(FStateTreeCustomVersi
 }
 PRAGMA_ENABLE_DEPRECATION_WARNINGS
 
+FConstStructView UStateTree::GetDefaultGlobalParameters() const
+{
+	if (GlobalParameters)
+	{
+		return FConstStructView(GlobalParameters, GlobalParameters->GetDefaultInstance());
+	}
+	return FConstStructView();
+}
+
 bool UStateTree::IsReadyToRun() const
 {
 #if WITH_EDITOR
@@ -229,6 +239,27 @@ void UStateTree::BeginDestroy()
 }
 #endif //WITH_STATETREE_DEBUG
 
+void UStateTree::PostRename(UObject* OldOuter, const FName OldName)
+{
+	Super::PostRename(OldOuter, OldName);
+
+	// Also, rename the GlobalParameters struct to make sure they are in the same package.
+	//The GlobalParameters is RF_Public and needs to follow the StateTree asset.
+	auto RenameStruct = [this, &OldOuter](UUserDefinedStruct* Struct)
+		{
+			if (Struct && Struct->GetOutermost() == OldOuter)
+			{
+				const FString UDSName = FString::Printf(TEXT("_%s_DefaultParameters"), *GetName());
+				Struct->Rename(*UDSName, GetOutermost(), REN_DoNotDirty | REN_SkipGeneratedClasses);
+			}
+		};
+#if WITH_EDITOR
+	RenameStruct(GlobalParametersStruct);
+#else
+	RenameStruct(GlobalParameters);
+#endif
+}
+
 #if WITH_EDITOR
 namespace UE::StateTree::Compiler
 {
@@ -263,6 +294,8 @@ void UStateTree::ResetCompiled()
 	PropertyFunctionEvaluationScopeMemoryRequirements.Reset();
 	Extensions.Reset();
 	Parameters.Reset();
+	GlobalParameters = nullptr;
+	// GlobalParametersType is not reset. We want the same type between different compilation.
 	ParameterDataType = EStateTreeParameterDataType::GlobalParameterData;
 	IDToStateMappings.Reset();
 	IDToNodeMappings.Reset();
@@ -679,6 +712,16 @@ bool UStateTree::Link()
 		return false;
 	}
 
+	// @todo Temporary test. Those are to test that the UDS loaded properly and matches the property bag before removing the Parameters.
+	if (GlobalParameters == nullptr && Parameters.IsValid())
+	{
+		UE_LOG(LogStateTree, Warning, TEXT("%s: The global parameters is invalid."), *GetPathName());
+	}
+	if (!UE::StructUtils::TheSameLayout(Parameters, GlobalParameters))
+	{
+		UE_LOG(LogStateTree, Warning, TEXT("%s: The global parameters doesn't have the same layout as the PropertyBag."), *GetPathName());
+	}
+
 	// Initialize the instance data default value.
 	// This data will be used to allocate runtime instance on all StateTree users.
 	ResetLinked();
@@ -1039,9 +1082,11 @@ bool UStateTree::PatchBindings()
 		{
 			// Check that the bag in linked state matches.
 			const FInstancedPropertyBag& TargetTreeParameters = State.LinkedAsset->Parameters;
+			const UUserDefinedStruct* TargetTreeGlobalParameters = State.LinkedAsset->GlobalParameters;
 			const FCompactStateTreeParameters& Params = DefaultInstanceData.GetMutableStruct(State.ParameterTemplateIndex.Get()).Get<FCompactStateTreeParameters>();
-
-			if (TargetTreeParameters.GetPropertyBagStruct() != Params.GetValue().GetScriptStruct())
+			const bool bIsValid = TargetTreeParameters.GetPropertyBagStruct() == Params.GetValue().GetScriptStruct()
+				|| TargetTreeGlobalParameters == Params.GetValue().GetScriptStruct();
+			if (!bIsValid)
 			{
 				UE_LOG(LogStateTree, Error, TEXT("%s: The parameters on state '%s' does not match the linked asset parameters '%s'. Please recompile the StateTree asset."),
 					*GetPathName(), *State.Name.ToString(), *State.LinkedAsset->GetPathName());
diff --git a/Engine/Plugins/Runtime/StateTree/Source/StateTreeModule/Public/StateTree.h b/Engine/Plugins/Runtime/StateTree/Source/StateTreeModule/Public/StateTree.h
--- a/Engine/Plugins/Runtime/StateTree/Source/StateTreeModule/Public/StateTree.h
+++ b/Engine/Plugins/Runtime/StateTree/Source/StateTreeModule/Public/StateTree.h
@@ -180,6 +180,9 @@ public:
 		return Parameters;
 	}
 
+	/** @return The global parameters of the state tree. Global parameter values can be overridden at runtime by the execution context. */
+	UE_API FConstStructView GetDefaultGlobalParameters() const;
+
 	/** @return true if the tree asset can be used at runtime. */
 	UE_API bool IsReadyToRun() const;
 
@@ -370,6 +373,7 @@ protected:
 	UE_API virtual void PostInitProperties() override;
 #endif
 
+	UE_API virtual void PostRename(UObject* OldOuter, const FName OldName) override;
 #if WITH_EDITOR
 	UE_API virtual void PostDuplicate(EDuplicateMode::Type DuplicateMode) override;
 	UE_API virtual void GetAssetRegistryTags(FAssetRegistryTagsContext Context) const override;
@@ -473,6 +477,20 @@ private:
 	UPROPERTY()
 	FInstancedPropertyBag Parameters;
 
+	/**
+	 * Parameters that could be used for bindings within the Tree.
+	 * Default values are stored within the asset but StateTreeReference can be used to parameterized the tree.
+	 * @see FStateTreeReference
+	 */
+	UPROPERTY()
+	TObjectPtr<UUserDefinedStruct> GlobalParameters;
+
+#if WITH_EDITORONLY_DATA
+	/** Keep a reference between compilation. GlobalParameters is reset between each compilation. */
+	UPROPERTY()
+	TObjectPtr<UUserDefinedStruct> GlobalParametersStruct;
+#endif
+
 	//~ Data created during linking.
 	/** List of external data required by the state tree, created during linking. */
 	UPROPERTY(Transient)
