Fixed non-unity compilation

#Synced-CL 49887695

diff --git a/Engine/Plugins/Runtime/StateTree/Source/StateTreeEditorModule/Private/Debugger/StateTreeRewindDebuggerExtensions.cpp b/Engine/Plugins/Runtime/StateTree/Source/StateTreeEditorModule/Private/Debugger/StateTreeRewindDebuggerExtensions.cpp
--- a/Engine/Plugins/Runtime/StateTree/Source/StateTreeEditorModule/Private/Debugger/StateTreeRewindDebuggerExtensions.cpp
+++ b/Engine/Plugins/Runtime/StateTree/Source/StateTreeEditorModule/Private/Debugger/StateTreeRewindDebuggerExtensions.cpp
@@ -6,6 +6,8 @@
 #include "Debugger/StateTreeTraceTypes.h"
 #include "StateTreeDelegates.h"
 #include "IRewindDebugger.h"
+#include "TraceServices/Model/AnalysisSession.h"
+#include "TraceServices/Model/Frames.h"
 
 namespace UE::StateTreeDebugger
 {
