[StateTree] Compile a persistent struct for the global parameter. Referencer, like StateTreeReference, can keep a reference to the persistent struct instead of copying the property bag description. It enables better tracking and propagation, utilizing the UObject override system.
#Synced-CL 49833236
#was-backout 49857190
#oldpreflight 696e7cefe025219a1e4d5fe5, 696e47f39d00213f1ff84c3d, 696aa9c3d36ea27cd7644382

#Synced-CL 49855772

diff --git a/Engine/Plugins/Runtime/StateTree/Source/StateTreeEditorModule/Private/StateTreeCompiler.cpp b/Engine/Plugins/Runtime/StateTree/Source/StateTreeEditorModule/Private/StateTreeCompiler.cpp
--- a/Engine/Plugins/Runtime/StateTree/Source/StateTreeEditorModule/Private/StateTreeCompiler.cpp
+++ b/Engine/Plugins/Runtime/StateTree/Source/StateTreeEditorModule/Private/StateTreeCompiler.cpp
@@ -21,6 +21,9 @@
 #include "StateTreeDelegate.h"
 #include "Customizations/StateTreeEditorNodeUtils.h"
 
+#include "StructUtils/UserDefinedStruct.h"
+#include "StructUtilsEditorUtils.h"
+
 namespace UE::StateTree
 {
 	struct FCompileNodeContext : ICompileNodeContext
@@ -678,10 +681,54 @@ UStateTreeState* FStateTreeCompiler::GetState(const FGuid& StateID) const
 
 bool FStateTreeCompiler::CreateParameters()
 {
-	// Copy parameters from EditorData	
-	StateTree->Parameters = EditorData->GetRootParametersPropertyBag();
+	// Copy parameters from EditorData
 	StateTree->ParameterDataType = EditorData->Schema->GetGlobalParameterDataType();
 
+	if (StateTree->ParameterDataType == EStateTreeParameterDataType::GlobalParameterData)
+	{
+		StateTree->Parameters = EditorData->GetRootParametersPropertyBag();
+
+		// Copy parameters from EditorData
+		const FInstancedPropertyBag& PropertyBag = EditorData->GetRootParametersPropertyBag();
+		if (UUserDefinedStruct* UDS = Cast<UUserDefinedStruct>(StateTree->GlobalParametersStructType))
+		{
+			UE::StructUtils::UpdateUserDefinedStructFromDescs(UDS, PropertyBag);
+		}
+		else
+		{
+			// The linker loads and saves in alphabetic order inside the same package. Add a _ to load/save the UDS before the StateTree.
+			const FName UDSName = *FString::Printf(TEXT("_%s_DefaultParameters"), *StateTree->GetName());
+			UE::StructUtils::FCreateUserDefinedStructArgs Args;
+			Args.UserDefinedStructFlags = RF_Transactional | RF_Public;
+			Args.bIsBlueprintType = false;
+			Args.bAddPropertyIDToPropertyName = false;
+			TValueOrError<UUserDefinedStruct*, void> NewUDS = UE::StructUtils::CreateUserDefinedStructFromDescs(StateTree->GetPackage(), PropertyBag, UDSName, Args);
+			StateTree->GlobalParametersStructType = NewUDS.HasValue() ? NewUDS.GetValue() : nullptr;
+		}
+
+		if (StateTree->GlobalParametersStructType == nullptr)
+		{
+			Log.Reportf(EMessageSeverity::Error, TEXT("Failed to create the global parameters type."));
+			return false;
+		}
+
+		StateTree->GlobalParameters = StateTree->GlobalParametersStructType;
+
+		// The layout should match because the bindings are expected to work.
+		constexpr bool bCheckPropertiesNames = true;
+		if (!UE::StructUtils::TheSameLayout(PropertyBag, StateTree->GlobalParameters, bCheckPropertiesNames))
+		{
+			Log.Reportf(EMessageSeverity::Warning, TEXT("The generated user defined struct do not matches with the editor version."));
+			return false;
+		}
+	}
+	else
+	{
+		StateTree->Parameters.Reset();
+		StateTree->GlobalParameters = nullptr;
+		StateTree->GlobalParametersStructType = nullptr;
+	}
+
 	// Mark parameters as binding source
 	const EStateTreeDataSourceType GlobalParameterDataType = UE::StateTree::CastToDataSourceType(StateTree->ParameterDataType);
 	const FStateTreeBindableStructDesc ParametersDesc = {
diff --git a/Engine/Plugins/Runtime/StateTree/Source/StateTreeModule/Private/StateTree.cpp b/Engine/Plugins/Runtime/StateTree/Source/StateTreeModule/Private/StateTree.cpp
--- a/Engine/Plugins/Runtime/StateTree/Source/StateTreeModule/Private/StateTree.cpp
+++ b/Engine/Plugins/Runtime/StateTree/Source/StateTreeModule/Private/StateTree.cpp
@@ -17,8 +17,10 @@
 #include "Misc/DataValidation.h"
 #include "Misc/EnumerateRange.h"
 #include "UObject/AssetRegistryTagsContext.h"
+#include "StructUtils/UserDefinedStruct.h"
 #include "StateTreePropertyFunctionBase.h"
 #include "AutoRTFM.h"
+#include "UObject/ObjectSaveContext.h"
 
 #include <atomic>
 
@@ -36,6 +38,15 @@ FCustomVersionRegistration GRegisterStateTreeCustomVersion(FStateTreeCustomVersi
 }
 PRAGMA_ENABLE_DEPRECATION_WARNINGS
 
+FConstStructView UStateTree::GetDefaultGlobalParameters() const
+{
+	if (GlobalParameters)
+	{
+		return FConstStructView(GlobalParameters, GlobalParameters->GetDefaultInstance());
+	}
+	return FConstStructView();
+}
+
 bool UStateTree::IsReadyToRun() const
 {
 #if WITH_EDITOR
@@ -229,6 +240,27 @@ void UStateTree::BeginDestroy()
 }
 #endif //WITH_STATETREE_DEBUG
 
+void UStateTree::PostRename(UObject* OldOuter, const FName OldName)
+{
+	Super::PostRename(OldOuter, OldName);
+
+	// Also, rename the GlobalParameters struct to make sure they are in the same package.
+	//The GlobalParameters is RF_Public and needs to follow the StateTree asset.
+	auto RenameStruct = [this, &OldOuter](UUserDefinedStruct* Struct)
+		{
+			if (Struct && Struct->GetOutermost() == OldOuter)
+			{
+				const FString UDSName = FString::Printf(TEXT("_%s_DefaultParameters"), *GetName());
+				Struct->Rename(*UDSName, GetOutermost(), REN_DoNotDirty | REN_SkipGeneratedClasses);
+			}
+		};
+#if WITH_EDITOR
+	RenameStruct(GlobalParametersStructType);
+#else
+	RenameStruct(GlobalParameters);
+#endif
+}
+
 #if WITH_EDITOR
 namespace UE::StateTree::Compiler
 {
@@ -263,6 +295,8 @@ void UStateTree::ResetCompiled()
 	PropertyFunctionEvaluationScopeMemoryRequirements.Reset();
 	Extensions.Reset();
 	Parameters.Reset();
+	GlobalParameters = nullptr;
+	// GlobalParametersStructType is not reset. We want the same type between different compilation.
 	ParameterDataType = EStateTreeParameterDataType::GlobalParameterData;
 	IDToStateMappings.Reset();
 	IDToNodeMappings.Reset();
@@ -309,6 +343,17 @@ void UStateTree::MarkAsModified(const bool bPubliclyModified)
 	CompileStatus = bPubliclyModified ? ECompileStatus::PublicDirty : ECompileStatus::InternalDirty;
 }
 
+void UStateTree::PreSave(FObjectPreSaveContext Context)
+{
+	const bool bIsUserObject = !HasAnyFlags(RF_ClassDefaultObject | RF_ArchetypeObject);
+	const bool bIsEditorAutoSave = Context.IsFromAutoSave();
+	if (bIsUserObject && !bIsEditorAutoSave)
+	{
+		bSavedWithGlobalParametersStructType = true;
+	}
+	Super::PreSave(Context);
+}
+
 void UStateTree::PostDuplicate(EDuplicateMode::Type DuplicateMode)
 {
 	Super::PostDuplicate(DuplicateMode);
@@ -679,6 +724,13 @@ bool UStateTree::Link()
 		return false;
 	}
 
+	// @todo Temporary test. Those are to test that the UDS loaded properly and matches the property bag before removing the Parameters.
+	constexpr bool bCheckPropertiesNames = true;
+	if (!UE::StructUtils::TheSameLayout(Parameters, GlobalParameters, bCheckPropertiesNames))
+	{
+		UE_LOG(LogStateTree, Warning, TEXT("%s: The global parameters doesn't have the same layout as the PropertyBag."), *GetPathName());
+	}
+
 	// Initialize the instance data default value.
 	// This data will be used to allocate runtime instance on all StateTree users.
 	ResetLinked();
@@ -1039,9 +1091,11 @@ bool UStateTree::PatchBindings()
 		{
 			// Check that the bag in linked state matches.
 			const FInstancedPropertyBag& TargetTreeParameters = State.LinkedAsset->Parameters;
+			const UUserDefinedStruct* TargetTreeGlobalParameters = State.LinkedAsset->GlobalParameters;
 			const FCompactStateTreeParameters& Params = DefaultInstanceData.GetMutableStruct(State.ParameterTemplateIndex.Get()).Get<FCompactStateTreeParameters>();
-
-			if (TargetTreeParameters.GetPropertyBagStruct() != Params.GetValue().GetScriptStruct())
+			const bool bIsValid = TargetTreeParameters.GetPropertyBagStruct() == Params.GetValue().GetScriptStruct()
+				|| TargetTreeGlobalParameters == Params.GetValue().GetScriptStruct();
+			if (!bIsValid)
 			{
 				UE_LOG(LogStateTree, Error, TEXT("%s: The parameters on state '%s' does not match the linked asset parameters '%s'. Please recompile the StateTree asset."),
 					*GetPathName(), *State.Name.ToString(), *State.LinkedAsset->GetPathName());
diff --git a/Engine/Plugins/Runtime/StateTree/Source/StateTreeModule/Public/StateTree.h b/Engine/Plugins/Runtime/StateTree/Source/StateTreeModule/Public/StateTree.h
--- a/Engine/Plugins/Runtime/StateTree/Source/StateTreeModule/Public/StateTree.h
+++ b/Engine/Plugins/Runtime/StateTree/Source/StateTreeModule/Public/StateTree.h
@@ -180,6 +180,9 @@ public:
 		return Parameters;
 	}
 
+	/** @return The global parameters of the state tree. Global parameter values can be overridden at runtime by the execution context. */
+	UE_API FConstStructView GetDefaultGlobalParameters() const;
+
 	/** @return true if the tree asset can be used at runtime. */
 	UE_API bool IsReadyToRun() const;
 
@@ -325,6 +328,12 @@ public:
 	/** Compile the state trees if the editor hash data as changed since the last compilation. */
 	UE_API void CompileIfChanged();
 
+	/** Was the GlobalParameters properly saved. */
+	bool WasSavedWithGlobalParametersStruct() const
+	{
+		return bSavedWithGlobalParametersStructType;
+	}
+
 #endif
 
 #if WITH_EDITOR || WITH_STATETREE_DEBUG
@@ -370,7 +379,9 @@ protected:
 	UE_API virtual void PostInitProperties() override;
 #endif
 
+	UE_API virtual void PostRename(UObject* OldOuter, const FName OldName) override;
 #if WITH_EDITOR
+	UE_API virtual void PreSave(FObjectPreSaveContext Context) override;
 	UE_API virtual void PostDuplicate(EDuplicateMode::Type DuplicateMode) override;
 	UE_API virtual void GetAssetRegistryTags(FAssetRegistryTagsContext Context) const override;
 	UE_API void BeginCacheForCookedPlatformData(const ITargetPlatform* TargetPlatform) override;
@@ -473,6 +484,24 @@ private:
 	UPROPERTY()
 	FInstancedPropertyBag Parameters;
 
+	/**
+	 * Parameters that could be used for bindings within the Tree.
+	 * Default values are stored within the asset but StateTreeReference can be used to parameterized the tree.
+	 * @see FStateTreeReference
+	 */
+	UPROPERTY()
+	TObjectPtr<UUserDefinedStruct> GlobalParameters;
+
+#if WITH_EDITORONLY_DATA
+	/** Keep a reference between compilation. GlobalParameters is reset between each compilation. */
+	UPROPERTY()
+	TObjectPtr<UUserDefinedStruct> GlobalParametersStructType;
+
+	/**  To know if the new type was saved and the reference can be considered safe to use.  */
+	UPROPERTY()
+	bool bSavedWithGlobalParametersStructType = false;
+#endif
+
 	//~ Data created during linking.
 	/** List of external data required by the state tree, created during linking. */
 	UPROPERTY(Transient)
diff --git a/Engine/Source/Runtime/CoreUObject/Private/StructUtils/StructUtilsTypes.cpp b/Engine/Source/Runtime/CoreUObject/Private/StructUtils/StructUtilsTypes.cpp
--- a/Engine/Source/Runtime/CoreUObject/Private/StructUtils/StructUtilsTypes.cpp
+++ b/Engine/Source/Runtime/CoreUObject/Private/StructUtils/StructUtilsTypes.cpp
@@ -236,7 +236,7 @@ namespace UE::StructUtils
 		return GetStructHash64Helper(SharedView);
 	}
 
-	bool TheSameLayout(const FInstancedPropertyBag& A, const UScriptStruct* B)
+	bool TheSameLayout(const FInstancedPropertyBag& A, const UScriptStruct* B, bool bCheckPropertiesNames)
 	{
 		// Both invalid or Both empty.
 		if (!A.IsValid())
@@ -244,6 +244,6 @@ namespace UE::StructUtils
 			return B == nullptr || B->PropertyLink == nullptr;
 		}
 
-		return FStructUtils::TheSameLayout(A.GetPropertyBagStruct(), B);
+		return FStructUtils::TheSameLayout(A.GetPropertyBagStruct(), B, bCheckPropertiesNames);
 	}
 }
diff --git a/Engine/Source/Runtime/CoreUObject/Public/StructUtils/StructUtilsTypes.h b/Engine/Source/Runtime/CoreUObject/Public/StructUtils/StructUtilsTypes.h
--- a/Engine/Source/Runtime/CoreUObject/Public/StructUtils/StructUtilsTypes.h
+++ b/Engine/Source/Runtime/CoreUObject/Public/StructUtils/StructUtilsTypes.h
@@ -75,7 +75,7 @@ namespace UE::StructUtils
 	 * The struct UPropertyBag might not be generated if there are no properties in the struct.
 	 * In that case, it matches if the UScriptStruct also doesn’t have properties.
 	 */
-	extern UE_API bool TheSameLayout(const FInstancedPropertyBag& A, const UScriptStruct* B);
+	extern UE_API bool TheSameLayout(const FInstancedPropertyBag& A, const UScriptStruct* B, bool bCheckPropertiesNames = false);
 }
 
 /* Predicate useful to find a struct of a specific type in an container */
