[StateTree] Modify FStateTreeReference to save the UserDefinedStruct produced by the StateTree asset. We keep the PropertyBag for now as a test. We want to make sure everything works before removing the PropertyBag.

#Synced-CL 49855772

#Synced-CL 49855772

diff --git a/Engine/Plugins/Runtime/StateTree/Source/StateTreeModule/Private/StateTreeReference.cpp b/Engine/Plugins/Runtime/StateTree/Source/StateTreeModule/Private/StateTreeReference.cpp
--- a/Engine/Plugins/Runtime/StateTree/Source/StateTreeModule/Private/StateTreeReference.cpp
+++ b/Engine/Plugins/Runtime/StateTree/Source/StateTreeModule/Private/StateTreeReference.cpp
@@ -2,9 +2,43 @@
 
 #include "StateTreeReference.h"
 #include "StateTree.h"
+#include "StructUtils/StructUtilsTypes.h"
+#include "UObject/StructPostLoadContext.h"
+#include "UObject/UObjectThreadContext.h"
 
 #include UE_INLINE_GENERATED_CPP_BY_NAME(StateTreeReference)
 
+void FStateTreeReference::SetStateTree(UStateTree* NewStateTree)
+{
+	StateTree = NewStateTree;
+	SyncParameters();
+
+	if (StateTree)
+	{
+		GlobalParameters = StateTree->GetDefaultGlobalParameters();
+	}
+	else
+	{
+		GlobalParameters.Reset();
+	}
+}
+
+FConstStructView FStateTreeReference::GetGlobalParameters() const
+{
+	ConditionallySyncParameters();
+	// @todo Temporary test. This is to test that the UDS loaded properly and matches the property bag before removing the Parameters.
+	ensure(UE::StructUtils::TheSameLayout(Parameters, GlobalParameters.GetScriptStruct()));
+	return Parameters.GetValue();
+}
+
+FStructView FStateTreeReference::GetMutableGlobalParameters()
+{
+	ConditionallySyncParameters();
+	// @todo Temporary test. This is to test that the UDS loaded properly and matches the property bag before removing the Parameters.
+	ensure(UE::StructUtils::TheSameLayout(Parameters, GlobalParameters.GetScriptStruct()));
+	return Parameters.GetMutableValue();
+}
+
 void FStateTreeReference::SyncParameters()
 {
 	if (StateTree == nullptr)
@@ -135,4 +169,62 @@ PRAGMA_ENABLE_DEPRECATION_WARNINGS
 			}
 		}
 	}
+
+#if WITH_EDITOR
+	if (Ar.IsLoading())
+	{
+		// The GlobalParameters need to be updated, but they are only available after the StateTree is serialized and loaded.
+ 		if (StateTree != nullptr && !GlobalParameters.IsValid())
+		{
+			if (FUObjectSerializeContext* SerializeContext = FUObjectThreadContext::Get().GetSerializeContext())
+			{
+				UE::FStructPostLoadContext::Get().RequestPostLoad(SerializeContext);
+			}
+		}
+		else if (StateTree == nullptr)
+		{
+			GlobalParameters.Reset();
+		}
+	}
+#endif
+}
+
+#if WITH_EDITOR
+void FStateTreeReference::PostLoad()
+{
+	if (StateTree)
+	{
+		StateTree->ConditionalPostLoad();
+		ConditionallySyncParameters();
+
+		GlobalParameters = StateTree->GetDefaultGlobalParameters();
+
+		if (ensure(GlobalParameters.IsValid()))
+		{
+			if (ensure(UE::StructUtils::TheSameLayout(Parameters, GlobalParameters.GetScriptStruct())))
+			{
+				for (const FGuid& PropertyID : PropertyOverrides)
+				{
+					if (const FPropertyBagPropertyDesc* BagPropertyDescription = Parameters.GetPropertyBagStruct()->FindPropertyDescByID(PropertyID))
+					{
+						const FProperty* NewProperty = GlobalParameters.GetScriptStruct()->FindPropertyByName(BagPropertyDescription->Name);
+						if (NewProperty == nullptr)
+						{
+							NewProperty = GlobalParameters.GetScriptStruct()->CustomFindProperty(BagPropertyDescription->Name);
+						}
+
+						if (ensure(NewProperty) && ensure(BagPropertyDescription->CachedProperty))
+						{
+							NewProperty->CopyCompleteValue_InContainer(GlobalParameters.GetMutableMemory(), Parameters.GetMutableValue().GetMemory());
+						}
+					}
+				}
+			}
+		}
+	}
+	else
+	{
+		GlobalParameters.Reset();
+	}
 }
+#endif
\ No newline at end of file
diff --git a/Engine/Plugins/Runtime/StateTree/Source/StateTreeModule/Public/StateTreeReference.h b/Engine/Plugins/Runtime/StateTree/Source/StateTreeModule/Public/StateTreeReference.h
--- a/Engine/Plugins/Runtime/StateTree/Source/StateTreeModule/Public/StateTreeReference.h
+++ b/Engine/Plugins/Runtime/StateTree/Source/StateTreeModule/Public/StateTreeReference.h
@@ -25,11 +25,7 @@ struct FStateTreeReference
 	}
 	
 	/** Sets the StateTree asset and referenced parameters. */
-	void SetStateTree(UStateTree* NewStateTree)
-	{
-		StateTree = NewStateTree;
-		SyncParameters();
-	}
+	UE_API void SetStateTree(UStateTree* NewStateTree);
 
 	/** @return const pointer to the referenced StateTree asset. */
 	const UStateTree* GetStateTree() const
@@ -44,18 +40,10 @@ struct FStateTreeReference
 	}
 
 	/** @return the parameters for the referenced StateTree asset. */
-	FConstStructView GetGlobalParameters() const
-	{
-		ConditionallySyncParameters();
-		return Parameters.GetValue();
-	}
+	UE_API FConstStructView GetGlobalParameters() const;
 
 	/** @return the parameters for the referenced StateTree asset. */
-	FStructView GetMutableGlobalParameters()
-	{
-		ConditionallySyncParameters();
-		return Parameters.GetMutableValue();
-	}
+	UE_API FStructView GetMutableGlobalParameters();
 
 	/** @return reference to the parameters for the referenced StateTree asset. */
 	const FInstancedPropertyBag& GetParameters() const
@@ -96,6 +84,9 @@ struct FStateTreeReference
 
 	UE_API bool Serialize(FStructuredArchive::FSlot Slot);
 	UE_API void PostSerialize(const FArchive& Ar);
+#if WITH_EDITOR
+	UE_API void PostLoad();
+#endif
 
 protected:
 	UPROPERTY(EditAnywhere, Category = "StateTree")
@@ -104,6 +95,9 @@ protected:
 	UPROPERTY(EditAnywhere, Category = "StateTree", meta = (FixedLayout))
 	FInstancedPropertyBag Parameters;
 
+	UPROPERTY(EditAnywhere, Category = "StateTree")
+	FInstancedStruct GlobalParameters;
+
 	/** Array of overridden properties. Non-overridden properties will inherit the values from the StateTree default parameters. */
 	UPROPERTY(EditAnywhere, Category = "StateTree")
 	TArray<FGuid> PropertyOverrides;
@@ -118,6 +112,7 @@ struct TStructOpsTypeTraits<FStateTreeReference> : public TStructOpsTypeTraitsBa
 	{
 		WithStructuredSerializer = true,
 		WithPostSerialize = true,
+		WithPostLoad = true,
 	};
 };
 
