[StateTree] Downgrade StateTreeReference update message.

#Synced-CL 49855772

diff --git a/Engine/Plugins/Runtime/StateTree/Source/StateTreeModule/Private/StateTreeReference.cpp b/Engine/Plugins/Runtime/StateTree/Source/StateTreeModule/Private/StateTreeReference.cpp
--- a/Engine/Plugins/Runtime/StateTree/Source/StateTreeModule/Private/StateTreeReference.cpp
+++ b/Engine/Plugins/Runtime/StateTree/Source/StateTreeModule/Private/StateTreeReference.cpp
@@ -124,7 +124,7 @@ void FStateTreeReference::ConditionallySyncParameters() const
 	{
 		FStateTreeReference* NonConstThis = const_cast<FStateTreeReference*>(this);
 		NonConstThis->SyncParameters();
-		UE_LOG(LogStateTree, Warning, TEXT("Parameters for '%s' stored in StateTreeReference were auto-fixed to be usable at runtime."), *GetNameSafe(StateTree));	
+		UE_LOG(LogStateTree, Log, TEXT("Parameters for '%s' stored in StateTreeReference were auto-fixed to be usable at runtime."), *GetNameSafe(StateTree));	
 	}
 }
 
