Deprecate ShouldForceHideProperty(const TSharedRef<FPropertyNode>) and replace with IPropertyHandle version
* Since FPropertyNode is a private type, the value of this delegate was questionable and resulted in invalid include paths to private headers
* Replaced with an IPropertyHandle version and replaced all usages in the engine with new version
* Fixed up other bad usages of the Private PropertyNode include
* Note, does not go through normal UE_DEPRECATION path since the lambda was not useful outside of the property editor module without hacking the #include to reach into Private folders.

#Synced-CL 49872642

diff --git a/Engine/Plugins/Chooser/Source/ChooserEditor/Private/CurveOverrideCustomization.cpp b/Engine/Plugins/Chooser/Source/ChooserEditor/Private/CurveOverrideCustomization.cpp
--- a/Engine/Plugins/Chooser/Source/ChooserEditor/Private/CurveOverrideCustomization.cpp
+++ b/Engine/Plugins/Chooser/Source/ChooserEditor/Private/CurveOverrideCustomization.cpp
@@ -10,7 +10,6 @@
 #include "Widgets/Input/SEditableTextBox.h"
 #include "Widgets/Input/SNumericEntryBox.h"
 #include "ScopedTransaction.h"
-#include "Editor/PropertyEditor/Private/PropertyNode.h"
 
 #define LOCTEXT_NAMESPACE "CurveOverrideCustomization"
 
diff --git a/Engine/Plugins/Experimental/Dataflow/Source/DataflowEditor/Private/Dataflow/DataflowMembersWidget.cpp b/Engine/Plugins/Experimental/Dataflow/Source/DataflowEditor/Private/Dataflow/DataflowMembersWidget.cpp
--- a/Engine/Plugins/Experimental/Dataflow/Source/DataflowEditor/Private/Dataflow/DataflowMembersWidget.cpp
+++ b/Engine/Plugins/Experimental/Dataflow/Source/DataflowEditor/Private/Dataflow/DataflowMembersWidget.cpp
@@ -14,13 +14,13 @@
 #include "Dataflow/DataflowVariablePaletteItem.h"
 #include "Dataflow/DataflowSubGraphPaletteItem.h"
 #include "Dataflow/DataflowVariableNodes.h"
-#include "Editor/PropertyEditor/Private/PropertyNode.h"
 #include "Framework/Commands/UICommandList.h"
 #include "GraphActionNode.h"
 #include "IDetailsView.h"
 #include "IStructureDetailsView.h"
 #include "DetailsViewArgs.h"
 #include "PropertyEditorModule.h"
+#include "PropertyHandle.h"
 #include "SGraphActionMenu.h"
 #include "Widgets/Input/SComboButton.h"
 #include "Widgets/Input/SSearchBox.h"
diff --git a/Engine/Plugins/Experimental/InstanceDataObjectFixupTool/Source/Private/InstanceDataObjectFixupPanel.cpp b/Engine/Plugins/Experimental/InstanceDataObjectFixupTool/Source/Private/InstanceDataObjectFixupPanel.cpp
--- a/Engine/Plugins/Experimental/InstanceDataObjectFixupTool/Source/Private/InstanceDataObjectFixupPanel.cpp
+++ b/Engine/Plugins/Experimental/InstanceDataObjectFixupTool/Source/Private/InstanceDataObjectFixupPanel.cpp
@@ -9,7 +9,7 @@
 #include "InstanceDataObjectFixupDetailCustomization.h"
 #include "Misc/StringOutputDevice.h"
 #include "Modules/ModuleManager.h"
-#include "Editor/PropertyEditor/Private/PropertyNode.h"
+#include "PropertyHandle.h"
 #include "UObject/PropertyBagRepository.h"
 #include "UObject/PropertyStateTracking.h"
 
@@ -275,9 +275,14 @@ TSharedPtr<IDetailsView>& FInstanceDataObjectFixupPanel::GenerateDetailsView(boo
 
 	if (HasViewFlag(EViewFlags::IncludeOnlySetBySerialization))
 	{
-		DetailsViewArgs.ShouldForceHideProperty.BindLambda([this](const TSharedRef<FPropertyNode>& PropertyNode)->bool
+		DetailsViewArgs.ShouldForceHideProperty.BindLambda([this](const TSharedRef<IPropertyHandle>& PropertyNode)->bool
 		{
-			return !IsPathOverridden(InstanceDataObject, *FPropertyNode::CreatePropertyPath(PropertyNode));
+			TSharedPtr<FPropertyPath> PropertyPath = PropertyNode->CreateFPropertyPath();
+			if (PropertyPath)
+			{
+				return !IsPathOverridden(InstanceDataObject, *PropertyPath);
+			}
+			return false;
 		});
 	}
 	
diff --git a/Engine/Plugins/Experimental/PhysicsControl/Source/PhysicsControlEditor/Private/PhysicsControlAssetEditor.cpp b/Engine/Plugins/Experimental/PhysicsControl/Source/PhysicsControlEditor/Private/PhysicsControlAssetEditor.cpp
--- a/Engine/Plugins/Experimental/PhysicsControl/Source/PhysicsControlEditor/Private/PhysicsControlAssetEditor.cpp
+++ b/Engine/Plugins/Experimental/PhysicsControl/Source/PhysicsControlEditor/Private/PhysicsControlAssetEditor.cpp
@@ -19,7 +19,7 @@
 #include "AnimationEditorPreviewActor.h"
 #include "AnimPreviewInstance.h"
 #include "Components/StaticMeshComponent.h"
-#include "Editor/PropertyEditor/Private/PropertyNode.h"
+#include "PropertyHandle.h"
 #include "EditorModeManager.h"
 #include "IDetailsView.h"
 #include "IPersonaToolkit.h"
@@ -532,7 +532,7 @@ void FPhysicsControlAssetEditor::ExtendViewportMenus()
 					// I'd prefer to do this using SetIsPropertyVisibleDelegate but for some reason
 					// any delegate registered gets cleared by the time the visibilty check call
 					// gets made.
-					DetailsViewArgs.ShouldForceHideProperty.BindLambda([](const TSharedRef<FPropertyNode>& PropertyNode)->bool
+					DetailsViewArgs.ShouldForceHideProperty.BindLambda([](const TSharedRef<IPropertyHandle>& PropertyNode)->bool
 						{
 							const FName PropertyName = PropertyNode->GetProperty()->GetFName();
 							if (PropertyName == GET_MEMBER_NAME_CHECKED(UPhysicsAssetEditorOptions, PhysicsBlend) ||
diff --git a/Engine/Plugins/PCG/Source/PCGEditor/Private/Widgets/SPCGEditorGraphDetailsView.cpp b/Engine/Plugins/PCG/Source/PCGEditor/Private/Widgets/SPCGEditorGraphDetailsView.cpp
--- a/Engine/Plugins/PCG/Source/PCGEditor/Private/Widgets/SPCGEditorGraphDetailsView.cpp
+++ b/Engine/Plugins/PCG/Source/PCGEditor/Private/Widgets/SPCGEditorGraphDetailsView.cpp
@@ -12,7 +12,7 @@
 #include "IDetailPropertyExtensionHandler.h"
 #include "IDetailsView.h"
 #include "PropertyEditorModule.h"
-#include "Editor/PropertyEditor/Private/PropertyNode.h"
+#include "PropertyHandle.h"
 #include "Framework/Notifications/NotificationManager.h"
 #include "Layout/Visibility.h"
 #include "Modules/ModuleManager.h"
@@ -132,20 +132,20 @@ void SPCGEditorGraphDetailsView::Construct(const FArguments& InArgs)
 	DetailsViewArgs.bHideSelectionTip = true;
 	DetailsViewArgs.DefaultsOnlyVisibility = EEditDefaultsOnlyNodeVisibility::Automatic;
 	// Filter out the graph parameters from the details view
-	DetailsViewArgs.ShouldForceHideProperty = FDetailsViewArgs::FShouldForceHideProperty::CreateLambda(
-	[](const TSharedRef<FPropertyNode>& PropertyNode)
-	{
-		if (const FProperty* Property = PropertyNode->GetProperty())
+	DetailsViewArgs.ShouldForceHideProperty.BindLambda(
+		[](const TSharedRef<IPropertyHandle>& PropertyNode)
 		{
-			if (const FStructProperty* PropertyStruct = CastField<FStructProperty>(Property))
+			if (const FProperty* Property = PropertyNode->GetProperty())
 			{
-				// Explicit because FInstancedPropertyBag holds the Custom DetailsView for the graphs.
-				return PropertyStruct->Struct == FInstancedPropertyBag::StaticStruct() && PropertyStruct->GetOwnerClass() == UPCGGraph::StaticClass();
+				if (const FStructProperty* PropertyStruct = CastField<FStructProperty>(Property))
+				{
+					// Explicit because FInstancedPropertyBag holds the Custom DetailsView for the graphs.
+					return PropertyStruct->Struct == FInstancedPropertyBag::StaticStruct() && PropertyStruct->GetOwnerClass() == UPCGGraph::StaticClass();
+				}
 			}
-		}
 
-		return false;
-	});
+			return false;
+		});
 
 	FPropertyEditorModule& PropertyEditorModule = FModuleManager::GetModuleChecked<FPropertyEditorModule>("PropertyEditor");
 	DetailsView = PropertyEditorModule.CreateDetailView(DetailsViewArgs);
diff --git a/Engine/Plugins/PCG/Source/PCGEditor/Private/Widgets/SPCGEditorGraphParamsView.cpp b/Engine/Plugins/PCG/Source/PCGEditor/Private/Widgets/SPCGEditorGraphParamsView.cpp
--- a/Engine/Plugins/PCG/Source/PCGEditor/Private/Widgets/SPCGEditorGraphParamsView.cpp
+++ b/Engine/Plugins/PCG/Source/PCGEditor/Private/Widgets/SPCGEditorGraphParamsView.cpp
@@ -8,8 +8,7 @@
 
 #include "PropertyBagDetails.h"
 #include "PropertyEditorModule.h"
-#include "Editor/PropertyEditor/Private/PropertyNode.h"
-#include "Editor/PropertyEditor/Private/StructurePropertyNode.h"
+#include "PropertyHandle.h"
 #include "Modules/ModuleManager.h"
 
 void SPCGEditorGraphUserParametersView::Construct(const FArguments& InArgs, const TSharedPtr<FPCGEditor>& InPCGEditor)
@@ -50,36 +49,36 @@ void SPCGEditorGraphUserParametersView::Construct(const FArguments& InArgs, cons
 		DetailsViewArgs.bUpdatesFromSelection = false;
 		DetailsViewArgs.bForceHiddenPropertyVisibility = false;
 		// Filter all but the graph parameters
-		DetailsViewArgs.ShouldForceHideProperty = FDetailsViewArgs::FShouldForceHideProperty::CreateLambda(
-		[](const TSharedRef<FPropertyNode>& PropertyNode)
-		{
-			auto PropertyIsPropertyBag = [](const FProperty* Property)
+		DetailsViewArgs.ShouldForceHideProperty.BindLambda(
+			[](const TSharedRef<IPropertyHandle>& PropertyNode)
 			{
-				const FStructProperty* PropertyStruct = CastField<FStructProperty>(Property);
-				// Explicitly check against FInstancedPropertyBag, because that is what holds the Custom DetailsView.
-				return PropertyStruct && PropertyStruct->Struct == FInstancedPropertyBag::StaticStruct();
-			};
+				auto PropertyIsPropertyBag = [](const FProperty* Property)
+				{
+					const FStructProperty* PropertyStruct = CastField<FStructProperty>(Property);
+					// Explicitly check against FInstancedPropertyBag, because that is what holds the Custom DetailsView.
+					return PropertyStruct && PropertyStruct->Struct == FInstancedPropertyBag::StaticStruct();
+				};
 
-			if (const FProperty* Property = PropertyNode->GetProperty())
-			{
-				const bool bPropertyIsPropertyBag = PropertyIsPropertyBag(Property);
-				bool bPropertyIsInPropertyBag = false;
-				// Traverse the parentage to see if the property is contained within the property bag.
-				for (const FPropertyNode* ParentNode = PropertyNode->GetParentNode(); ParentNode != nullptr; ParentNode = ParentNode->GetParentNode())
+				if (const FProperty* Property = PropertyNode->GetProperty())
 				{
-					if (PropertyIsPropertyBag(ParentNode->GetProperty()))
+					const bool bPropertyIsPropertyBag = PropertyIsPropertyBag(Property);
+					bool bPropertyIsInPropertyBag = false;
+					// Traverse the parentage to see if the property is contained within the property bag.
+					for (TSharedPtr<IPropertyHandle> ParentNode = PropertyNode->GetParentHandle(); ParentNode != nullptr; ParentNode = ParentNode->GetParentHandle())
 					{
-						bPropertyIsInPropertyBag = true;
-						break;
+						if (PropertyIsPropertyBag(ParentNode->GetProperty()))
+						{
+							bPropertyIsInPropertyBag = true;
+							break;
+						}
 					}
-				}
 
-				// Filter if the property isn't the property bag and the property is not inside the property bag.
-				return !(bPropertyIsPropertyBag || bPropertyIsInPropertyBag);
-			}
+					// Filter if the property isn't the property bag and the property is not inside the property bag.
+					return !(bPropertyIsPropertyBag || bPropertyIsInPropertyBag);
+				}
 
-			return true;
-		});
+				return true;
+			});
 	}
 
 	// Note: Single Property View Header notes that it doesn't work with arrays or structs, so using filters for now.
diff --git a/Engine/Plugins/Runtime/StateTree/Source/StateTreeEditorModule/Private/StateCentricView/View/STransitionNode.cpp b/Engine/Plugins/Runtime/StateTree/Source/StateTreeEditorModule/Private/StateCentricView/View/STransitionNode.cpp
--- a/Engine/Plugins/Runtime/StateTree/Source/StateTreeEditorModule/Private/StateCentricView/View/STransitionNode.cpp
+++ b/Engine/Plugins/Runtime/StateTree/Source/StateTreeEditorModule/Private/StateCentricView/View/STransitionNode.cpp
@@ -16,6 +16,7 @@
 #include "Modules/ModuleManager.h"
 #include "ObjectEditorUtils.h"
 #include "PropertyEditorModule.h"
+#include "PropertyHandle.h"
 #include "PropertyPathHelpers.h"
 #include "Styling/SlateBrush.h"
 #include "Styling/StyleColors.h"
@@ -24,9 +25,6 @@
 #include "Widgets/Layout/SBorder.h"
 #include "Widgets/Text/STextBlock.h"
 
-// Bad API, need to make not private
-#include "Editor/PropertyEditor/Private/PropertyNode.h"
-
 #define LOCTEXT_NAMESPACE "STransitionNode"
 
 namespace UE::StateTree::Editor::StateCentricView
@@ -357,11 +355,11 @@ TSharedRef<SWidget> STransitionNode::Construct_GenerateNodeBody(const EExtendabl
 			DetailsViewArgs.bAllowSearch = false;
 			DetailsViewArgs.bShowObjectLabel = false;
 			DetailsViewArgs.bShowHeaders = false;
-			DetailsViewArgs.ShouldForceHideProperty.BindLambda([ViewModel = ViewModel, TransitionID = TransitionID, StateOwningTransition = StateOwningTransition](const TSharedRef<FPropertyNode>& InNode)
+			DetailsViewArgs.ShouldForceHideProperty.BindLambda([ViewModel = ViewModel, TransitionID = TransitionID, StateOwningTransition = StateOwningTransition](const TSharedRef<IPropertyHandle>& InPropertyHandle)
 			{
 				const FStructProperty* TransitionPropertyStruct = nullptr;
 
-				if (const FStructProperty* PropertyStruct = CastField<FStructProperty>(InNode->GetProperty()))
+				if (const FStructProperty* PropertyStruct = CastField<FStructProperty>(InPropertyHandle->GetProperty()))
 				{
 					if (PropertyStruct->Struct == FStateTreeTransition::StaticStruct())
 					{
@@ -369,12 +367,12 @@ TSharedRef<SWidget> STransitionNode::Construct_GenerateNodeBody(const EExtendabl
 					}
 				}
 
-				const FPropertyNode* TransitionPropertyNode = TransitionPropertyStruct ? &InNode.Get() : InNode->GetParentNode();
+				TSharedPtr<IPropertyHandle> TransitionPropertyHandle = TransitionPropertyStruct ? InPropertyHandle : InPropertyHandle->GetParentHandle();
 
 				// Search parents for transition property
-				while (TransitionPropertyNode && !TransitionPropertyStruct)
+				while (TransitionPropertyHandle && !TransitionPropertyStruct)
 				{
-					if (const FProperty* ParentProperty = TransitionPropertyNode->GetProperty())
+					if (const FProperty* ParentProperty = TransitionPropertyHandle->GetProperty())
 					{
 						if (const FStructProperty* ParentPropertyStruct = CastField<FStructProperty>(ParentProperty))
 						{
@@ -387,14 +385,15 @@ TSharedRef<SWidget> STransitionNode::Construct_GenerateNodeBody(const EExtendabl
 
 					if (!TransitionPropertyStruct)
 					{
-						TransitionPropertyNode = TransitionPropertyNode->GetParentNode();
+						TransitionPropertyHandle = TransitionPropertyHandle->GetParentHandle();
 					}
 				}
 
-				if (TransitionPropertyStruct && TransitionPropertyNode)
+				if (TransitionPropertyStruct && TransitionPropertyHandle)
 				{
 					FStateTreeTransition TransitionValue; 
-					if (PropertyPathHelpers::GetPropertyValue(StateOwningTransition, TransitionPropertyNode->GetPropertyPath(), OUT TransitionValue))
+					FString TransitionPropertyPath(TransitionPropertyHandle->GetPropertyPath());
+					if (PropertyPathHelpers::GetPropertyValue(StateOwningTransition, TransitionPropertyPath, OUT TransitionValue))
 					{
 						return TransitionValue.ID != TransitionID;
 					}
diff --git a/Engine/Source/Editor/PropertyEditor/Private/DetailFilter.h b/Engine/Source/Editor/PropertyEditor/Private/DetailFilter.h
--- a/Engine/Source/Editor/PropertyEditor/Private/DetailFilter.h
+++ b/Engine/Source/Editor/PropertyEditor/Private/DetailFilter.h
@@ -63,6 +63,6 @@ struct FDetailFilter
 	TSet<FName> VisibleCategories;
 
 	/** Delegate that determines whether a property should be forced hidden - evaluated on panel refresh */
-	DECLARE_DELEGATE_RetVal_OneParam(bool, FShouldForceHideProperty, const TSharedRef<class FPropertyNode>&);
+	DECLARE_DELEGATE_RetVal_OneParam(bool, FShouldForceHideProperty, const TSharedRef<class IPropertyHandle>&);
 	FShouldForceHideProperty ShouldForceHideProperty;
 };
diff --git a/Engine/Source/Editor/PropertyEditor/Private/DetailItemNode.cpp b/Engine/Source/Editor/PropertyEditor/Private/DetailItemNode.cpp
--- a/Engine/Source/Editor/PropertyEditor/Private/DetailItemNode.cpp
+++ b/Engine/Source/Editor/PropertyEditor/Private/DetailItemNode.cpp
@@ -1014,10 +1014,15 @@ void FDetailItemNode::FilterNode(const FDetailFilter& InFilter)
 
 		if (!bForceHidden && InFilter.ShouldForceHideProperty.IsBound())
 		{
-			if (InFilter.ShouldForceHideProperty.Execute(PropertyNodePin.ToSharedRef()))
+			if (Customization.PropertyRow.IsValid())
 			{
-				bForceHidden = true;
+				TSharedPtr<IPropertyHandle> PropertyHandle = Customization.PropertyRow->GetPropertyHandle();
+				if (PropertyHandle.IsValid() && InFilter.ShouldForceHideProperty.Execute(PropertyHandle.ToSharedRef()))
+				{
+					bForceHidden = true;
+				}
 			}
+			
 		}
 
 		if (!bForceHidden)
diff --git a/Engine/Source/Editor/PropertyEditor/Public/DetailsViewArgs.h b/Engine/Source/Editor/PropertyEditor/Public/DetailsViewArgs.h
--- a/Engine/Source/Editor/PropertyEditor/Public/DetailsViewArgs.h
+++ b/Engine/Source/Editor/PropertyEditor/Public/DetailsViewArgs.h
@@ -15,7 +15,7 @@ class FUICommandList;
 class IClassViewerFilter;
 class FDetailsViewStyleKey;
 class FDetailsViewStyleKeys;
-class FPropertyNode;
+class IPropertyHandle;
 
 enum class EEditDefaultsOnlyNodeVisibility : uint8
 {
@@ -122,9 +122,9 @@ struct FDetailsViewArgs
 	bool bShowSectionSelector : 1;
 	/** If true, loose properties should be shown. */
 	bool bShowLooseProperties : 1;
-
+	
 	/** Delegate that determines whether a property should be forced hidden - evaluated on panel refresh */
-	DECLARE_DELEGATE_RetVal_OneParam(bool, FShouldForceHideProperty, const TSharedRef<FPropertyNode>&);
+	DECLARE_DELEGATE_RetVal_OneParam(bool, FShouldForceHideProperty, const TSharedRef<IPropertyHandle>&);
 	FShouldForceHideProperty ShouldForceHideProperty;
 
 	/** Resolve any UObject* to their IDO counterparts if they exist */
diff --git a/Engine/Source/Editor/PropertyEditor/Public/IPropertyTableCell.h b/Engine/Source/Editor/PropertyEditor/Public/IPropertyTableCell.h
--- a/Engine/Source/Editor/PropertyEditor/Public/IPropertyTableCell.h
+++ b/Engine/Source/Editor/PropertyEditor/Public/IPropertyTableCell.h
@@ -2,7 +2,6 @@
 #pragma once
 
 #include "CoreMinimal.h"
-#include "Editor/PropertyEditor/Private/PropertyNode.h"
 #include "PropertyHandle.h"
 
 class IPropertyTableCell
diff --git a/Engine/Source/Editor/StatsViewer/Private/ActorArrayHyperlinkColumn.cpp b/Engine/Source/Editor/StatsViewer/Private/ActorArrayHyperlinkColumn.cpp
--- a/Engine/Source/Editor/StatsViewer/Private/ActorArrayHyperlinkColumn.cpp
+++ b/Engine/Source/Editor/StatsViewer/Private/ActorArrayHyperlinkColumn.cpp
@@ -10,6 +10,7 @@
 #include "PropertyHandle.h"
 #include "IPropertyTable.h"
 #include "IPropertyTableCell.h"
+#include "PropertyPath.h"
 #include "ScopedTransaction.h"
 #include "Widgets/Input/SHyperlink.h"
 
diff --git a/Engine/Source/Editor/StatsViewer/Private/ObjectHyperlinkColumn.cpp b/Engine/Source/Editor/StatsViewer/Private/ObjectHyperlinkColumn.cpp
--- a/Engine/Source/Editor/StatsViewer/Private/ObjectHyperlinkColumn.cpp
+++ b/Engine/Source/Editor/StatsViewer/Private/ObjectHyperlinkColumn.cpp
@@ -11,6 +11,7 @@
 #include "Editor.h"
 #include "IPropertyTable.h"
 #include "IPropertyTableCell.h"
+#include "PropertyPath.h"
 #include "ScopedTransaction.h"
 #include "StatsViewerUtils.h"
 #include "Widgets/Input/SHyperlink.h"
diff --git a/Engine/Source/Editor/StatsViewer/Private/SStatsViewer.cpp b/Engine/Source/Editor/StatsViewer/Private/SStatsViewer.cpp
--- a/Engine/Source/Editor/StatsViewer/Private/SStatsViewer.cpp
+++ b/Engine/Source/Editor/StatsViewer/Private/SStatsViewer.cpp
@@ -24,6 +24,7 @@
 #include "ObjectHyperlinkColumn.h"
 #include "StatsCustomColumn.h"
 #include "ActorArrayHyperlinkColumn.h"
+#include "PropertyPath.h"
 #include "StatsViewerUtils.h"
 #include "Widgets/Input/SSearchBox.h"
 #include "Framework/Notifications/NotificationManager.h"
