Revert UAF ST Parameters not populated. Resulting in bindings not working.

#tests editor
#Synced-CL 49955823

diff --git a/Engine/Plugins/Runtime/StateTree/Source/StateTreeEditorModule/Private/StateTreeCompiler.cpp b/Engine/Plugins/Runtime/StateTree/Source/StateTreeEditorModule/Private/StateTreeCompiler.cpp
--- a/Engine/Plugins/Runtime/StateTree/Source/StateTreeEditorModule/Private/StateTreeCompiler.cpp
+++ b/Engine/Plugins/Runtime/StateTree/Source/StateTreeEditorModule/Private/StateTreeCompiler.cpp
@@ -724,7 +724,7 @@ bool FStateTreeCompiler::CreateParameters()
 	}
 	else
 	{
-		StateTree->Parameters.Reset();
+		StateTree->Parameters = EditorData->GetRootParametersPropertyBag();
 		StateTree->GlobalParameters = nullptr;
 		StateTree->GlobalParametersStructType = nullptr;
 	}
