Hide non-enter condition properties in state centric
Make RigVM functions / graph tasks in ST show the funciton name.
Add remove button to transitions in state centric
Add API for removing transitions to ST ViewModel

#tests editor

#Synced-CL 50015298

diff --git a/Engine/Plugins/Experimental/UAF/UAFStateTree/Source/UAFStateTree/Internal/Tasks/AnimNextStateTreeGraphInstanceTask.h b/Engine/Plugins/Experimental/UAF/UAFStateTree/Source/UAFStateTree/Internal/Tasks/AnimNextStateTreeGraphInstanceTask.h
--- a/Engine/Plugins/Experimental/UAF/UAFStateTree/Source/UAFStateTree/Internal/Tasks/AnimNextStateTreeGraphInstanceTask.h
+++ b/Engine/Plugins/Experimental/UAF/UAFStateTree/Source/UAFStateTree/Internal/Tasks/AnimNextStateTreeGraphInstanceTask.h
@@ -122,6 +122,7 @@ protected:
 	virtual EStateTreeRunStatus Tick(FStateTreeExecutionContext& Context, const float DeltaTime) const override;
 	virtual void ExitState(FStateTreeExecutionContext& Context, const FStateTreeTransitionResult& Transition) const override;
 #if WITH_EDITOR
+	virtual FText GetDescription(const FGuid& ID, FStateTreeDataView InstanceDataView, const IStateTreeBindingLookup& BindingLookup, EStateTreeNodeFormatting Formatting = EStateTreeNodeFormatting::Text) const override;
 	virtual void GetObjectReferences(TArray<const UObject*>& OutReferencedObjects, const FStateTreeDataView InstanceDataView) const override;
 #endif
 public:
diff --git a/Engine/Plugins/Experimental/UAF/UAFStateTree/Source/UAFStateTree/Private/Conditions/AnimNextStateTreeRigVMConditionBase.cpp b/Engine/Plugins/Experimental/UAF/UAFStateTree/Source/UAFStateTree/Private/Conditions/AnimNextStateTreeRigVMConditionBase.cpp
--- a/Engine/Plugins/Experimental/UAF/UAFStateTree/Source/UAFStateTree/Private/Conditions/AnimNextStateTreeRigVMConditionBase.cpp
+++ b/Engine/Plugins/Experimental/UAF/UAFStateTree/Source/UAFStateTree/Private/Conditions/AnimNextStateTreeRigVMConditionBase.cpp
@@ -78,7 +78,11 @@ bool FAnimNextStateTreeRigVMConditionBase::TestCondition(FStateTreeExecutionCont
 #if WITH_EDITOR
 FText FAnimNextStateTreeRigVMConditionBase::GetDescription(const FGuid& ID, FStateTreeDataView InstanceDataView, const IStateTreeBindingLookup& BindingLookup, EStateTreeNodeFormatting Formatting) const
 {
-	return LOCTEXT("AnimNextStateTreeConditon_Desc", "RigVM function driven condition");
+	if (CallFunctionInfo.Function.IsValid())
+	{
+		return FText::Format(LOCTEXT("UAF_StateTreeConditon_Desc", "Function ({0})"), FText::FromName(CallFunctionInfo.Function));
+	}
+	return LOCTEXT("UAF_StateTreeConditon_FunctionMissingDesc", "Function (NONE)");
 }
 
 void FAnimNextStateTreeRigVMConditionBase::PostEditNodeChangeChainProperty(const FPropertyChangedChainEvent& PropertyChangedEvent, FStateTreeDataView InstanceDataView)
diff --git a/Engine/Plugins/Experimental/UAF/UAFStateTree/Source/UAFStateTree/Private/Tasks/AnimNextStateTreeGraphInstanceTask.cpp b/Engine/Plugins/Experimental/UAF/UAFStateTree/Source/UAFStateTree/Private/Tasks/AnimNextStateTreeGraphInstanceTask.cpp
--- a/Engine/Plugins/Experimental/UAF/UAFStateTree/Source/UAFStateTree/Private/Tasks/AnimNextStateTreeGraphInstanceTask.cpp
+++ b/Engine/Plugins/Experimental/UAF/UAFStateTree/Source/UAFStateTree/Private/Tasks/AnimNextStateTreeGraphInstanceTask.cpp
@@ -12,6 +12,8 @@
 #include "UAF/UAFAssetFactory.h"
 #include "UObject/FortniteMainBranchObjectVersion.h"
 
+#define LOCTEXT_NAMESPACE "AnimNextStateTreeGraphInstanceTask"
+
 #if WITH_EDITORONLY_DATA
 bool FAnimNextGraphInstanceTaskInstanceData::Serialize(FArchive& Ar)
 {
@@ -214,6 +216,19 @@ void FAnimNextStateTreeGraphInstanceTask::ExitState(FStateTreeExecutionContext&
 
 #if WITH_EDITOR
 
+FText FAnimNextStateTreeGraphInstanceTask::GetDescription(const FGuid& ID, FStateTreeDataView InstanceDataView, const IStateTreeBindingLookup& BindingLookup, EStateTreeNodeFormatting Formatting) const
+{
+	if (InstanceDataView.IsValid())
+	{
+		const UE::UAF::FAssetHandle& AssetHandle = InstanceDataView.Get<FAnimNextGraphInstanceTaskInstanceData>().AssetData;
+		if (AssetHandle.IsValid())
+		{
+			return FText::Format(LOCTEXT("UAF_Task_Desc", "UAF Task ({0})"), FText::FromString(AssetHandle.Get().GetName()));
+		}
+	}
+	return LOCTEXT("UAF_Task_AssetMissingDesc", "UAF Task (NONE)");
+}
+
 void FAnimNextStateTreeGraphInstanceTask::GetObjectReferences(TArray<const UObject*>& OutReferencedObjects, FStateTreeDataView InstanceDataView) const
 {
 	if (InstanceDataView.IsValid())
@@ -221,9 +236,10 @@ void FAnimNextStateTreeGraphInstanceTask::GetObjectReferences(TArray<const UObje
 		const UE::UAF::FAssetHandle& AssetHandle = InstanceDataView.Get<FAnimNextGraphInstanceTaskInstanceData>().AssetData;
 		if (AssetHandle.IsValid())
 		{
-			InstanceDataView.Get<FAnimNextGraphInstanceTaskInstanceData>().AssetData.Get().GetObjectReferences(OutReferencedObjects);
+			AssetHandle.Get().GetObjectReferences(OutReferencedObjects);
 		}
 	}
 }
 #endif
 
+#undef LOCTEXT_NAMESPACE
\ No newline at end of file
diff --git a/Engine/Plugins/Experimental/UAF/UAFStateTree/Source/UAFStateTree/Private/Tasks/AnimNextStateTreeRigVMTaskBase.cpp b/Engine/Plugins/Experimental/UAF/UAFStateTree/Source/UAFStateTree/Private/Tasks/AnimNextStateTreeRigVMTaskBase.cpp
--- a/Engine/Plugins/Experimental/UAF/UAFStateTree/Source/UAFStateTree/Private/Tasks/AnimNextStateTreeRigVMTaskBase.cpp
+++ b/Engine/Plugins/Experimental/UAF/UAFStateTree/Source/UAFStateTree/Private/Tasks/AnimNextStateTreeRigVMTaskBase.cpp
@@ -83,7 +83,11 @@ EStateTreeRunStatus FAnimNextStateTreeRigVMTaskBase::EnterState(FStateTreeExecut
 #if WITH_EDITOR
 FText FAnimNextStateTreeRigVMTaskBase::GetDescription(const FGuid& ID, FStateTreeDataView InstanceDataView, const IStateTreeBindingLookup& BindingLookup, EStateTreeNodeFormatting Formatting) const
 {
-	return LOCTEXT("AnimNextStateTreeConditon_Desc", "RigVM function driven Task");
+	if (CallFunctionInfo.Function.IsValid())
+	{
+		return FText::Format(LOCTEXT("UAF_StateTreeTask_Desc", "Function ({0})"), FText::FromName(CallFunctionInfo.Function));
+	}
+	return LOCTEXT("UAF_StateTreeTask_FunctionMissingDesc", "Function (NONE)");
 }
 
 void FAnimNextStateTreeRigVMTaskBase::PostEditNodeChangeChainProperty(const FPropertyChangedChainEvent& PropertyChangedEvent, FStateTreeDataView InstanceDataView)
diff --git a/Engine/Plugins/Runtime/StateTree/Source/StateTreeEditorModule/Private/StateCentricView/View/SParentStateExtensionNode.cpp b/Engine/Plugins/Runtime/StateTree/Source/StateTreeEditorModule/Private/StateCentricView/View/SParentStateExtensionNode.cpp
--- a/Engine/Plugins/Runtime/StateTree/Source/StateTreeEditorModule/Private/StateCentricView/View/SParentStateExtensionNode.cpp
+++ b/Engine/Plugins/Runtime/StateTree/Source/StateTreeEditorModule/Private/StateCentricView/View/SParentStateExtensionNode.cpp
@@ -272,6 +272,18 @@ TSharedRef<SWidget> SParentStateExtensionNode::Construct_GenerateNodeBody(const
 			DetailsViewArgs.bShowScrollBar = false;
 			DetailsViewArgs.bShowOptions = false;
 			DetailsViewArgs.bAllowSearch = false;
+			DetailsViewArgs.ShouldForceHideProperty.BindLambda([ViewModel = ViewModel](const TSharedRef<IPropertyHandle>& InPropertyHandle)
+			{
+				// Hide properties other than enter conditions in the enter condition area
+				if (InPropertyHandle->GetProperty()->GetFName() == GET_MEMBER_NAME_CHECKED(UStateTreeState, bCheckPrerequisitesWhenActivatingChildDirectly)
+					|| InPropertyHandle->GetProperty()->GetFName() == GET_MEMBER_NAME_CHECKED(UStateTreeState, bHasRequiredEventToEnter)
+					|| InPropertyHandle->GetProperty()->GetFName() == GET_MEMBER_NAME_CHECKED(UStateTreeState, RequiredEventToEnter))
+				{
+					return true;
+				}
+
+				return false;
+			});
 
 			TSharedRef<IDetailsView> StateDetailsView = PropertyEditorModule.CreateDetailView(DetailsViewArgs);
 			StateDetailsView->SetExtensionHandler(MakeShared<FStateTreeBindingExtension>());
diff --git a/Engine/Plugins/Runtime/StateTree/Source/StateTreeEditorModule/Private/StateCentricView/View/SStateCentricViewMainStateNode.cpp b/Engine/Plugins/Runtime/StateTree/Source/StateTreeEditorModule/Private/StateCentricView/View/SStateCentricViewMainStateNode.cpp
--- a/Engine/Plugins/Runtime/StateTree/Source/StateTreeEditorModule/Private/StateCentricView/View/SStateCentricViewMainStateNode.cpp
+++ b/Engine/Plugins/Runtime/StateTree/Source/StateTreeEditorModule/Private/StateCentricView/View/SStateCentricViewMainStateNode.cpp
@@ -277,6 +277,18 @@ TSharedRef<SWidget> SStateCentricViewMainStateNode::Construct_GenerateNodeBody(c
 		DetailsViewArgs.bShowScrollBar = false;
 		DetailsViewArgs.bShowOptions = false;
 		DetailsViewArgs.bAllowSearch = false;
+		DetailsViewArgs.ShouldForceHideProperty.BindLambda([ViewModel = ViewModel](const TSharedRef<IPropertyHandle>& InPropertyHandle)
+		{
+			// Hide properties other than enter conditions in the enter condition area
+			if (InPropertyHandle->GetProperty()->GetFName() == GET_MEMBER_NAME_CHECKED(UStateTreeState, bCheckPrerequisitesWhenActivatingChildDirectly)
+				|| InPropertyHandle->GetProperty()->GetFName() == GET_MEMBER_NAME_CHECKED(UStateTreeState, bHasRequiredEventToEnter)
+				|| InPropertyHandle->GetProperty()->GetFName() == GET_MEMBER_NAME_CHECKED(UStateTreeState, RequiredEventToEnter))
+			{
+				return true;
+			}
+
+			return false;
+		});
 
 		TSharedRef<IDetailsView> StateDetailsView = PropertyEditorModule.CreateDetailView(DetailsViewArgs);
 		StateDetailsView->SetExtensionHandler(MakeShared<FStateTreeBindingExtension>());
diff --git a/Engine/Plugins/Runtime/StateTree/Source/StateTreeEditorModule/Private/StateCentricView/View/STransitionNode.cpp b/Engine/Plugins/Runtime/StateTree/Source/StateTreeEditorModule/Private/StateCentricView/View/STransitionNode.cpp
--- a/Engine/Plugins/Runtime/StateTree/Source/StateTreeEditorModule/Private/StateCentricView/View/STransitionNode.cpp
+++ b/Engine/Plugins/Runtime/StateTree/Source/StateTreeEditorModule/Private/StateCentricView/View/STransitionNode.cpp
@@ -252,8 +252,9 @@ TSharedRef<SWidget> STransitionNode::Construct_GenerateNodeHeader(const EExtenda
 			TransitionStateBrush = FStateCentricViewUtils::MakeStateSubnodeBackgroundOutlined(ViewModel.ToSharedRef(), State);
 
 			// @TODO move this logic below into FStateCentricViewUtils
+			TSharedPtr<SHorizontalBox> NodeHeaderMainHorizontalBox = nullptr;
 
-			return SNew(SOverlay)
+			TSharedPtr<SOverlay> NodeHeaderOverlay = SNew(SOverlay)
 				+ SOverlay::Slot()
 				[
 					SNew(SImage)
@@ -266,12 +267,13 @@ TSharedRef<SWidget> STransitionNode::Construct_GenerateNodeHeader(const EExtenda
 					.Image(TransitionStateBrush->GetSlateBrush())
 				]
 				+ SOverlay::Slot()
+				.Padding(FMargin(0.0, 0.0, InLOD == EExtendableNodeLOD::Full ? 26.0 : 0, 0.0))
 				[
 					FStateCentricViewUtils::GenerateGotoStateButton(InLOD, ViewModel.ToSharedRef(), State)
 				]
 				+ SOverlay::Slot()
 				[
-					SNew(SHorizontalBox)
+					SAssignNew(NodeHeaderMainHorizontalBox, SHorizontalBox)
 					+ SHorizontalBox::Slot()
 					.HAlign(HAlign_Left)
 					.VAlign(VAlign_Center)
@@ -295,6 +297,28 @@ TSharedRef<SWidget> STransitionNode::Construct_GenerateNodeHeader(const EExtenda
 						.AutoWrapText(UStateCentricViewSettings::GetSchemaSettingsDefaultFallback(ViewModel.Get()).bWordWrapTransitionStates)
 					]
 				];
+
+			if (InLOD == EExtendableNodeLOD::Full)
+			{
+				NodeHeaderMainHorizontalBox->AddSlot()
+				.HAlign(HAlign_Right)
+				.VAlign(VAlign_Center)
+				.Padding(FMargin(0.0, 0.0, 6.0, 0.0))
+				[
+					SAssignNew(DeleteTransitionButton, SButton)
+					.OnClicked(this, &STransitionNode::HandleOnDeleteTransitionClicked)
+					.ButtonStyle(&FAppStyle::GetWidgetStyle<FButtonStyle>("NoBorder"))
+					.ContentPadding(FMargin())
+					.ToolTipText(LOCTEXT("TransitionNodeDelete", "Delete this Transition"))
+					[
+						SNew(SImage)
+						.Image(FAppStyle::GetBrush("Icons.X"))
+						.ColorAndOpacity(this, &STransitionNode::GetDeleteTransitionColor)
+					]
+				];
+			}
+
+			return NodeHeaderOverlay.ToSharedRef();
 		}
 		else
 		{
@@ -436,6 +460,25 @@ void STransitionNode::HandleOnTransitionModified(const UStateTreeState* OwningSt
 	}
 }
 
+FReply STransitionNode::HandleOnDeleteTransitionClicked()
+{
+	FStateTreeTransition* Transition = ViewModel->GetMutableTransitionByID(TransitionID);
+	if (Transition && ViewState.IsValid())
+	{
+		if (const UStateTreeState* State = ViewModel->GetSourceStateFromTransition(Transition))
+		{
+			ViewModel->RemoveTransition(const_cast<UStateTreeState*>(State), TransitionID);
+		}
+	}
+
+	return FReply::Handled();
+}
+
+FSlateColor STransitionNode::GetDeleteTransitionColor() const
+{
+	return DeleteTransitionButton->IsHovered() ? FStyleColors::White : FStyleColors::Foreground;
+}
+
 
 //////////////////////////////////////////////////////////////////////////
 // SInTransitionNode
diff --git a/Engine/Plugins/Runtime/StateTree/Source/StateTreeEditorModule/Private/StateCentricView/View/STransitionNode.h b/Engine/Plugins/Runtime/StateTree/Source/StateTreeEditorModule/Private/StateCentricView/View/STransitionNode.h
--- a/Engine/Plugins/Runtime/StateTree/Source/StateTreeEditorModule/Private/StateCentricView/View/STransitionNode.h
+++ b/Engine/Plugins/Runtime/StateTree/Source/StateTreeEditorModule/Private/StateCentricView/View/STransitionNode.h
@@ -50,6 +50,12 @@ protected:
 	/** Callback for transtion modified */
 	void HandleOnTransitionModified(const UStateTreeState* OwningState, const FGuid& InTransitionID);
 
+	/** Callback for delete transition button */
+	FReply HandleOnDeleteTransitionClicked();
+
+	/** Minor UI / UX callbacks */
+	FSlateColor GetDeleteTransitionColor() const;
+
 protected:
 
 	/** 
@@ -68,6 +74,7 @@ private:
 	TSharedPtr<FDeferredCleanupSlateBrush> TransitionStateBackgroundBrush;
 
 	TSharedPtr<IStructureDetailsView> TransitionDetailsView;
+	TSharedPtr<SButton> DeleteTransitionButton;
 };
 
 
diff --git a/Engine/Plugins/Runtime/StateTree/Source/StateTreeEditorModule/Private/StateTreeViewModel.cpp b/Engine/Plugins/Runtime/StateTree/Source/StateTreeEditorModule/Private/StateTreeViewModel.cpp
--- a/Engine/Plugins/Runtime/StateTree/Source/StateTreeEditorModule/Private/StateTreeViewModel.cpp
+++ b/Engine/Plugins/Runtime/StateTree/Source/StateTreeEditorModule/Private/StateTreeViewModel.cpp
@@ -2269,6 +2269,12 @@ FStateTreeTransition& FStateTreeViewModel::AddTransition(TNotNull<UStateTreeStat
 	return *Result;
 }
 
+void FStateTreeViewModel::RemoveTransition(TNotNull<UStateTreeState*> State, const FGuid TransitionID)
+{
+	// Same as delete node. This method exists to make API more intuitive
+	DeleteNode(State, TransitionID);
+}
+
 FLinearColor FStateTreeViewModel::GetStateColor(TNotNull<const UStateTreeState*> State, EStateTreeViewModelColorAdjustment ColorAdjustment) const
 {
 	if (bEnableSimpleHierarchyDrivenColors)
diff --git a/Engine/Plugins/Runtime/StateTree/Source/StateTreeEditorModule/Public/StateTreeViewModel.h b/Engine/Plugins/Runtime/StateTree/Source/StateTreeEditorModule/Public/StateTreeViewModel.h
--- a/Engine/Plugins/Runtime/StateTree/Source/StateTreeEditorModule/Public/StateTreeViewModel.h
+++ b/Engine/Plugins/Runtime/StateTree/Source/StateTreeEditorModule/Public/StateTreeViewModel.h
@@ -170,6 +170,7 @@ public:
 	// Transition Manipulation
 	UE_API FStateTreeTransition& AddTransition(TNotNull<UStateTreeState*> State, const EStateTreeTransitionTrigger InTrigger, const EStateTreeTransitionType InType, const UStateTreeState* InState = nullptr) const;
 	UE_API FStateTreeTransition& AddTransition(TNotNull<UStateTreeState*> State, const EStateTreeTransitionTrigger InTrigger, const FGameplayTag InEventTag, const EStateTreeTransitionType InType, const UStateTreeState* InState = nullptr) const;
+	UE_API void RemoveTransition(TNotNull<UStateTreeState*> State, const FGuid TransitionID);
 
 	// Color manipulation
 	UE_API FLinearColor GetStateColor(TNotNull<const UStateTreeState*> State, EStateTreeViewModelColorAdjustment ColorAdjustment = EStateTreeViewModelColorAdjustment::Normal) const;
