[RewindDebugger] RecordingStarted is now called immediately once the connection gets established (not deferred to game thread) since the channels need to be enabled before the next game thread frame to not miss events. The only part that is still postpone to the gamethread is the TRACE_WORLD and TRACE_PAWN_POSSESS calls relying on TActorIterator which can be used outside of game thread.
Fixed StateTree runtime extension to handle being called from any thread.

#Synced-CL 50059331

diff --git a/Engine/Plugins/Animation/GameplayInsights/Source/RewindDebuggerRuntime/Private/RewindDebuggerRuntime.cpp b/Engine/Plugins/Animation/GameplayInsights/Source/RewindDebuggerRuntime/Private/RewindDebuggerRuntime.cpp
--- a/Engine/Plugins/Animation/GameplayInsights/Source/RewindDebuggerRuntime/Private/RewindDebuggerRuntime.cpp
+++ b/Engine/Plugins/Animation/GameplayInsights/Source/RewindDebuggerRuntime/Private/RewindDebuggerRuntime.cpp
@@ -1,4 +1,4 @@
-﻿// Copyright Epic Games, Inc. All Rights Reserved.
+// Copyright Epic Games, Inc. All Rights Reserved.
 
 #if WITH_TRACE_BASED_DEBUGGERS
 
@@ -83,11 +83,7 @@ namespace RewindDebugger
 
 	void FRewindDebuggerRuntime::OnRecordingStartedInternal()
 	{
-		// Previous implementation calling the following code was executed on the game thread
-		// after calling FTraceAuxiliary::Start, but we are now waiting for the connection
-		// to be ready (i.e. FTraceAuxiliary::OnConnection) which comes from a worker thread.
-		// Extensions code in 'RecordingStarted' might require to run on the Game thread so this
-		// is why we defer the callbacks.
+		// TObjectIterator requires execution on the game thread
 		ExecuteOnGameThread(UE_SOURCE_LOCATION, [this]
 			{
 				// trace each play-in-editor world, and all the actors in it.
@@ -103,13 +99,12 @@ namespace RewindDebugger
 						}
 					}
 				}
+			});
 
-				// update extensions
-				IterateExtensions([](IRewindDebuggerRuntimeExtension* Extension)
-					{
-						Extension->RecordingStarted();
-					}
-				);
+		// update extensions so they can enable their channels now that the trace connection is established
+		IterateExtensions([](IRewindDebuggerRuntimeExtension* Extension)
+			{
+				Extension->RecordingStarted();
 			});
 
 		RecordingStarted.Broadcast();
diff --git a/Engine/Plugins/Runtime/StateTree/Source/StateTreeDeveloper/Private/StateTreeRewindDebuggerRecordingExtension.cpp b/Engine/Plugins/Runtime/StateTree/Source/StateTreeDeveloper/Private/StateTreeRewindDebuggerRecordingExtension.cpp
--- a/Engine/Plugins/Runtime/StateTree/Source/StateTreeDeveloper/Private/StateTreeRewindDebuggerRecordingExtension.cpp
+++ b/Engine/Plugins/Runtime/StateTree/Source/StateTreeDeveloper/Private/StateTreeRewindDebuggerRecordingExtension.cpp
@@ -14,10 +14,13 @@ namespace UE::StateTreeDebugger
 //----------------------------------------------------------------------//
 void FRewindDebuggerRecordingExtension::RecordingStarted()
 {
+	// Broadcast delegate before enabling the channel so event can't be traced until
+	// the delegate listeners are notified.
+	StateTree::Delegates::OnTracingStateChanged.Broadcast(EStateTreeTraceStatus::TracesStarted);
+
 	if (!UE_TRACE_CHANNELEXPR_IS_ENABLED(StateTreeDebugChannel))
 	{
 		Trace::ToggleChannel(TEXT("StateTreeDebugChannel"), true);
-		StateTree::Delegates::OnTracingStateChanged.Broadcast(EStateTreeTraceStatus::TracesStarted);
 	}
 }
 
diff --git a/Engine/Plugins/Runtime/StateTree/Source/StateTreeModule/Private/Debugger/StateTreeTrace.cpp b/Engine/Plugins/Runtime/StateTree/Source/StateTreeModule/Private/Debugger/StateTreeTrace.cpp
--- a/Engine/Plugins/Runtime/StateTree/Source/StateTreeModule/Private/Debugger/StateTreeTrace.cpp
+++ b/Engine/Plugins/Runtime/StateTree/Source/StateTreeModule/Private/Debugger/StateTreeTrace.cpp
@@ -207,9 +207,6 @@ struct FBufferedDataList
 	/** Traces required events (world time, phase events, etc.) for a given instance. */
 	void Flush_AnyThread(const FStateTreeInstanceDebugId InstanceId);
 
-	/** Traces buffered events, if needed, when traces are enabled. */
-	void OnTracesStarted_GameThread();
-
 	/** Traces closing events and resets some data for subsequent traces. */
 	void OnStoppingTraces_GameThread();
 
@@ -225,11 +222,15 @@ struct FBufferedDataList
 	/** Converts the provided list of execution frames to traceable data format and outputs it to the trace. */
 	void OutputActiveStates_AnyThread(const FStateTreeInstanceDebugId InstanceId, const TConstArrayView<FStateTreeExecutionFrame> ActiveFrames);
 
+	/** Initializes recording world time to 0 if not initialized. */
+	void InitializeWorldTimeIfNeeded_AnyThread();
+
 	/** Sets the world time that might need to be traced by worker threads. */
 	void SetWorldTime_GameThread(const double WorldTime);
 
-	void Reset()
+	void Reset_GameThread()
 	{
+		checkf(IsInGameThread(), TEXT("Expecting to only be called on the main thread before worker threads trace events."));
 		{
 			FRWScopeLock RWScopeLock(TracedInstancesRWLock, SLT_Write);
 			TracedInstances.Empty();
@@ -263,7 +264,7 @@ private:
 	static constexpr double UNINITIALIZED_WORLD_TIME = -1;
 
 	/** World time provided by the game thread at the beginning of the frame */
-	double RecordingWorldTime = UNINITIALIZED_WORLD_TIME;
+	std::atomic<double> RecordingWorldTime = UNINITIALIZED_WORLD_TIME;
 
 	/** World time traced only once by the first worker thread that needs to trace an event */
 	std::atomic<double> TracedRecordingWorldTime = UNINITIALIZED_WORLD_TIME;
@@ -335,13 +336,13 @@ void FBufferedDataList::Flush_AnyThread(const FStateTreeInstanceDebugId Instance
 
 void FBufferedDataList::TraceWorldTimeIfNeeded_AnyThread()
 {
-	checkf(RecordingWorldTime != UNINITIALIZED_WORLD_TIME
-		, TEXT("Expecting world time to always be set at the beginning of the frame before any worker thread attempts to trace."));
+	InitializeWorldTimeIfNeeded_AnyThread();
+
 	double Expected = UNINITIALIZED_WORLD_TIME;
 	if (TracedRecordingWorldTime.compare_exchange_strong(Expected, RecordingWorldTime))
 	{
 		UE_TRACE_LOG(StateTreeDebugger, WorldTimestampEvent, StateTreeDebugChannel)
-			<< WorldTimestampEvent.WorldTime(RecordingWorldTime);
+			<< WorldTimestampEvent.WorldTime(RecordingWorldTime.load());
 	}
 }
 
@@ -390,20 +391,6 @@ void FBufferedDataList::TraceStackedPhases_AnyThread(const FStateTreeInstanceDeb
 	}
 }
 
-void FBufferedDataList::OnTracesStarted_GameThread()
-{
-	checkf(IsInGameThread(), TEXT("Expecting to only be called by the statetree delegate on the main thread before worker threads trace events."));
-
-	// Delegate can be received before first call to SetWorldTime_GameThread()
-	// In that case we set the time to 0 for any pending lifetime events
-	if (RecordingWorldTime == UNINITIALIZED_WORLD_TIME)
-	{
-		SetWorldTime_GameThread(0);
-	}
-
-	TraceWorldTimeIfNeeded_AnyThread();
-}
-
 void FBufferedDataList::OnStoppingTraces_GameThread()
 {
 	checkf(IsInGameThread(), TEXT("Expecting to only be called by the statetree delegate on the main thread after worker threads traced events."));
@@ -507,6 +494,13 @@ void FBufferedDataList::SetWorldTime_GameThread(const double WorldTime)
 	TracedRecordingWorldTime = UNINITIALIZED_WORLD_TIME;
 }
 
+void FBufferedDataList::InitializeWorldTimeIfNeeded_AnyThread()
+{
+	constexpr double InitialValue = 0;
+	double Expected = UNINITIALIZED_WORLD_TIME;
+	RecordingWorldTime.compare_exchange_strong(Expected, InitialValue);
+}
+
 /**
  * Pushed or pops an entry on the Phase stack for a given Instance.
  * Will send the Pop events for phases popped if their associated Push events were sent.
@@ -648,19 +642,14 @@ void RegisterGlobalDelegates()
 
 	GTracingStateChangedDelegateHandle = UE::StateTree::Delegates::OnTracingStateChanged.AddLambda([](const EStateTreeTraceStatus TraceStatus)
 	{
-		if (TraceStatus == EStateTreeTraceStatus::TracesStarted)
-		{
-			GBufferedEvents.Reset();
-			GBufferedEvents.OnTracesStarted_GameThread();
-		}
 		// Traces are about to be stopped so allow the buffered events to react.
-		else if (TraceStatus == EStateTreeTraceStatus::StoppingTrace)
+		if (TraceStatus == EStateTreeTraceStatus::StoppingTrace)
 		{
 			GBufferedEvents.OnStoppingTraces_GameThread();
 		}
 		else if (TraceStatus == EStateTreeTraceStatus::Cleared)
 		{
-			GBufferedEvents.Reset();
+			GBufferedEvents.Reset_GameThread();
 		}
 	});
 }
diff --git a/Engine/Plugins/Runtime/StateTree/Source/StateTreeModule/Private/StateTreeModule.cpp b/Engine/Plugins/Runtime/StateTree/Source/StateTreeModule/Private/StateTreeModule.cpp
--- a/Engine/Plugins/Runtime/StateTree/Source/StateTreeModule/Private/StateTreeModule.cpp
+++ b/Engine/Plugins/Runtime/StateTree/Source/StateTreeModule/Private/StateTreeModule.cpp
@@ -214,9 +214,6 @@ bool FStateTreeModule::StartTraces(int32& OutTraceId)
 		, nullptr);
 	}
 
-	UE::Trace::ToggleChannel(TEXT("StateTreeDebugChannel"), true);
-	UE::Trace::ToggleChannel(TEXT("FrameChannel"), true);
-
 	bool bAreTracesStarted = false;
 	if (bAlreadyConnected == false)
 	{
@@ -231,6 +228,14 @@ bool FStateTreeModule::StartTraces(int32& OutTraceId)
 		UE::StateTree::Delegates::OnTracingStateChanged.Broadcast(EStateTreeTraceStatus::TracesStarted);
 	}
 
+	// Enable channels after starting the trace and notifying listeners so events can't
+	// be traced (on other threads) before this point.
+	// This might not be required most of the time in this case (i.e., legacy StateTreeDebugger)
+	// which is happening from the game thread but this is to be consistent with the RewindDebugger extension
+	// that could enable channels from a worker thread.
+	UE::Trace::ToggleChannel(TEXT("StateTreeDebugChannel"), true);
+	UE::Trace::ToggleChannel(TEXT("FrameChannel"), true);
+
 	return bAreTracesStarted;
 #else
 	return false;
diff --git a/Engine/Source/Runtime/RewindDebuggerRuntimeInterface/Public/RewindDebuggerRuntimeInterface/IRewindDebuggerRuntimeExtension.h b/Engine/Source/Runtime/RewindDebuggerRuntimeInterface/Public/RewindDebuggerRuntimeInterface/IRewindDebuggerRuntimeExtension.h
--- a/Engine/Source/Runtime/RewindDebuggerRuntimeInterface/Public/RewindDebuggerRuntimeInterface/IRewindDebuggerRuntimeExtension.h
+++ b/Engine/Source/Runtime/RewindDebuggerRuntimeInterface/Public/RewindDebuggerRuntimeInterface/IRewindDebuggerRuntimeExtension.h
@@ -42,17 +42,21 @@ public:
 
 	static UE_API const FName ModularFeatureName;
 
-	/** called when recording has started */
+	/**
+	 * This method is called once the trace connection is established (i.e., FTraceAuxiliary::OnConnection delegate) which can be
+	 * from any thread.
+	 * This is the method to override to enable trace channels specific to the extension.
+	 */
 	virtual void RecordingStarted()
 	{
 	}
 
-	/** called when recording has ended */
+	/** This method is called from the game thread when recording gets stopped */
 	virtual void RecordingStopped()
 	{
 	}
 
-	/** called when recording is unloaded or a new one starts */
+	/** This method is called from the game thread before executing the "start recording" command (i.e., FTraceAuxiliary::Start/Relay methods) */
 	virtual void Clear()
 	{
 	}
