[StateTree] Assign the GlobalParameters value before updating the UDS. The UDS compiler will trigger callbacks, and the UX code needs those set.
Because a CL was backout, try to find the UDS before creating a new one. The UDS could be saved in the package, then the backout occurs, and the asset is resaved again without the StateTree reference to the UDS.
#Synced-CL 49995628
#Synced-CL 49995628

diff --git a/Engine/Plugins/Runtime/StateTree/Source/StateTreeEditorModule/Private/StateTreeCompiler.cpp b/Engine/Plugins/Runtime/StateTree/Source/StateTreeEditorModule/Private/StateTreeCompiler.cpp
--- a/Engine/Plugins/Runtime/StateTree/Source/StateTreeEditorModule/Private/StateTreeCompiler.cpp
+++ b/Engine/Plugins/Runtime/StateTree/Source/StateTreeEditorModule/Private/StateTreeCompiler.cpp
@@ -65,7 +65,7 @@ namespace UE::StateTree::Compiler
 	// Version used by AppendToClassSchema. Bump this version number when a recompile of all assets is needed.
 	//Only bump when only the compiler changes.
 	//You do not need to bump the version when the memory layout of UStateTree has changed.
-	static constexpr int32 ClassSchemaVersion = 1;
+	static constexpr int32 ClassSchemaVersion = 2;
 
 	FAutoConsoleVariable CVarLogEnableParameterDelegateDispatcherBinding(
 		TEXT("StateTree.Compiler.EnableParameterDelegateDispatcherBinding"),
@@ -690,9 +690,10 @@ bool FStateTreeCompiler::CreateParameters()
 
 		// Copy parameters from EditorData
 		const FInstancedPropertyBag& PropertyBag = EditorData->GetRootParametersPropertyBag();
-		if (UUserDefinedStruct* UDS = Cast<UUserDefinedStruct>(StateTree->GlobalParametersStructType))
+		if (StateTree->GlobalParametersStructType)
 		{
-			UE::StructUtils::UpdateUserDefinedStructFromDescs(UDS, PropertyBag);
+			StateTree->GlobalParameters = StateTree->GlobalParametersStructType;
+			UE::StructUtils::UpdateUserDefinedStructFromDescs(StateTree->GlobalParametersStructType, PropertyBag);
 		}
 		else
 		{
@@ -704,26 +705,24 @@ bool FStateTreeCompiler::CreateParameters()
 			Args.bAddPropertyIDToPropertyName = false;
 			TValueOrError<UUserDefinedStruct*, void> NewUDS = UE::StructUtils::CreateUserDefinedStructFromDescs(StateTree->GetPackage(), PropertyBag, UDSName, Args);
 			StateTree->GlobalParametersStructType = NewUDS.HasValue() ? NewUDS.GetValue() : nullptr;
+			StateTree->GlobalParameters = StateTree->GlobalParametersStructType;
 		}
 
-		if (StateTree->GlobalParametersStructType == nullptr)
+		if (StateTree->GlobalParameters == nullptr)
 		{
-			Log.Reportf(EMessageSeverity::Error, TEXT("Failed to create the global parameters type."));
-			return false;
+			Log.Reportf(EMessageSeverity::Warning, TEXT("Failed to create the global parameters type."));
 		}
 
-		StateTree->GlobalParameters = StateTree->GlobalParametersStructType;
-
 		// The layout should match because the bindings are expected to work.
 		constexpr bool bCheckPropertiesNames = true;
 		if (!UE::StructUtils::TheSameLayout(PropertyBag, StateTree->GlobalParameters, bCheckPropertiesNames))
 		{
 			Log.Reportf(EMessageSeverity::Warning, TEXT("The generated user defined struct do not matches with the editor version."));
-			return false;
 		}
 	}
 	else
 	{
+		ensure(StateTree->ParameterDataType == EStateTreeParameterDataType::ExternalGlobalParameterData);
 		StateTree->Parameters = EditorData->GetRootParametersPropertyBag();
 		StateTree->GlobalParameters = nullptr;
 		StateTree->GlobalParametersStructType = nullptr;
