[StateTree] Added a constructor to TStateTreeInstanceDataStructRef that can take storage instead of an exec frame.
Added support for instance data objects.

#Synced-CL 50157473

diff --git a/Engine/Plugins/Runtime/StateTree/Source/StateTreeModule/Public/StateTreeInstanceData.h b/Engine/Plugins/Runtime/StateTree/Source/StateTreeModule/Public/StateTreeInstanceData.h
--- a/Engine/Plugins/Runtime/StateTree/Source/StateTreeModule/Public/StateTreeInstanceData.h
+++ b/Engine/Plugins/Runtime/StateTree/Source/StateTreeModule/Public/StateTreeInstanceData.h
@@ -681,14 +681,21 @@ struct TStructOpsTypeTraits<FStateTreeInstanceData> : public TStructOpsTypeTrait
 template <typename T>
 struct TStateTreeInstanceDataStructRef
 {
-	TStateTreeInstanceDataStructRef(FStateTreeInstanceData& InInstanceData, const FStateTreeExecutionFrame& CurrentFrame, const FStateTreeDataHandle InDataHandle)
-		: WeakStorage(InInstanceData.GetWeakMutableStorage())
+	TStateTreeInstanceDataStructRef(TWeakPtr<FStateTreeInstanceStorage> InWeakStorage, const FStateTreeExecutionFrame& CurrentFrame, const FStateTreeDataHandle InDataHandle)
+		: WeakStorage(InWeakStorage)
 		, FrameID(CurrentFrame.FrameID)
 		, DataHandle(InDataHandle)
 	{
 		checkf(InDataHandle.GetSource() == EStateTreeDataSourceType::ActiveInstanceData
-			|| InDataHandle.GetSource() == EStateTreeDataSourceType::GlobalInstanceData,
-			TEXT("TStateTreeInstanceDataStructRef supports only struct instance data."));
+			|| InDataHandle.GetSource() == EStateTreeDataSourceType::ActiveInstanceDataObject
+			|| InDataHandle.GetSource() == EStateTreeDataSourceType::GlobalInstanceData
+			|| InDataHandle.GetSource() == EStateTreeDataSourceType::GlobalInstanceDataObject,
+			TEXT("TStateTreeInstanceDataStructRef supports only instance data."));
+	}
+
+	TStateTreeInstanceDataStructRef(FStateTreeInstanceData& InInstanceData, const FStateTreeExecutionFrame& CurrentFrame, const FStateTreeDataHandle InDataHandle)
+		: TStateTreeInstanceDataStructRef(InInstanceData.GetWeakMutableStorage(), CurrentFrame, InDataHandle)
+	{
 	}
 
 PRAGMA_DISABLE_DEPRECATION_WARNINGS
@@ -714,20 +721,24 @@ PRAGMA_ENABLE_DEPRECATION_WARNINGS
 		const FStateTreeExecutionState& Exec = StoragePtr->GetExecutionState();
 		const FStateTreeExecutionFrame* CurrentFrame = Exec.FindActiveFrame(FrameID);
 
-		FStructView Struct;
+		FStateTreeDataView DataView;
 		if (CurrentFrame)
 		{
 			FStateTreeInstanceStorage& Storage = *(StoragePtr.Get());
 			if (IsHandleSourceValid(Storage, *CurrentFrame, DataHandle))
 			{
-				Struct = GetDataView(Storage, *CurrentFrame, DataHandle);
+				DataView = GetDataView(Storage, *CurrentFrame, DataHandle);
 			}
 			else
 			{
-				Struct = Storage.GetMutableTemporaryStruct(CurrentFrame->FrameID, DataHandle);
+				DataView = Storage.GetMutableTemporaryStruct(CurrentFrame->FrameID, DataHandle);
 			}
 
-			check(Struct.GetScriptStruct() == TBaseStructure<T>::Get());
+
+			if constexpr (!std::is_same_v<T, void>)
+			{
+				check(DataView.GetStruct() == TBaseStructure<T>::Get());
+			}
 		}
 		else
 		{
@@ -739,22 +750,33 @@ PRAGMA_ENABLE_DEPRECATION_WARNINGS
 				});
 			if (ExistingInstance)
 			{
-				Struct = FStructView(ExistingInstance->Instance);
+				if (const FStateTreeInstanceObjectWrapper* ObjectWrapper = ExistingInstance->Instance.GetPtr<FStateTreeInstanceObjectWrapper>())
+				{
+					DataView = FStateTreeDataView(ObjectWrapper->InstanceObject);
+				}
+				else
+				{
+					DataView = FStateTreeDataView(ExistingInstance->Instance);
+				}
 			}
 		}
-		return reinterpret_cast<T*>(Struct.GetMemory());
+
+		return reinterpret_cast<T*>(DataView.GetMutableMemory());
 	}
 
 protected:
-
-	FStructView GetDataView(FStateTreeInstanceStorage& Storage, const FStateTreeExecutionFrame& CurrentFrame, const FStateTreeDataHandle Handle) const
+	FStateTreeDataView GetDataView(FStateTreeInstanceStorage& Storage, const FStateTreeExecutionFrame& CurrentFrame, const FStateTreeDataHandle Handle) const
 	{
-		switch (DataHandle.GetSource())
+		switch (Handle.GetSource())
 		{
 		case EStateTreeDataSourceType::GlobalInstanceData:
-			return Storage.GetMutableStruct(CurrentFrame.GlobalInstanceIndexBase.Get() + DataHandle.GetIndex());
+			return Storage.GetMutableStruct(CurrentFrame.GlobalInstanceIndexBase.Get() + Handle.GetIndex());
+		case EStateTreeDataSourceType::GlobalInstanceDataObject:
+			return Storage.GetMutableObject(CurrentFrame.GlobalInstanceIndexBase.Get() + Handle.GetIndex());
 		case EStateTreeDataSourceType::ActiveInstanceData:
-			return Storage.GetMutableStruct(CurrentFrame.ActiveInstanceIndexBase.Get() + DataHandle.GetIndex());
+			return Storage.GetMutableStruct(CurrentFrame.ActiveInstanceIndexBase.Get() + Handle.GetIndex());
+		case EStateTreeDataSourceType::ActiveInstanceDataObject:
+			return Storage.GetMutableObject(CurrentFrame.ActiveInstanceIndexBase.Get() + Handle.GetIndex());
 		default:
 			checkf(false, TEXT("Unhandle case %s"), *UEnum::GetValueAsString(Handle.GetSource()));
 		}
@@ -766,10 +788,12 @@ protected:
 		switch (Handle.GetSource())
 		{
 		case EStateTreeDataSourceType::GlobalInstanceData:
+		case EStateTreeDataSourceType::GlobalInstanceDataObject:
 			return CurrentFrame.GlobalInstanceIndexBase.IsValid()
 				&& Storage.IsValidIndex(CurrentFrame.GlobalInstanceIndexBase.Get() + Handle.GetIndex());
 
 		case EStateTreeDataSourceType::ActiveInstanceData:
+		case EStateTreeDataSourceType::ActiveInstanceDataObject:
 			return CurrentFrame.ActiveInstanceIndexBase.IsValid()
 				&& CurrentFrame.ActiveStates.Contains(Handle.GetState())
 				&& Storage.IsValidIndex(CurrentFrame.ActiveInstanceIndexBase.Get() + Handle.GetIndex());
