[StateTree] Fix StateTree asset duplication without duplicating the GlobalParameter struct.

#Synced-CL 50200679

diff --git a/Engine/Plugins/Runtime/StateTree/Source/StateTreeEditorModule/Private/StateTreeCompiler.cpp b/Engine/Plugins/Runtime/StateTree/Source/StateTreeEditorModule/Private/StateTreeCompiler.cpp
--- a/Engine/Plugins/Runtime/StateTree/Source/StateTreeEditorModule/Private/StateTreeCompiler.cpp
+++ b/Engine/Plugins/Runtime/StateTree/Source/StateTreeEditorModule/Private/StateTreeCompiler.cpp
@@ -688,6 +688,13 @@ bool FStateTreeCompiler::CreateParameters()
 	{
 		StateTree->Parameters = EditorData->GetRootParametersPropertyBag();
 
+		// Previous version didn't handle the duplicate the struct properly.
+		if (StateTree->GlobalParametersStructType
+			&& StateTree->GlobalParametersStructType->GetPackage() != StateTree->GetPackage())
+		{
+			StateTree->GlobalParametersStructType = nullptr;
+		}
+
 		// Copy parameters from EditorData
 		const FInstancedPropertyBag& PropertyBag = EditorData->GetRootParametersPropertyBag();
 		if (StateTree->GlobalParametersStructType)
diff --git a/Engine/Plugins/Runtime/StateTree/Source/StateTreeModule/Public/StateTree.h b/Engine/Plugins/Runtime/StateTree/Source/StateTreeModule/Public/StateTree.h
--- a/Engine/Plugins/Runtime/StateTree/Source/StateTreeModule/Public/StateTree.h
+++ b/Engine/Plugins/Runtime/StateTree/Source/StateTreeModule/Public/StateTree.h
@@ -494,7 +494,7 @@ private:
 
 #if WITH_EDITORONLY_DATA
 	/** Keep a reference between compilation. GlobalParameters is reset between each compilation. */
-	UPROPERTY()
+	UPROPERTY(Instanced)
 	TObjectPtr<UUserDefinedStruct> GlobalParametersStructType;
 
 	/**  To know if the new type was saved and the reference can be considered safe to use.  */
