[StateTree] Change WasSavedWithGlobalParametersStruct for a more generic GetAssetSavedVersion. We had to backout some changes, and the flag might have been saved with the backout modification.

#rnx
#Synced-CL 50200679

diff --git a/Engine/Plugins/Runtime/StateTree/Source/StateTreeModule/Private/StateTree.cpp b/Engine/Plugins/Runtime/StateTree/Source/StateTreeModule/Private/StateTree.cpp
--- a/Engine/Plugins/Runtime/StateTree/Source/StateTreeModule/Private/StateTree.cpp
+++ b/Engine/Plugins/Runtime/StateTree/Source/StateTreeModule/Private/StateTree.cpp
@@ -38,6 +38,12 @@ FCustomVersionRegistration GRegisterStateTreeCustomVersion(FStateTreeCustomVersi
 }
 PRAGMA_ENABLE_DEPRECATION_WARNINGS
 
+#if WITH_EDITOR
+// 0 default
+// 1 with DefaultGlobalParameters
+const int32 UStateTree::LatestCustomAssetSavedVersion = 1;
+#endif
+
 FConstStructView UStateTree::GetDefaultGlobalParameters() const
 {
 	if (GlobalParameters)
@@ -349,7 +355,7 @@ void UStateTree::PreSave(FObjectPreSaveContext Context)
 	const bool bIsEditorAutoSave = Context.IsFromAutoSave();
 	if (bIsUserObject && !bIsEditorAutoSave)
 	{
-		bSavedWithGlobalParametersStructType = true;
+		CustomAssetSavedVersion = LatestCustomAssetSavedVersion;
 	}
 	Super::PreSave(Context);
 }
@@ -437,6 +443,11 @@ void UStateTree::PostLoad()
 {
 	Super::PostLoad();
 
+	if (GlobalParameters)
+	{
+		GlobalParameters->ConditionalPostLoad();
+	}
+
 	for (int32 NodeIndex = 0; NodeIndex < Nodes.Num(); ++NodeIndex)
 	{
 		FStructView NodeView = Nodes[NodeIndex];
diff --git a/Engine/Plugins/Runtime/StateTree/Source/StateTreeModule/Public/StateTree.h b/Engine/Plugins/Runtime/StateTree/Source/StateTreeModule/Public/StateTree.h
--- a/Engine/Plugins/Runtime/StateTree/Source/StateTreeModule/Public/StateTree.h
+++ b/Engine/Plugins/Runtime/StateTree/Source/StateTreeModule/Public/StateTree.h
@@ -328,12 +328,14 @@ public:
 	/** Compile the state trees if the editor hash data as changed since the last compilation. */
 	UE_API void CompileIfChanged();
 
-	/** Was the GlobalParameters properly saved. */
-	bool WasSavedWithGlobalParametersStruct() const
+	/** The expected result of GetAssetSavedVersion. If the values do not match, the UI will suggest saving the asset. */
+	UE_API static const int32 LatestCustomAssetSavedVersion;
+
+	/** Custom asset version to notify the user when the package needs to be resaved. */
+	int32 GetAssetSavedVersion() const
 	{
-		return bSavedWithGlobalParametersStructType;
+		return CustomAssetSavedVersion;
 	}
-
 #endif
 
 #if WITH_EDITOR || WITH_STATETREE_DEBUG
@@ -496,10 +498,6 @@ private:
 	/** Keep a reference between compilation. GlobalParameters is reset between each compilation. */
 	UPROPERTY(Instanced)
 	TObjectPtr<UUserDefinedStruct> GlobalParametersStructType;
-
-	/**  To know if the new type was saved and the reference can be considered safe to use.  */
-	UPROPERTY()
-	bool bSavedWithGlobalParametersStructType = false;
 #endif
 
 	//~ Data created during linking.
@@ -592,6 +590,10 @@ private:
 #if WITH_EDITORONLY_DATA
 	/** List of Struct that are out of date and waiting to be replaced with the new instance. */
 	TSet<FObjectKey> OutOfDateStructs;
+
+	/** The custom version number the package was last saved with. */
+	UPROPERTY()
+	int32 CustomAssetSavedVersion = 0;
 #endif
 
 private:
