[StateTree] Fix false positive warning when the GlobalProperty struct is created. The struct should be at the same level as the StateTree asset in the package.

#rnx

#Synced-CL 50200679

diff --git a/Engine/Plugins/Runtime/StateTree/Source/StateTreeEditorModule/Private/StateTreeCompiler.cpp b/Engine/Plugins/Runtime/StateTree/Source/StateTreeEditorModule/Private/StateTreeCompiler.cpp
--- a/Engine/Plugins/Runtime/StateTree/Source/StateTreeEditorModule/Private/StateTreeCompiler.cpp
+++ b/Engine/Plugins/Runtime/StateTree/Source/StateTreeEditorModule/Private/StateTreeCompiler.cpp
@@ -105,13 +105,15 @@ namespace UE::StateTree::Compiler
 	class FCheckOutersArchive : public FArchiveUObject
 	{
 		using Super = FArchiveUObject;
-		const UStateTree& StateTree;
-		const UStateTreeEditorData& EditorData;
+		TNotNull<const UStateTree*> StateTree;
+		const UUserDefinedStruct* StateTreeGlobalParameterStruct;
+		TNotNull<const UStateTreeEditorData*> EditorData;
 		FStateTreeCompilerLog& Log;
 	public:
 
-		FCheckOutersArchive(const UStateTree& InStateTree, const UStateTreeEditorData& InEditorData, FStateTreeCompilerLog& InLog)
+		FCheckOutersArchive(TNotNull<const UStateTree*> InStateTree, TNotNull<const UStateTreeEditorData*> InEditorData, const UUserDefinedStruct* InStateTreeGlobalParameterStruct, FStateTreeCompilerLog& InLog)
 			: StateTree(InStateTree)
+			, StateTreeGlobalParameterStruct(InStateTreeGlobalParameterStruct)
 			, EditorData(InEditorData)
 			, Log(InLog)
 		{
@@ -140,16 +142,28 @@ namespace UE::StateTree::Compiler
 				{
 					if (Property->HasAnyPropertyFlags(CPF_InstancedReference))
 					{
-						if (!Object->IsInOuter(&StateTree))
+						if (Object == StateTreeGlobalParameterStruct)
 						{
-							Log.Reportf(EMessageSeverity::Error, TEXT("Compiled StateTree contains instanced object %s (%s), which does not belong to the StateTree. This is due to error in the State Tree node implementation."),
-								*GetFullNameSafe(Object), *GetFullNameSafe(Object->GetClass()));
+							if (StateTreeGlobalParameterStruct->GetPackage() != StateTree->GetPackage()
+								|| StateTreeGlobalParameterStruct->IsInOuter(StateTree))
+							{
+								Log.Reportf(EMessageSeverity::Error, TEXT("the GlobalParameter struct object %s (%s), should not belong inside the StateTree. This is due to error in the State Tree compiler implementation."),
+									*GetFullNameSafe(Object), *GetFullNameSafe(Object->GetClass()));
+							}
 						}
-
-						if (Object->IsInOuter(&EditorData))
+						else
 						{
-							Log.Reportf(EMessageSeverity::Error, TEXT("Compiled StateTree contains instanced object %s (%s), which still belongs to the Editor data. This is due to error in the State Tree node implementation."),
-								*GetFullNameSafe(Object), *GetFullNameSafe(Object->GetClass()));
+							if (!Object->IsInOuter(StateTree))
+							{
+								Log.Reportf(EMessageSeverity::Error, TEXT("Compiled StateTree contains instanced object %s (%s), which does not belong to the StateTree. This is due to error in the State Tree node implementation."),
+									*GetFullNameSafe(Object), *GetFullNameSafe(Object->GetClass()));
+							}
+
+							if (Object->IsInOuter(EditorData))
+							{
+								Log.Reportf(EMessageSeverity::Error, TEXT("Compiled StateTree contains instanced object %s (%s), which still belongs to the Editor data. This is due to error in the State Tree node implementation."),
+									*GetFullNameSafe(Object), *GetFullNameSafe(Object->GetClass()));
+							}
 						}
 					}
 				}
@@ -625,7 +639,7 @@ bool FStateTreeCompiler::CompileInternalImpl()
 
 	{
 		// Test if there are invalid reference
-		UE::StateTree::Compiler::FCheckOutersArchive CheckOuters(*StateTree, *EditorData, Log);
+		UE::StateTree::Compiler::FCheckOutersArchive CheckOuters(StateTree, EditorData, StateTree->GlobalParametersStructType, Log);
 		StateTree->Serialize(CheckOuters);
 	}
 
