[StateTree] Add a flag to easily disable the UDS compilation.

#rnx

#Synced-CL 50306593

diff --git a/Engine/Plugins/Runtime/StateTree/Source/StateTreeEditorModule/Private/StateTreeCompiler.cpp b/Engine/Plugins/Runtime/StateTree/Source/StateTreeEditorModule/Private/StateTreeCompiler.cpp
--- a/Engine/Plugins/Runtime/StateTree/Source/StateTreeEditorModule/Private/StateTreeCompiler.cpp
+++ b/Engine/Plugins/Runtime/StateTree/Source/StateTreeEditorModule/Private/StateTreeCompiler.cpp
@@ -101,6 +101,12 @@ namespace UE::StateTree::Compiler
 		TEXT("After a StateTree compiles, log the internal content of the StateTree.")
 	);
 
+	FAutoConsoleVariable CVarEnableCompiledUserDefinedStructForGlobalParameters(
+		TEXT("StateTree.Compiler.EnableCompiledUserDefinedStructForGlobalParameters"),
+		false,
+		TEXT("Compiles a User Defined Struct to replace the GlobalParameters PropertyBag.")
+	);
+
 	// Helper archive that checks that the all instanced sub-objects have correct outer. 
 	class FCheckOutersArchive : public FArchiveUObject
 	{
@@ -627,7 +633,7 @@ bool FStateTreeCompiler::CompileInternalImpl()
 	}
 
 	// Store mapping between state transition identifier and compact transition index. Used for debugging purposes.
-	for (const TPair<FGuid, int32>& ToTransition: IDToTransition)
+	for (const TPair<FGuid, int32>& ToTransition : IDToTransition)
 	{
 		StateTree->IDToTransitionMappings.Emplace(ToTransition.Key, FStateTreeIndex16(ToTransition.Value));
 	}
@@ -702,43 +708,46 @@ bool FStateTreeCompiler::CreateParameters()
 	{
 		StateTree->Parameters = EditorData->GetRootParametersPropertyBag();
 
-		// Previous version didn't handle the duplicate the struct properly.
-		if (StateTree->GlobalParametersStructType
-			&& StateTree->GlobalParametersStructType->GetPackage() != StateTree->GetPackage())
+		if (UE::StateTree::Compiler::CVarEnableCompiledUserDefinedStructForGlobalParameters->GetBool())
 		{
-			StateTree->GlobalParametersStructType = nullptr;
-		}
+			// Previous version didn't handle the duplicate the struct properly.
+			if (StateTree->GlobalParametersStructType
+				&& StateTree->GlobalParametersStructType->GetPackage() != StateTree->GetPackage())
+			{
+				StateTree->GlobalParametersStructType = nullptr;
+			}
 
-		// Copy parameters from EditorData
-		const FInstancedPropertyBag& PropertyBag = EditorData->GetRootParametersPropertyBag();
-		if (StateTree->GlobalParametersStructType)
-		{
-			StateTree->GlobalParameters = StateTree->GlobalParametersStructType;
-			UE::StructUtils::UpdateUserDefinedStructFromDescs(StateTree->GlobalParametersStructType, PropertyBag);
-		}
-		else
-		{
-			// The linker loads and saves in alphabetic order inside the same package. Add a _ to load/save the UDS before the StateTree.
-			const FName UDSName = *FString::Printf(TEXT("_%s_DefaultParameters"), *StateTree->GetName());
-			UE::StructUtils::FCreateUserDefinedStructArgs Args;
-			Args.UserDefinedStructFlags = RF_Transactional | RF_Public;
-			Args.bIsBlueprintType = false;
-			Args.bAddPropertyIDToPropertyName = false;
-			TValueOrError<UUserDefinedStruct*, void> NewUDS = UE::StructUtils::CreateUserDefinedStructFromDescs(StateTree->GetPackage(), PropertyBag, UDSName, Args);
-			StateTree->GlobalParametersStructType = NewUDS.HasValue() ? NewUDS.GetValue() : nullptr;
-			StateTree->GlobalParameters = StateTree->GlobalParametersStructType;
-		}
+			// Copy parameters from EditorData
+			const FInstancedPropertyBag& PropertyBag = EditorData->GetRootParametersPropertyBag();
+			if (StateTree->GlobalParametersStructType)
+			{
+				StateTree->GlobalParameters = StateTree->GlobalParametersStructType;
+				UE::StructUtils::UpdateUserDefinedStructFromDescs(StateTree->GlobalParametersStructType, PropertyBag);
+			}
+			else
+			{
+				// The linker loads and saves in alphabetic order inside the same package. Add a _ to load/save the UDS before the StateTree.
+				const FName UDSName = *FString::Printf(TEXT("_%s_DefaultParameters"), *StateTree->GetName());
+				UE::StructUtils::FCreateUserDefinedStructArgs Args;
+				Args.UserDefinedStructFlags = RF_Transactional | RF_Public;
+				Args.bIsBlueprintType = false;
+				Args.bAddPropertyIDToPropertyName = false;
+				TValueOrError<UUserDefinedStruct*, void> NewUDS = UE::StructUtils::CreateUserDefinedStructFromDescs(StateTree->GetPackage(), PropertyBag, UDSName, Args);
+				StateTree->GlobalParametersStructType = NewUDS.HasValue() ? NewUDS.GetValue() : nullptr;
+				StateTree->GlobalParameters = StateTree->GlobalParametersStructType;
+			}
 
-		if (StateTree->GlobalParameters == nullptr)
-		{
-			Log.Reportf(EMessageSeverity::Warning, TEXT("Failed to create the global parameters type."));
-		}
+			if (StateTree->GlobalParameters == nullptr)
+			{
+				Log.Reportf(EMessageSeverity::Warning, TEXT("Failed to create the global parameters type."));
+			}
 
-		// The layout should match because the bindings are expected to work.
-		constexpr bool bCheckPropertiesNames = true;
-		if (!UE::StructUtils::TheSameLayout(PropertyBag, StateTree->GlobalParameters, bCheckPropertiesNames))
-		{
-			Log.Reportf(EMessageSeverity::Warning, TEXT("The generated user defined struct do not matches with the editor version."));
+			// The layout should match because the bindings are expected to work.
+			constexpr bool bCheckPropertiesNames = true;
+			if (!UE::StructUtils::TheSameLayout(PropertyBag, StateTree->GlobalParameters, bCheckPropertiesNames))
+			{
+				Log.Reportf(EMessageSeverity::Warning, TEXT("The generated user defined struct do not matches with the editor version."));
+			}
 		}
 	}
 	else
