[StateTree] Remove temporary test until the UserDefinedStruct can compile correctly.

#rnx

#Synced-CL 50306593

diff --git a/Engine/Plugins/Runtime/StateTree/Source/StateTreeModule/Private/StateTree.cpp b/Engine/Plugins/Runtime/StateTree/Source/StateTreeModule/Private/StateTree.cpp
--- a/Engine/Plugins/Runtime/StateTree/Source/StateTreeModule/Private/StateTree.cpp
+++ b/Engine/Plugins/Runtime/StateTree/Source/StateTreeModule/Private/StateTree.cpp
@@ -737,7 +737,7 @@ bool UStateTree::Link()
 
 	// @todo Temporary test. Those are to test that the UDS loaded properly and matches the property bag before removing the Parameters.
 	constexpr bool bCheckPropertiesNames = true;
-	if (!UE::StructUtils::TheSameLayout(Parameters, GlobalParameters, bCheckPropertiesNames))
+	if (GlobalParameters != nullptr && !UE::StructUtils::TheSameLayout(Parameters, GlobalParameters, bCheckPropertiesNames))
 	{
 		UE_LOG(LogStateTree, Warning, TEXT("%s: The global parameters doesn't have the same layout as the PropertyBag."), *GetPathName());
 	}
