[State Tree] Added option to override start state by gameplay tag when starting running the tree.

States corresponding to the same gameplay tag will be searched through in breadth first order.

#Synced-CL 50391333

diff --git a/Engine/Plugins/Runtime/StateTree/Source/StateTreeEditorModule/Public/StateTreeState.h b/Engine/Plugins/Runtime/StateTree/Source/StateTreeEditorModule/Public/StateTreeState.h
--- a/Engine/Plugins/Runtime/StateTree/Source/StateTreeEditorModule/Public/StateTreeState.h
+++ b/Engine/Plugins/Runtime/StateTree/Source/StateTreeEditorModule/Public/StateTreeState.h
@@ -298,6 +298,24 @@ public:
 		return static_cast<TStateTreeEditorNode<T>&>(CondNode);
 	}
 
+	template<typename T, typename... TArgs>
+	TStateTreeEditorNode<T>& AddConsideration(TArgs&&... InArgs)
+	{
+		FStateTreeEditorNode& ConsiderationNode = Considerations.AddDefaulted_GetRef();
+		ConsiderationNode.ID = FGuid::NewGuid();
+		ConsiderationNode.Node.InitializeAs<T>(Forward<TArgs>(InArgs)...);
+		const FStateTreeNodeBase& Node = ConsiderationNode.Node.GetMutable<FStateTreeNodeBase>();
+		if (const UScriptStruct* InstanceType = Cast<const UScriptStruct>(Node.GetInstanceDataType()))
+		{
+			ConsiderationNode.Instance.InitializeAs(InstanceType);
+		}
+		if (const UScriptStruct* InstanceType = Cast<const UScriptStruct>(Node.GetExecutionRuntimeDataType()))
+		{
+			ConsiderationNode.ExecutionRuntimeData.InitializeAs(InstanceType);
+		}
+		return static_cast<TStateTreeEditorNode<T>&>(ConsiderationNode);
+	}
+	
 	/**
 	 * Adds Task of specified type.
 	 * @return reference to the new Task.
diff --git a/Engine/Plugins/Runtime/StateTree/Source/StateTreeModule/Private/StateTree.cpp b/Engine/Plugins/Runtime/StateTree/Source/StateTreeModule/Private/StateTree.cpp
--- a/Engine/Plugins/Runtime/StateTree/Source/StateTreeModule/Private/StateTree.cpp
+++ b/Engine/Plugins/Runtime/StateTree/Source/StateTreeModule/Private/StateTree.cpp
@@ -1,6 +1,7 @@
 // Copyright Epic Games, Inc. All Rights Reserved.
 
 #include "StateTree.h"
+#include "Containers/RingBuffer.h"
 #include "Misc/PackageName.h"
 #include "StateTreeLinker.h"
 #include "StateTreeTaskBase.h"
@@ -102,6 +103,58 @@ FStateTreeStateHandle UStateTree::GetStateHandleFromId(const FGuid Id) const
 	return Entry != nullptr ? Entry->Handle : FStateTreeStateHandle::Invalid;
 }
 
+FStateTreeStateHandle UStateTree::GetStateHandleFromGameplayTag(const FGameplayTag Tag, EStateGameplayTagQueryMethod QueryMethod) const
+{
+	TRingBuffer<uint16, TInlineAllocator<32>> HandleBFSQueue;
+
+	auto AddRootStatesIndicesToQueue = [Self = this, &HandleBFSQueue]()
+	{
+		uint16 RootStateIndex = 0;
+		while (RootStateIndex < Self->States.Num())
+		{
+			HandleBFSQueue.Add(RootStateIndex);
+			RootStateIndex = Self->States[RootStateIndex].GetNextSibling();
+		}
+	};
+
+	auto TagMatches = [TagToQuery = Tag, QueryMethod](const FGameplayTag StateTag)
+		{
+			if (QueryMethod == EStateGameplayTagQueryMethod::Includes)
+			{
+				return StateTag.MatchesTag(TagToQuery);
+			}
+
+			if (QueryMethod == EStateGameplayTagQueryMethod::MatchesExact)
+			{
+				return StateTag.MatchesTagExact(TagToQuery);
+			}
+
+			return false;
+		};
+
+	AddRootStatesIndicesToQueue();
+
+	while (!HandleBFSQueue.IsEmpty())
+	{
+		const uint16 CurrentStateHandleIndex = HandleBFSQueue.PopFrontValue();
+		const FCompactStateTreeState& CurrentState = States[CurrentStateHandleIndex];
+
+		if (TagMatches(CurrentState.Tag))
+		{
+			return FStateTreeStateHandle(CurrentStateHandleIndex);
+		}
+		
+		uint16 CurrentStateNextLevelChildIndex = CurrentState.ChildrenBegin;
+		while (CurrentStateNextLevelChildIndex != CurrentState.ChildrenEnd)
+		{
+			HandleBFSQueue.Add(CurrentStateNextLevelChildIndex);
+			CurrentStateNextLevelChildIndex = States[CurrentStateNextLevelChildIndex].GetNextSibling();
+		}
+	}
+
+	return FStateTreeStateHandle::Invalid;
+}
+
 FGuid UStateTree::GetStateIdFromHandle(const FStateTreeStateHandle Handle) const
 {
 	const FStateTreeStateIdToHandle* Entry = IDToStateMappings.FindByPredicate([Handle](const FStateTreeStateIdToHandle& Entry){ return Entry.Handle == Handle; });
diff --git a/Engine/Plugins/Runtime/StateTree/Source/StateTreeModule/Private/StateTreeExecutionContext.cpp b/Engine/Plugins/Runtime/StateTree/Source/StateTreeModule/Private/StateTreeExecutionContext.cpp
--- a/Engine/Plugins/Runtime/StateTree/Source/StateTreeModule/Private/StateTreeExecutionContext.cpp
+++ b/Engine/Plugins/Runtime/StateTree/Source/StateTreeModule/Private/StateTreeExecutionContext.cpp
@@ -239,6 +239,49 @@ namespace UE::StateTree::ExecutionContext::Private
 			}
 		}
 	}
+
+	struct FSelectStartStateInfo
+	{
+		FStateTreeStateHandle StartStateToSelectHandle;
+		FStateTreeStateHandle FrameRootStateHandle;
+	};
+	FSelectStartStateInfo GetSelectStartStateInfo(TNotNull<const UStateTree*> InRootTree, const FStateTreeExecutionContext::FStartParameters& InStartParameters)
+	{
+		FSelectStartStateInfo SelectStartState = { FStateTreeStateHandle::Invalid, FStateTreeStateHandle::Invalid };
+
+		if (const FStateTreeExecutionContext::FStartParameters::FStateToSelectOverrideArgs* SelectStateOverrideArgs = InStartParameters.SelectStateOverrideArgs.GetPtrOrNull())
+		{
+			if (SelectStateOverrideArgs->StateTag.IsValid())
+			{
+				FStateTreeStateHandle OverrideStateToSelectHandle = InRootTree->GetStateHandleFromGameplayTag(SelectStateOverrideArgs->StateTag, SelectStateOverrideArgs->TagQueryMethod);
+				if (OverrideStateToSelectHandle.IsValid())
+				{
+					check(!OverrideStateToSelectHandle.IsCompletionState());
+
+					const FCompactStateTreeState* OverrideStateToSelect = InRootTree->GetStateFromHandle(OverrideStateToSelectHandle);
+					check(OverrideStateToSelect);
+
+					SelectStartState.StartStateToSelectHandle = OverrideStateToSelectHandle;
+					SelectStartState.FrameRootStateHandle = OverrideStateToSelectHandle;
+
+					// A tree could have multiple roots and the selected state could be within a subtree(the subtree will be treated as standalone)
+					FStateTreeStateHandle OverrideParentStateHandle = OverrideStateToSelect->Parent;
+					while (const FCompactStateTreeState* OverrideParentState = InRootTree->GetStateFromHandle(OverrideParentStateHandle))
+					{
+						SelectStartState.FrameRootStateHandle = OverrideParentStateHandle;
+						OverrideParentStateHandle = OverrideParentState->Parent;
+					}
+				}
+			}
+		}
+		else
+		{
+			// No override info is provided, tree will start from root state by default
+			SelectStartState = {FStateTreeStateHandle::Root, FStateTreeStateHandle::Root};
+		}
+
+		return SelectStartState;
+	}
 } // namespace UE::StateTree::ExecutionContext::Private
 
 namespace UE::StateTree::ExecutionContext
@@ -1461,7 +1504,19 @@ EStateTreeRunStatus FStateTreeExecutionContext::Start(FStartParameters Parameter
 	TGuardValue<TSharedPtr<FSelectStateResult>> TemporaryStorageScope(CurrentlyProcessedTemporaryStorage, SelectStateResult);
 
 	const FActiveFrameID InitFrameID = FActiveFrameID(Storage.GenerateUniqueId());
-	const FExecutionFrameHandle InitFrameHandle = FExecutionFrameHandle(&RootStateTree, FStateTreeStateHandle::Root);
+
+	const FSelectStartStateInfo SelectStartStateInfo = GetSelectStartStateInfo(&RootStateTree, Parameters);
+	if (!SelectStartStateInfo.StartStateToSelectHandle.IsValid())
+	{
+		STATETREE_LOG(Warning, TEXT("%hs: The start state in tree %s could not be found. "
+			"The overriding state tag is either invalid or no state corresponds to that tag."),
+			__FUNCTION__, *GetFullNameSafe(&RootStateTree));
+
+		return EStateTreeRunStatus::Failed;
+	}
+
+	const FExecutionFrameHandle InitFrameHandle = FExecutionFrameHandle(&RootStateTree, SelectStartStateInfo.FrameRootStateHandle);
+
 	constexpr bool bIsGlobalFrame = true;
 	FStateTreeExecutionFrame& InitFrame = SelectStateResult->MakeAndAddTemporaryFrame(InitFrameID, InitFrameHandle, bIsGlobalFrame);
 	InitFrame.ExecutionRuntimeIndexBase = FStateTreeIndex16(Storage.AddExecutionRuntimeData(GetOwner(), InitFrameHandle));
@@ -1505,8 +1560,8 @@ EStateTreeRunStatus FStateTreeExecutionContext::Start(FStartParameters Parameter
 		Exec.LastTickStatus = EStateTreeRunStatus::Unset;
 
 		FSelectStateArguments SelectStateArgs;
-		SelectStateArgs.SourceState = FActiveState(InitFrameID, FActiveStateID(), InitFrameHandle.GetRootState());
-		SelectStateArgs.TargetState = FStateHandleContext(InitFrameHandle.GetStateTree(), InitFrameHandle.GetRootState());
+		SelectStateArgs.SourceState = FActiveState(InitFrameID, FActiveStateID(), FStateTreeStateHandle());
+		SelectStateArgs.TargetState = FStateHandleContext(InitFrameHandle.GetStateTree(), SelectStartStateInfo.StartStateToSelectHandle);
 		SelectStateArgs.Behavior = ESelectStateBehavior::StateTransition;
 		SelectStateArgs.SelectionRules = InitFrameHandle.GetStateTree()->StateSelectionRules;
 		if (SelectState(SelectStateArgs, SelectStateResult))
diff --git a/Engine/Plugins/Runtime/StateTree/Source/StateTreeModule/Public/StateTree.h b/Engine/Plugins/Runtime/StateTree/Source/StateTreeModule/Public/StateTree.h
--- a/Engine/Plugins/Runtime/StateTree/Source/StateTreeModule/Public/StateTree.h
+++ b/Engine/Plugins/Runtime/StateTree/Source/StateTreeModule/Public/StateTree.h
@@ -119,6 +119,10 @@ struct FStateTreeMemoryUsage
 };
 #endif
 
+namespace UE::StateTree
+{
+	
+}
 
 /**
  * StateTree asset. Contains the StateTree definition in both editor and runtime (baked) formats.
@@ -201,6 +205,26 @@ public:
 	/** @return State handle matching a given Id; invalid handle if state not found. */
 	UE_API FStateTreeStateHandle GetStateHandleFromId(const FGuid Id) const;
 
+	/*
+	 * Determines how we match a state gameplay tag
+	 */
+	enum class EStateGameplayTagQueryMethod
+	{
+		// Foo matches both Foo and Foo.Bar. Foo.Bar does not match Foo.
+		Includes,
+
+		// Foo only matches Foo. Foo does not match Foo.Bar. Foo.Bar does not match Foo
+		MatchesExact
+	};
+
+	/**
+	 * Get state handle by gameplay tag. States with same tag are searched in breadth first order.
+	 * @param Tag Gameplay tags to query
+	 * @param QueryMethod how to query tags
+	 * @return State handle matching a given gameplay tag
+	 */
+	UE_API FStateTreeStateHandle GetStateHandleFromGameplayTag(const FGameplayTag Tag, EStateGameplayTagQueryMethod QueryMethod) const;
+
 	/** @return ID of the state matching a given state handle; invalid Id if state not found. */
 	UE_API FGuid GetStateIdFromHandle(const FStateTreeStateHandle Handle) const;
 
diff --git a/Engine/Plugins/Runtime/StateTree/Source/StateTreeModule/Public/StateTreeExecutionContext.h b/Engine/Plugins/Runtime/StateTree/Source/StateTreeModule/Public/StateTreeExecutionContext.h
--- a/Engine/Plugins/Runtime/StateTree/Source/StateTreeModule/Public/StateTreeExecutionContext.h
+++ b/Engine/Plugins/Runtime/StateTree/Source/StateTreeModule/Public/StateTreeExecutionContext.h
@@ -443,10 +443,22 @@ PRAGMA_ENABLE_DEPRECATION_WARNINGS
 
 		/** Optional event queue from another instance data. Marks the event queue not owned. */
 		const TSharedPtr<FStateTreeEventQueue> SharedEventQueue;
-
+		
 		/** Optional override of initial seed for RandomStream. By default FPlatformTime::Cycles() will be used. */
 		TOptional<int32> RandomSeed;
 
+		struct FStateToSelectOverrideArgs
+		{
+			/** Optional override of tag of start state to select. */
+			FGameplayTag StateTag;
+			
+			/** Does the tag needs to match exactly when searching the state. */
+			UStateTree::EStateGameplayTagQueryMethod TagQueryMethod = UStateTree::EStateGameplayTagQueryMethod::MatchesExact;
+		};
+
+		/** Optional override of initial state to select. By default, root state will be selected. */
+		TOptional<FStateToSelectOverrideArgs> SelectStateOverrideArgs;
+
 #if WITH_EDITORONLY_DATA
 		UE_DEPRECATED(5.8, "Use InitialGlobalParameters instead.")
 		/** Optional override of global parameters initial values. */
diff --git a/Engine/Plugins/Runtime/StateTree/Source/StateTreeTestSuite/Private/StateTreeTest.cpp b/Engine/Plugins/Runtime/StateTree/Source/StateTreeTestSuite/Private/StateTreeTest.cpp
--- a/Engine/Plugins/Runtime/StateTree/Source/StateTreeTestSuite/Private/StateTreeTest.cpp
+++ b/Engine/Plugins/Runtime/StateTree/Source/StateTreeTestSuite/Private/StateTreeTest.cpp
@@ -12,6 +12,7 @@
 #include "AutoRTFM.h"
 #include "Engine/World.h"
 #include "Async/ParallelFor.h"
+#include "Considerations/StateTreeCommonConsiderations.h"
 #include "GameplayTagsManager.h"
 #include "StateTreeReference.h"
 
@@ -2417,6 +2418,154 @@ struct FStateTreeTest_Reentrant : FStateTreeTestBase
 };
 IMPLEMENT_STATE_TREE_INSTANT_TEST(FStateTreeTest_Reentrant, "System.StateTree.Reentrant");
 
+struct FStateTreeTest_OverrideStartState : FStateTreeTestBase
+{
+	virtual bool InstantTest() override
+	{
+		//Tree
+		//	Root1
+		//    Child1
+		//    Child2
+		//      Child3
+		//    Subtree1
+		//  Root2
+		//	  Child4
+		//  Root3 SelectionBehavior: Utility
+		//    Child5 score 0.5
+		//    Child6 score 1
+		UStateTree& StateTree = NewStateTree();
+		{
+			UStateTreeEditorData& EditorData = *Cast<UStateTreeEditorData>(StateTree.EditorData);
+			UStateTreeState& Root1 = EditorData.AddSubTree(TEXT("Root1"));
+
+			UStateTreeState& Child1 = Root1.AddChildState(TEXT("Child1"));
+
+			UStateTreeState& Child2 = Root1.AddChildState(TEXT("Child2"));
+			Child2.Tag = GetTestTag1();
+			
+			UStateTreeState& Child3 = Child2.AddChildState(TEXT("Child3"));
+			Child3.Tag = GetTestTag2();
+
+			UStateTreeState& Subtree1 = Root1.AddChildState(TEXT("Subtree1"));
+			Subtree1.Type = EStateTreeStateType::Subtree;
+			Subtree1.Tag = GetTestTag3();
+
+			UStateTreeState& Root2 = EditorData.AddSubTree(TEXT("Root2"));
+			Root2.Tag = GetTestTag4();
+			UStateTreeState& Child4 = Root2.AddChildState(TEXT("Child4"));
+
+			UStateTreeState& Root3 = EditorData.AddSubTree(TEXT("Root4"));
+			Root3.SelectionBehavior = EStateTreeStateSelectionBehavior::TrySelectChildrenWithHighestUtility;
+			{
+				UStateTreeState& Child5 = Root3.AddChildState(TEXT("Child6"));
+				Child5.Tag = GetTestTag5();
+				TStateTreeEditorNode<FStateTreeConstantConsideration>& Child5ConsiderationEditorNode = Child5.AddConsideration<FStateTreeConstantConsideration>();
+				Child5ConsiderationEditorNode.GetInstanceData().Constant = 0.5f;
+			}
+			{
+				UStateTreeState& Child6 = Root3.AddChildState(TEXT("Child7"));
+				TStateTreeEditorNode<FStateTreeConstantConsideration>& Child6ConsiderationEditorNode = Child6.AddConsideration<FStateTreeConstantConsideration>();
+				Child6ConsiderationEditorNode.GetInstanceData().Constant = 1.f;
+			}
+		}
+
+		{
+			FStateTreeCompilerLog Log;
+			FStateTreeCompiler Compiler(Log);
+			const bool bResult = Compiler.Compile(StateTree);
+			AITEST_TRUE("StateTree should get compiled", bResult);
+		}
+
+		{
+			FStateTreeInstanceData InstanceData;
+			
+			FTestStateTreeExecutionContext Exec(StateTree, StateTree, InstanceData);
+			const bool bInitSucceeded = Exec.IsValid();
+			AITEST_TRUE("StateTree should init", bInitSucceeded);
+			
+			// No start state override
+			{
+				EStateTreeRunStatus Status = Exec.Start();
+				AITEST_TRUE("StateTree Active States should be in Root1/Child1",
+					Exec.ExpectInActiveStates("Root1", "Child1"));
+				AITEST_TRUE("StateTree should be running", Status == EStateTreeRunStatus::Running);
+				
+				Exec.LogClear();
+			}
+
+			// Start state Child2
+			{
+				FStateTreeExecutionContext::FStartParameters::FStateToSelectOverrideArgs SelectStateOverrideArgs =
+				{ GetTestTag1(), UStateTree::EStateGameplayTagQueryMethod::MatchesExact };
+
+				EStateTreeRunStatus Status = Exec.Start(FStateTreeExecutionContext::FStartParameters{ .SelectStateOverrideArgs = SelectStateOverrideArgs });
+
+				AITEST_TRUE("StateTree Active States should be in Root1/Child2/Child3 with Child2 override",
+					Exec.ExpectInActiveStates("Root1", "Child2", "Child3"));
+				AITEST_TRUE("StateTree should be running", Status == EStateTreeRunStatus::Running);
+				Exec.LogClear();
+			}
+
+			// Start state Child3
+			{
+				FStateTreeExecutionContext::FStartParameters::FStateToSelectOverrideArgs SelectStateOverrideArgs =
+				{ GetTestTag2(), UStateTree::EStateGameplayTagQueryMethod::MatchesExact };
+
+				EStateTreeRunStatus Status = Exec.Start(FStateTreeExecutionContext::FStartParameters{ .SelectStateOverrideArgs = SelectStateOverrideArgs });
+
+				AITEST_TRUE("StateTree Active States should be in Root1/Child2/Child3 with Child3 override",
+					Exec.ExpectInActiveStates("Root1", "Child2", "Child3"));
+				AITEST_TRUE("StateTree should be running", Status == EStateTreeRunStatus::Running);
+				Exec.LogClear();
+			}
+
+			// Start state Subtree1
+			{
+				FStateTreeExecutionContext::FStartParameters::FStateToSelectOverrideArgs SelectStateOverrideArgs =
+				{ GetTestTag3(), UStateTree::EStateGameplayTagQueryMethod::MatchesExact };
+				
+				EStateTreeRunStatus Status = Exec.Start(FStateTreeExecutionContext::FStartParameters{ .SelectStateOverrideArgs = SelectStateOverrideArgs });
+
+				AITEST_TRUE("StateTree Active States should be in Subtree1 with Subtree1 override",
+					Exec.ExpectInActiveStates("Subtree1"));
+				AITEST_TRUE("StateTree should be running", Status == EStateTreeRunStatus::Running);
+				Exec.LogClear();
+			}
+
+			// Start state Root2
+			{
+				FStateTreeExecutionContext::FStartParameters::FStateToSelectOverrideArgs SelectStateOverrideArgs =
+				{ GetTestTag4(), UStateTree::EStateGameplayTagQueryMethod::MatchesExact };
+				
+				EStateTreeRunStatus Status = Exec.Start(FStateTreeExecutionContext::FStartParameters{ .SelectStateOverrideArgs = SelectStateOverrideArgs });
+
+				AITEST_TRUE("StateTree Active States should be in Root2/Child4 with Root2 override",
+					Exec.ExpectInActiveStates("Root2", "Child4"));
+				AITEST_TRUE("StateTree should be running", Status == EStateTreeRunStatus::Running);
+				Exec.LogClear();
+			}
+
+			// Start state Child6
+			{
+				FStateTreeExecutionContext::FStartParameters::FStateToSelectOverrideArgs SelectStateOverrideArgs =
+				{ GetTestTag5(), UStateTree::EStateGameplayTagQueryMethod::MatchesExact };
+				
+				EStateTreeRunStatus Status = Exec.Start(FStateTreeExecutionContext::FStartParameters{ .SelectStateOverrideArgs = SelectStateOverrideArgs });
+
+				AITEST_TRUE("StateTree Active States should be in Root4/Child6 with Child6 override",
+					Exec.ExpectInActiveStates("Root4", "Child6"));
+				AITEST_TRUE("StateTree should be running", Status == EStateTreeRunStatus::Running);
+				Exec.LogClear();
+			}
+
+			Exec.Stop();
+		}
+
+		return true;
+	}
+};
+IMPLEMENT_STATE_TREE_INSTANT_TEST(FStateTreeTest_OverrideStartState, "System.StateTree.OverrideStartState");
+
 } // namespace UE::StateTree::Tests
 
 UE_ENABLE_OPTIMIZATION_SHIP
diff --git a/Engine/Plugins/Runtime/StateTree/Source/StateTreeTestSuite/Private/StateTreeTestBase.cpp b/Engine/Plugins/Runtime/StateTree/Source/StateTreeTestSuite/Private/StateTreeTestBase.cpp
--- a/Engine/Plugins/Runtime/StateTree/Source/StateTreeTestSuite/Private/StateTreeTestBase.cpp
+++ b/Engine/Plugins/Runtime/StateTree/Source/StateTreeTestSuite/Private/StateTreeTestBase.cpp
@@ -43,13 +43,19 @@ struct FNativeGameplayTags : public FGameplayTagNativeAdder
 	FGameplayTag TestTag;
 	FGameplayTag TestTag2;
 	FGameplayTag TestTag3;
+	FGameplayTag TestTag4;
+	FGameplayTag TestTag5;
+	FGameplayTag TestTag6;
 
 	virtual void AddTags() override
 	{
 		UGameplayTagsManager& Manager = UGameplayTagsManager::Get();
 		TestTag = Manager.AddNativeGameplayTag(TEXT("Test.StateTree.Tag"));
 		TestTag2 = Manager.AddNativeGameplayTag(TEXT("Test.StateTree.Tag2"));
-		TestTag2 = Manager.AddNativeGameplayTag(TEXT("Test.StateTree.Tag3"));
+		TestTag3 = Manager.AddNativeGameplayTag(TEXT("Test.StateTree.Tag3"));
+		TestTag4 = Manager.AddNativeGameplayTag(TEXT("Test.StateTree.Tag4"));
+		TestTag5 = Manager.AddNativeGameplayTag(TEXT("Test.StateTree.Tag5"));
+		TestTag6 = Manager.AddNativeGameplayTag(TEXT("Test.StateTree.Tag6"));
 	}
 };
 static FNativeGameplayTags GameplayTagStaticInstance;
@@ -69,6 +75,21 @@ FGameplayTag FStateTreeTestBase::GetTestTag3()
 	return GameplayTagStaticInstance.TestTag3;
 }
 
+FGameplayTag FStateTreeTestBase::GetTestTag4()
+{
+	return GameplayTagStaticInstance.TestTag4;
+}
+
+FGameplayTag FStateTreeTestBase::GetTestTag5()
+{
+	return GameplayTagStaticInstance.TestTag5;
+}
+
+FGameplayTag FStateTreeTestBase::GetTestTag6()
+{
+	return GameplayTagStaticInstance.TestTag6;
+}
+
 FInstancedPropertyBag& FStateTreeTestBase::GetRootPropertyBag(UStateTreeEditorData& EditorData) const
 {
 	return const_cast<FInstancedPropertyBag&>(EditorData.GetRootParametersPropertyBag());
diff --git a/Engine/Plugins/Runtime/StateTree/Source/StateTreeTestSuite/Public/StateTreeTestBase.h b/Engine/Plugins/Runtime/StateTree/Source/StateTreeTestSuite/Public/StateTreeTestBase.h
--- a/Engine/Plugins/Runtime/StateTree/Source/StateTreeTestSuite/Public/StateTreeTestBase.h
+++ b/Engine/Plugins/Runtime/StateTree/Source/StateTreeTestSuite/Public/StateTreeTestBase.h
@@ -22,6 +22,9 @@ protected:
 	static FGameplayTag GetTestTag1();
 	static FGameplayTag GetTestTag2();
 	static FGameplayTag GetTestTag3();
+	static FGameplayTag GetTestTag4();
+	static FGameplayTag GetTestTag5();
+	static FGameplayTag GetTestTag6();
 	FInstancedPropertyBag& GetRootPropertyBag(UStateTreeEditorData& EditorData) const;
 };
 
