[StateTree] Add test for enum bindings. int32 enum is not supported and StateTreePropertyRef with enum.

#Synced-CL 50535091

diff --git a/Engine/Plugins/Runtime/StateTree/Source/StateTreeModule/Private/StateTreePropertyRefHelpers.cpp b/Engine/Plugins/Runtime/StateTree/Source/StateTreeModule/Private/StateTreePropertyRefHelpers.cpp
--- a/Engine/Plugins/Runtime/StateTree/Source/StateTreeModule/Private/StateTreePropertyRefHelpers.cpp
+++ b/Engine/Plugins/Runtime/StateTree/Source/StateTreeModule/Private/StateTreePropertyRefHelpers.cpp
@@ -189,6 +189,13 @@ namespace UE::StateTree::PropertyRefHelpers
 						return true;
 					}
 				}
+				else if (const FByteProperty* ByteProperty = CastField<FByteProperty>(TestProperty))
+				{
+					if (ByteProperty->Enum == TargetRefField)
+					{
+						return true;
+					}
+				}
 			}
 		}
 
@@ -616,9 +623,12 @@ PRAGMA_ENABLE_DEPRECATION_WARNINGS
 				}
 				else if (UEnum* Enum = Cast<UEnum>(TargetRefField))
 				{
-					PinType.PinCategory = UEdGraphSchema_K2::PC_Enum;
 					PinType.PinSubCategoryObject = Enum;
-					if (Enum->GetMaxEnumValue() <= (int64)std::numeric_limits<uint8>::max())
+					if (Enum->GetCppForm() == UEnum::ECppForm::EnumClass)
+					{
+						PinType.PinCategory = UEdGraphSchema_K2::PC_Enum;
+					}
+					else
 					{
 						// Use byte for BP. It will use the correct picker and enum k2 node.
 						PinType.PinCategory = UEdGraphSchema_K2::PC_Byte;
diff --git a/Engine/Plugins/Runtime/StateTree/Source/StateTreeModule/Public/StateTreePropertyRefHelpers.h b/Engine/Plugins/Runtime/StateTree/Source/StateTreeModule/Public/StateTreePropertyRefHelpers.h
--- a/Engine/Plugins/Runtime/StateTree/Source/StateTreeModule/Public/StateTreePropertyRefHelpers.h
+++ b/Engine/Plugins/Runtime/StateTree/Source/StateTreeModule/Public/StateTreePropertyRefHelpers.h
@@ -179,6 +179,15 @@ namespace UE::StateTree::PropertyRefHelpers
 		}
 	};
 
+	template<typename T>
+	struct Validator<TEnumAsByte<T>>
+	{
+		static bool IsValid(const FProperty& Property)
+		{
+			return Property.IsA<FByteProperty>();
+		}
+	};
+
 	/* Checks if provided property is compatible with selected ScriptStruct. */
 	bool STATETREEMODULE_API IsPropertyCompatibleWithStruct(const FProperty& Property, const class UScriptStruct& Struct);
 
diff --git a/Engine/Plugins/Runtime/StateTree/Source/StateTreeTestSuite/Private/StateTreeBindingTest.cpp b/Engine/Plugins/Runtime/StateTree/Source/StateTreeTestSuite/Private/StateTreeBindingTest.cpp
--- a/Engine/Plugins/Runtime/StateTree/Source/StateTreeTestSuite/Private/StateTreeBindingTest.cpp
+++ b/Engine/Plugins/Runtime/StateTree/Source/StateTreeTestSuite/Private/StateTreeBindingTest.cpp
@@ -1660,6 +1660,258 @@ struct FStateTreeTest_BindingStructRef : FStateTreeTestBase
 };
 IMPLEMENT_STATE_TREE_INSTANT_TEST(FStateTreeTest_BindingStructRef, "System.StateTree.Binding.StructRef");
 
+struct FStateTreeTest_BindingEnum : FStateTreeTestBase
+{
+	// The test check if all enum type supported by UHT are also supported by the binding systems.
+	//Also test if the StateTreePropertyRef with enum works.
+	//The value will be set from Root to State1. The default value of Root are set to B. The default value of State1 are set to C.
+	//We expect the value in State1 to be set to B after the State1::OnEnter
+	virtual bool InstantTest() override
+	{
+		//	Root (global variable, )
+		//		State1 (binding and ref to root variable)
+
+		UStateTree& StateTree = NewStateTree();
+		{
+			UStateTreeEditorData& EditorData = *Cast<UStateTreeEditorData>(StateTree.EditorData);
+
+			UStateTreeState& Root = EditorData.AddSubTree("Tree1StateRoot");
+			UStateTreeState& State1 = Root.AddChildState("Tree1State1", EStateTreeStateType::State);
+			{
+				using namespace UE::StateTree::Test;
+				Root.Parameters.Parameters.AddProperty("StructWithEnum", EPropertyBagPropertyType::Struct, FStructWithEnum::StaticStruct());
+				FStructWithEnum StructValue;
+				StructValue.CppClassEnumByte = ECppClassEnumByte::B;
+				StructValue.CppClassEnumInt = ECppClassEnumInt::B;
+				StructValue.NamespaceEnumByte = ENamespaceEnumByte::B;
+				StructValue.EnumByte = EEnumByte::EEnumByte_B;
+				Root.Parameters.Parameters.SetValueStruct("StructWithEnum", StructValue);
+
+				Root.Parameters.Parameters.AddProperty("CppClassEnumByte", EPropertyBagPropertyType::Enum, StaticEnum<ECppClassEnumByte>());
+				Root.Parameters.Parameters.SetValueEnum("CppClassEnumByte", ECppClassEnumByte::B);
+				Root.Parameters.Parameters.AddProperty("CppClassEnumInt", EPropertyBagPropertyType::Enum, StaticEnum<ECppClassEnumInt>());
+				Root.Parameters.Parameters.SetValueEnum("CppClassEnumInt", ECppClassEnumInt::B);
+				Root.Parameters.Parameters.AddProperty("NamespaceEnumByte", EPropertyBagPropertyType::Enum, StaticEnum<ENamespaceEnumByte::Type>());
+				Root.Parameters.Parameters.SetValueEnum("NamespaceEnumByte", ENamespaceEnumByte::B);
+				Root.Parameters.Parameters.AddProperty("EnumByte", EPropertyBagPropertyType::Enum, StaticEnum<EEnumByte>());
+				Root.Parameters.Parameters.SetValueEnum("EnumByte", EEnumByte::EEnumByte_B);
+			}
+			{
+				using namespace UE::StateTree::Test;
+				State1.Parameters.Parameters = Root.Parameters.Parameters;
+				State1.Parameters.Parameters.AddProperty("StructWithEnumRef1", EPropertyBagPropertyType::Struct, FStructWithEnumRef::StaticStruct());
+				State1.Parameters.Parameters.AddProperty("StructWithEnumRef2", EPropertyBagPropertyType::Struct, FStructWithEnumRef::StaticStruct());
+
+				FStructWithEnum StructValue;
+				StructValue.CppClassEnumByte = ECppClassEnumByte::C;
+				StructValue.CppClassEnumInt = ECppClassEnumInt::C;
+				StructValue.NamespaceEnumByte = ENamespaceEnumByte::C;
+				StructValue.EnumByte = EEnumByte::EEnumByte_C;
+				State1.Parameters.Parameters.SetValueStruct("StructWithEnum", StructValue);
+				State1.Parameters.Parameters.SetValueEnum("CppClassEnumByte", ECppClassEnumByte::C);
+				State1.Parameters.Parameters.SetValueEnum("CppClassEnumInt", ECppClassEnumInt::C);
+				State1.Parameters.Parameters.SetValueEnum("NamespaceEnumByte", ENamespaceEnumByte::C);
+				State1.Parameters.Parameters.SetValueEnum("EnumByte", EEnumByte::EEnumByte_C);
+
+				EditorData.AddPropertyBinding(FPropertyBindingPath(Root.Parameters.ID, "StructWithEnum"), FPropertyBindingPath(State1.Parameters.ID, "StructWithEnum"));
+				EditorData.AddPropertyBinding(FPropertyBindingPath(Root.Parameters.ID, "CppClassEnumByte"), FPropertyBindingPath(State1.Parameters.ID, "CppClassEnumByte"));
+				EditorData.AddPropertyBinding(FPropertyBindingPath(Root.Parameters.ID, "CppClassEnumInt"), FPropertyBindingPath(State1.Parameters.ID, "CppClassEnumInt"));
+				EditorData.AddPropertyBinding(FPropertyBindingPath(Root.Parameters.ID, "NamespaceEnumByte"), FPropertyBindingPath(State1.Parameters.ID, "NamespaceEnumByte"));
+				EditorData.AddPropertyBinding(FPropertyBindingPath(Root.Parameters.ID, "EnumByte"), FPropertyBindingPath(State1.Parameters.ID, "EnumByte"));
+				EditorData.AddPropertyBinding(FPropertyBindingPath(Root.Parameters.ID, "CppClassEnumByte"), FPropertyBindingPath(State1.Parameters.ID, {FPropertyBindingPathSegment("StructWithEnumRef1"), FPropertyBindingPathSegment("CppClassEnumByte")}));
+				EditorData.AddPropertyBinding(FPropertyBindingPath(Root.Parameters.ID, "CppClassEnumInt"), FPropertyBindingPath(State1.Parameters.ID, {FPropertyBindingPathSegment("StructWithEnumRef1"), FPropertyBindingPathSegment("CppClassEnumInt")}));
+				EditorData.AddPropertyBinding(FPropertyBindingPath(Root.Parameters.ID, "NamespaceEnumByte"), FPropertyBindingPath(State1.Parameters.ID, { FPropertyBindingPathSegment("StructWithEnumRef1"), FPropertyBindingPathSegment("NamespaceEnumByte") }));
+				EditorData.AddPropertyBinding(FPropertyBindingPath(Root.Parameters.ID, "EnumByte"), FPropertyBindingPath(State1.Parameters.ID, { FPropertyBindingPathSegment("StructWithEnumRef1"), FPropertyBindingPathSegment("EnumByte") }));
+				EditorData.AddPropertyBinding(FPropertyBindingPath(Root.Parameters.ID, { FPropertyBindingPathSegment("StructWithEnum"), FPropertyBindingPathSegment("CppClassEnumByte") }), FPropertyBindingPath(State1.Parameters.ID, {FPropertyBindingPathSegment("StructWithEnumRef2"), FPropertyBindingPathSegment("CppClassEnumByte")}));
+				EditorData.AddPropertyBinding(FPropertyBindingPath(Root.Parameters.ID, { FPropertyBindingPathSegment("StructWithEnum"), FPropertyBindingPathSegment("CppClassEnumInt") }), FPropertyBindingPath(State1.Parameters.ID, {FPropertyBindingPathSegment("StructWithEnumRef2"), FPropertyBindingPathSegment("CppClassEnumInt")}));
+				EditorData.AddPropertyBinding(FPropertyBindingPath(Root.Parameters.ID, { FPropertyBindingPathSegment("StructWithEnum"), FPropertyBindingPathSegment("NamespaceEnumByte") }), FPropertyBindingPath(State1.Parameters.ID, { FPropertyBindingPathSegment("StructWithEnumRef2"), FPropertyBindingPathSegment("NamespaceEnumByte") }));
+				EditorData.AddPropertyBinding(FPropertyBindingPath(Root.Parameters.ID, { FPropertyBindingPathSegment("StructWithEnum"), FPropertyBindingPathSegment("EnumByte") }), FPropertyBindingPath(State1.Parameters.ID, { FPropertyBindingPathSegment("StructWithEnumRef2"), FPropertyBindingPathSegment("EnumByte") }));
+			}
+		}
+		{
+			FStateTreeCompilerLog Log;
+			FStateTreeCompiler Compiler(Log);
+			const bool bResult = Compiler.Compile(StateTree);
+			AITEST_TRUE(TEXT("StateTree1 should compiled"), bResult);
+		}
+
+		const FStateTreeStateHandle State1Handle = FStateTreeStateHandle(1);
+		const FCompactStateTreeState* CompactState1 = StateTree.GetStateFromHandle(State1Handle);
+		AITEST_TRUE(TEXT("Invalid state"), CompactState1 != nullptr);
+		CA_ASSUME(CompactState1); // Suppress code analyzer warning C6011
+
+		{
+			FStateTreeInstanceData InstanceData;
+			{
+				FTestStateTreeExecutionContext Exec(StateTree, StateTree, InstanceData);
+				const bool bInitSucceeded = Exec.IsValid();
+				AITEST_TRUE(TEXT("StateTree should init"), bInitSucceeded);
+				const EStateTreeRunStatus Status = Exec.Start();
+				AITEST_EQUAL(TEXT("Start should complete with Running"), Status, EStateTreeRunStatus::Running);
+			}
+			{
+				AITEST_TRUE(TEXT("A single frame"), InstanceData.GetExecutionState()->ActiveFrames.Num() == 1);
+				FStateTreeDataView State1DataView = UE::StateTree::InstanceData::GetDataView(
+					InstanceData.GetMutableStorage(),
+					nullptr,
+					InstanceData.GetExecutionState()->ActiveFrames[0],
+					CompactState1->ParameterDataHandle);
+
+				const UStruct* Struct = State1DataView.GetStruct();
+				const void* Memory = State1DataView.GetMemory();
+				AITEST_TRUE(TEXT("Valid data view"), Struct != nullptr && Memory != nullptr);
+				CA_ASSUME(Struct); // Suppress code analyzer warning C6011
+				CA_ASSUME(Memory); // Suppress code analyzer warning C6011
+				
+				{
+					using namespace UE::StateTree::Test;
+					const FStructProperty* Property = CastField<const FStructProperty>(Struct->FindPropertyByName("StructWithEnum"));
+					AITEST_TRUE(TEXT("StructProperty"), Property != nullptr);
+					CA_ASSUME(Property); // Suppress code analyzer warning C6011
+					FStructWithEnum StructValue;
+					Property->GetValue_InContainer(Memory, &StructValue);
+					AITEST_TRUE(TEXT("StructWithEnum"), StructValue.CppClassEnumByte == ECppClassEnumByte::B);
+					AITEST_TRUE(TEXT("StructWithEnum"), StructValue.CppClassEnumInt == ECppClassEnumInt::B);
+					AITEST_TRUE(TEXT("StructWithEnum"), StructValue.NamespaceEnumByte == ENamespaceEnumByte::B);
+					AITEST_TRUE(TEXT("StructWithEnum"), StructValue.EnumByte == EEnumByte::EEnumByte_B);
+				}
+				{
+					using namespace UE::StateTree::Test;
+					const FEnumProperty* EnumProperty = CastField<const FEnumProperty>(Struct->FindPropertyByName("CppClassEnumByte"));
+					AITEST_TRUE(TEXT("CppClassEnumByte"), EnumProperty != nullptr);
+					CA_ASSUME(EnumProperty); // Suppress code analyzer warning C6011
+					AITEST_TRUE(TEXT("FByteProperty"), CastField<const FByteProperty>(EnumProperty->GetUnderlyingProperty()) != nullptr);
+					ECppClassEnumByte Value;
+					EnumProperty->GetValue_InContainer(Memory, &Value);
+					AITEST_TRUE(TEXT("Value"), Value == ECppClassEnumByte::B);
+				}
+				//@todo doesn't support uint32
+				//{
+				//	using namespace UE::StateTree::Test;
+				//	const FEnumProperty* EnumProperty = CastField<const FEnumProperty>(Struct->FindPropertyByName("CppClassEnumInt"));
+				//	AITEST_TRUE(TEXT("EnumProperty"), EnumProperty != nullptr);
+				//	AITEST_TRUE(TEXT("FUInt32Property"), CastField<const FUInt32Property>(EnumProperty->GetUnderlyingProperty()) != nullptr);
+				//	ECppClassEnumInt Value;
+				//	EnumProperty->GetValue_InContainer(Memory, &Value);
+				//	AITEST_TRUE(TEXT("Value"), Value == ECppClassEnumInt::B);
+				//}
+				{
+					using namespace UE::StateTree::Test;
+					const FProperty* Property = CastField<const FByteProperty>(Struct->FindPropertyByName("NamespaceEnumByte"));
+					AITEST_TRUE(TEXT("NamespaceEnumByte"), Property != nullptr);
+					CA_ASSUME(Property); // Suppress code analyzer warning C6011
+					TEnumAsByte<ENamespaceEnumByte::Type> Value;
+					Property->GetValue_InContainer(Memory, &Value);
+					AITEST_TRUE(TEXT("Value"), Value == ENamespaceEnumByte::B);
+				}
+				{
+					using namespace UE::StateTree::Test;
+					const FProperty* Property = CastField<const FByteProperty>(Struct->FindPropertyByName("EnumByte"));
+					AITEST_TRUE(TEXT("EnumByte"), Property != nullptr);
+					CA_ASSUME(Property); // Suppress code analyzer warning C6011
+					TEnumAsByte<EEnumByte> Value;
+					Property->GetValue_InContainer(Memory, &Value);
+					AITEST_TRUE(TEXT("Value"), Value == EEnumByte::EEnumByte_B);
+				}
+				{
+					using namespace UE::StateTree::Test;
+					const FStructProperty* Property = CastField<const FStructProperty>(Struct->FindPropertyByName("StructWithEnumRef1"));
+					AITEST_TRUE(TEXT("StructWithEnumRef1 property"), Property != nullptr);
+					CA_ASSUME(Property); // Suppress code analyzer warning C6011
+					FStructWithEnumRef Value;
+					Property->GetValue_InContainer(Memory, &Value);
+
+					ECppClassEnumByte* CppClassEnumByte = UE::StateTree::PropertyRefHelpers::GetMutablePtrToProperty<ECppClassEnumByte>(
+						Value.CppClassEnumByte.GetInternalPropertyRef(),
+						InstanceData.GetMutableStorage(),
+						nullptr,
+						InstanceData.GetExecutionState()->ActiveFrames[0]);
+					AITEST_TRUE(TEXT("CppClassEnumByte"), CppClassEnumByte != nullptr);
+					CA_ASSUME(CppClassEnumByte); // Suppress code analyzer warning C6011
+					AITEST_TRUE(TEXT("Value"), *CppClassEnumByte == ECppClassEnumByte::B);
+
+					//@todo PropertyBag doesn't support enum > int8. It always generates an enum of int8. That can creates bad memory access.
+					//ECppClassEnumInt* CppClassEnumInt = UE::StateTree::PropertyRefHelpers::GetMutablePtrToProperty<ECppClassEnumInt>(
+					//	Value.CppClassEnumInt.GetInternalPropertyRef(),
+					//	InstanceData.GetMutableStorage(),
+					//	nullptr,
+					//	InstanceData.GetExecutionState()->ActiveFrames[0]);
+					//AITEST_TRUE(TEXT("CppClassEnumInt"), CppClassEnumInt != nullptr);
+					//AITEST_TRUE(TEXT("Value"), *CppClassEnumInt == ECppClassEnumInt::B);
+
+					TEnumAsByte<ENamespaceEnumByte::Type>* NamespaceEnumByte = UE::StateTree::PropertyRefHelpers::GetMutablePtrToProperty<TEnumAsByte<ENamespaceEnumByte::Type>>(
+						Value.NamespaceEnumByte.GetInternalPropertyRef(),
+						InstanceData.GetMutableStorage(),
+						nullptr,
+						InstanceData.GetExecutionState()->ActiveFrames[0]); 
+					AITEST_TRUE(TEXT("NamespaceEnumByte"), NamespaceEnumByte != nullptr);
+					CA_ASSUME(NamespaceEnumByte); // Suppress code analyzer warning C6011
+					AITEST_TRUE(TEXT("Value"), *NamespaceEnumByte == ENamespaceEnumByte::B);
+
+					TEnumAsByte<EEnumByte>* EnumByte = UE::StateTree::PropertyRefHelpers::GetMutablePtrToProperty<TEnumAsByte<EEnumByte>>(
+						Value.EnumByte.GetInternalPropertyRef(),
+						InstanceData.GetMutableStorage(),
+						nullptr,
+						InstanceData.GetExecutionState()->ActiveFrames[0]); 
+					AITEST_TRUE(TEXT("EnumByte"), EnumByte != nullptr);
+					CA_ASSUME(EnumByte); // Suppress code analyzer warning C6011
+					AITEST_TRUE(TEXT("Value"), *EnumByte == EEnumByte::EEnumByte_B);
+				}
+				{
+					using namespace UE::StateTree::Test;
+					const FStructProperty* Property = CastField<const FStructProperty>(Struct->FindPropertyByName("StructWithEnumRef2"));
+					AITEST_TRUE(TEXT("StructWithEnumRef2 property"), Property != nullptr);
+					CA_ASSUME(Property); // Suppress code analyzer warning C6011
+					FStructWithEnumRef Value;
+					Property->GetValue_InContainer(Memory, &Value);
+
+					ECppClassEnumByte* CppClassEnumByte = UE::StateTree::PropertyRefHelpers::GetMutablePtrToProperty<ECppClassEnumByte>(
+						Value.CppClassEnumByte.GetInternalPropertyRef(),
+						InstanceData.GetMutableStorage(),
+						nullptr,
+						InstanceData.GetExecutionState()->ActiveFrames[0]);
+					AITEST_TRUE(TEXT("CppClassEnumByte"), CppClassEnumByte != nullptr);
+					CA_ASSUME(CppClassEnumByte); // Suppress code analyzer warning C6011
+					AITEST_TRUE(TEXT("Value"), *CppClassEnumByte == ECppClassEnumByte::B);
+
+					ECppClassEnumInt* CppClassEnumInt = UE::StateTree::PropertyRefHelpers::GetMutablePtrToProperty<ECppClassEnumInt>(
+						Value.CppClassEnumInt.GetInternalPropertyRef(),
+						InstanceData.GetMutableStorage(),
+						nullptr,
+						InstanceData.GetExecutionState()->ActiveFrames[0]);
+					AITEST_TRUE(TEXT("CppClassEnumInt"), CppClassEnumInt != nullptr);
+					CA_ASSUME(CppClassEnumInt); // Suppress code analyzer warning C6011
+					AITEST_TRUE(TEXT("Value"), *CppClassEnumInt == ECppClassEnumInt::B);
+
+					TEnumAsByte<ENamespaceEnumByte::Type>* NamespaceEnumByte = UE::StateTree::PropertyRefHelpers::GetMutablePtrToProperty<TEnumAsByte<ENamespaceEnumByte::Type>>(
+						Value.NamespaceEnumByte.GetInternalPropertyRef(),
+						InstanceData.GetMutableStorage(),
+						nullptr,
+						InstanceData.GetExecutionState()->ActiveFrames[0]);
+					AITEST_TRUE(TEXT("NamespaceEnumByte"), NamespaceEnumByte != nullptr);
+					AITEST_TRUE(TEXT("Value"), *NamespaceEnumByte == ENamespaceEnumByte::B);
+
+					TEnumAsByte<EEnumByte>* EnumByte = UE::StateTree::PropertyRefHelpers::GetMutablePtrToProperty<TEnumAsByte<EEnumByte>>(
+						Value.EnumByte.GetInternalPropertyRef(),
+						InstanceData.GetMutableStorage(),
+						nullptr,
+						InstanceData.GetExecutionState()->ActiveFrames[0]);
+					AITEST_TRUE(TEXT("EnumByte"), EnumByte != nullptr);
+					CA_ASSUME(EnumByte); // Suppress code analyzer warning C6011
+					AITEST_TRUE(TEXT("Value"), *EnumByte == EEnumByte::EEnumByte_B);
+				}
+			}
+			{
+				FTestStateTreeExecutionContext Exec(StateTree, StateTree, InstanceData);
+				Exec.Stop();
+			}
+		}
+
+		return true;
+	}
+};
+IMPLEMENT_STATE_TREE_INSTANT_TEST(FStateTreeTest_BindingEnum, "System.StateTree.Binding.Enum");
+
 struct FStateTreeTest_OutputBinding : FStateTreeTestBase
 {
 	virtual bool InstantTest() override
diff --git a/Engine/Plugins/Runtime/StateTree/Source/StateTreeTestSuite/Private/StateTreeTestTypes.h b/Engine/Plugins/Runtime/StateTree/Source/StateTreeTestSuite/Private/StateTreeTestTypes.h
--- a/Engine/Plugins/Runtime/StateTree/Source/StateTreeTestSuite/Private/StateTreeTestTypes.h
+++ b/Engine/Plugins/Runtime/StateTree/Source/StateTreeTestSuite/Private/StateTreeTestTypes.h
@@ -158,6 +158,74 @@ struct FTestStateTreeExecutionContext : public FStateTreeExecutionContext
 	
 };
 
+namespace UE::StateTree::Test
+{
+	UENUM()
+	enum class ECppClassEnumByte : uint8
+	{
+		A = 0,
+		B = 0x11,
+		C = 0x22,
+	};
+
+	UENUM()
+	enum class ECppClassEnumInt : uint32
+	{
+		A = 0,
+		B = 0xF333F333,
+		C = 0xF444F444,
+	};
+
+	UENUM()
+	namespace ENamespaceEnumByte
+	{
+		enum Type : uint8
+		{
+			A = 0,
+			B = 0x55,
+			C = 0x66,
+		};
+	}
+
+	UENUM()
+	enum EEnumByte : uint8
+	{
+		EEnumByte_A = 0,
+		EEnumByte_B = 0x77,
+		EEnumByte_C = 0x88,
+	};
+
+	USTRUCT()
+	struct FStructWithEnum
+	{
+		GENERATED_BODY()
+
+		UPROPERTY()
+		ECppClassEnumByte CppClassEnumByte;
+		UPROPERTY()
+		ECppClassEnumInt CppClassEnumInt;
+		UPROPERTY()
+		TEnumAsByte<ENamespaceEnumByte::Type> NamespaceEnumByte;
+		UPROPERTY()
+		TEnumAsByte<EEnumByte> EnumByte;
+	};
+
+	USTRUCT()
+	struct FStructWithEnumRef
+	{
+		GENERATED_BODY()
+
+		UPROPERTY()
+		TStateTreePropertyRef<ECppClassEnumByte> CppClassEnumByte;
+		UPROPERTY()
+		TStateTreePropertyRef<ECppClassEnumInt> CppClassEnumInt;
+		UPROPERTY()
+		TStateTreePropertyRef<TEnumAsByte<ENamespaceEnumByte::Type>> NamespaceEnumByte;
+		UPROPERTY()
+		TStateTreePropertyRef<TEnumAsByte<EEnumByte>> EnumByte;
+	};
+}
+
 USTRUCT()
 struct FTestEval_AInstanceData
 {
diff --git a/Engine/Plugins/Runtime/StateTree/Source/StateTreeTestSuite/StateTreeTestSuite.Build.cs b/Engine/Plugins/Runtime/StateTree/Source/StateTreeTestSuite/StateTreeTestSuite.Build.cs
--- a/Engine/Plugins/Runtime/StateTree/Source/StateTreeTestSuite/StateTreeTestSuite.Build.cs
+++ b/Engine/Plugins/Runtime/StateTree/Source/StateTreeTestSuite/StateTreeTestSuite.Build.cs
@@ -10,6 +10,8 @@ namespace UnrealBuildTool.Rules
 		{
 			CppCompileWarningSettings.UnsafeTypeCastWarningLevel = WarningLevel.Warning;
 
+			bAllowUETypesInNamespaces = true;
+
 			PublicIncludePaths.AddRange(
 				new string[] {
 				}
diff --git a/Engine/Source/Programs/Shared/EpicGames.UHT/Types/Properties/UhtStructProperty.cs b/Engine/Source/Programs/Shared/EpicGames.UHT/Types/Properties/UhtStructProperty.cs
--- a/Engine/Source/Programs/Shared/EpicGames.UHT/Types/Properties/UhtStructProperty.cs
+++ b/Engine/Source/Programs/Shared/EpicGames.UHT/Types/Properties/UhtStructProperty.cs
@@ -701,10 +701,11 @@ namespace EpicGames.UHT.Types
 			}
 
 			tokenReader.Require('<');
-
+			//@todo support TObjectptr, TOptional and other container.
 			bool isRefToArray = tokenReader.TryOptional("TArray");
+			bool isByteEnum = tokenReader.TryOptional("TEnumAsByte");
 
-			if (isRefToArray)
+			if (isRefToArray || isByteEnum)
 			{
 				tokenReader.Require('<');
 			}
@@ -745,9 +746,12 @@ namespace EpicGames.UHT.Types
 				rootSettings.MetaData.Add(RefTypeName, identifier.Join("::"));
 			}
 
-			if (isRefToArray)
+			if (isRefToArray || isByteEnum)
 			{
 				tokenReader.Require(">");
+			}
+			if (isRefToArray)
+			{
 				rootSettings.MetaData.Add(IsRefToArrayName, true);
 			}
 
