Fix issue where we use stale linked state tree override whenever there is a change overrides in the middle of the current context's execution

Since we are only allowed to set linked state tree overrides before we start context's execution, it is not possible to update it when an override changes sometime between the start and the end. To fix this we use a delegate to fetch current state tree overrides from external source.

#Synced-CL 50391333

diff --git a/Engine/Plugins/Runtime/StateTree/Source/StateTreeModule/Private/StateTreeExecutionContext.cpp b/Engine/Plugins/Runtime/StateTree/Source/StateTreeModule/Private/StateTreeExecutionContext.cpp
--- a/Engine/Plugins/Runtime/StateTree/Source/StateTreeModule/Private/StateTreeExecutionContext.cpp
+++ b/Engine/Plugins/Runtime/StateTree/Source/StateTreeModule/Private/StateTreeExecutionContext.cpp
@@ -1192,13 +1192,6 @@ void FStateTreeExecutionContext::SetLinkedStateTreeOverrides(FStateTreeReference
 		return;
 	}
 
-	FStateTreeExecutionState& Exec = GetExecState();
-	if (!ensureMsgf(Exec.CurrentPhase == EStateTreeUpdatePhase::Unset, TEXT("%hs can't be called while already in %s ('%s' using StateTree '%s')."),
-		__FUNCTION__, *UEnum::GetDisplayValueAsText(Exec.CurrentPhase).ToString(), *GetNameSafe(&Owner), *GetFullNameSafe(&RootStateTree)))
-	{
-		return;
-	}
-
 	bool bValid = true;
 
 	// Confirms that the overrides schema matches.
