[StateTree] removed deprecated method
#rnx

#Synced-CL 50651167

diff --git a/Engine/Plugins/Runtime/StateTree/Source/StateTreeModule/Private/StateTreeEvaluatorBase.cpp b/Engine/Plugins/Runtime/StateTree/Source/StateTreeModule/Private/StateTreeEvaluatorBase.cpp
--- a/Engine/Plugins/Runtime/StateTree/Source/StateTreeModule/Private/StateTreeEvaluatorBase.cpp
+++ b/Engine/Plugins/Runtime/StateTree/Source/StateTreeModule/Private/StateTreeEvaluatorBase.cpp
@@ -13,11 +13,5 @@ FString FStateTreeEvaluatorBase::GetDebugInfo(const FStateTreeReadOnlyExecutionC
 	Builder << TEXT("]\n");
 	return Builder.ToString();
 }
-
-// Deprecated
-void FStateTreeEvaluatorBase::AppendDebugInfoString(FString& DebugString, const FStateTreeExecutionContext& Context) const
-{
-	DebugString += FString::Printf(TEXT("[%s]\n"), *Name.ToString());
-}
 #endif // WITH_GAMEPLAY_DEBUGGER
 
diff --git a/Engine/Plugins/Runtime/StateTree/Source/StateTreeModule/Private/StateTreeTaskBase.cpp b/Engine/Plugins/Runtime/StateTree/Source/StateTreeModule/Private/StateTreeTaskBase.cpp
--- a/Engine/Plugins/Runtime/StateTree/Source/StateTreeModule/Private/StateTreeTaskBase.cpp
+++ b/Engine/Plugins/Runtime/StateTree/Source/StateTreeModule/Private/StateTreeTaskBase.cpp
@@ -15,11 +15,5 @@ FString FStateTreeTaskBase::GetDebugInfo(const FStateTreeReadOnlyExecutionContex
 	Builder << TEXT("]\n");
 	return Builder.ToString();
 }
-
-//Deprecated
-void FStateTreeTaskBase::AppendDebugInfoString(FString& DebugString, const FStateTreeExecutionContext& Context) const
-{
-	DebugString += FString::Printf(TEXT("[%s]\n"), *Name.ToString());
-}
-#endif
+#endif // WITH_GAMEPLAY_DEBUGGER
 
diff --git a/Engine/Plugins/Runtime/StateTree/Source/StateTreeModule/Public/StateTreeEvaluatorBase.h b/Engine/Plugins/Runtime/StateTree/Source/StateTreeModule/Public/StateTreeEvaluatorBase.h
--- a/Engine/Plugins/Runtime/StateTree/Source/StateTreeModule/Public/StateTreeEvaluatorBase.h
+++ b/Engine/Plugins/Runtime/StateTree/Source/StateTreeModule/Public/StateTreeEvaluatorBase.h
@@ -41,8 +41,10 @@ struct FStateTreeEvaluatorBase : public FStateTreeNodeBase
 #if WITH_GAMEPLAY_DEBUGGER
 	UE_API virtual FString GetDebugInfo(const FStateTreeReadOnlyExecutionContext& Context) const;
 
-	UE_DEPRECATED(5.6, "Use the version with the FStateTreeReadOnlyExecutionContext.")
-	UE_API virtual void AppendDebugInfoString(FString& DebugString, const FStateTreeExecutionContext& Context) const;
+	UE_DEPRECATED(5.8, "Use the version with the FStateTreeReadOnlyExecutionContext.")
+	virtual void AppendDebugInfoString(FString& DebugString, const FStateTreeExecutionContext& Context) const final
+	{
+	}
 #endif // WITH_GAMEPLAY_DEBUGGER
 };
 
diff --git a/Engine/Plugins/Runtime/StateTree/Source/StateTreeModule/Public/StateTreeTaskBase.h b/Engine/Plugins/Runtime/StateTree/Source/StateTreeModule/Public/StateTreeTaskBase.h
--- a/Engine/Plugins/Runtime/StateTree/Source/StateTreeModule/Public/StateTreeTaskBase.h
+++ b/Engine/Plugins/Runtime/StateTree/Source/StateTreeModule/Public/StateTreeTaskBase.h
@@ -101,9 +101,6 @@ struct FStateTreeTaskBase : public FStateTreeNodeBase
 
 #if WITH_GAMEPLAY_DEBUGGER
 	UE_API virtual FString GetDebugInfo(const FStateTreeReadOnlyExecutionContext& Context) const;
-
-	UE_DEPRECATED(5.6, "Use the version with the FStateTreeReadOnlyExecutionContext.")
-	UE_API virtual void AppendDebugInfoString(FString& DebugString, const FStateTreeExecutionContext& Context) const final;
 #endif
 
 	/**
