[StateTree] Add default value for FStructWithEnum used by StateTreeTestSuite.
#rnx

#Synced-CL 50769645

diff --git a/Engine/Plugins/Runtime/StateTree/Source/StateTreeTestSuite/Private/StateTreeTestTypes.h b/Engine/Plugins/Runtime/StateTree/Source/StateTreeTestSuite/Private/StateTreeTestTypes.h
--- a/Engine/Plugins/Runtime/StateTree/Source/StateTreeTestSuite/Private/StateTreeTestTypes.h
+++ b/Engine/Plugins/Runtime/StateTree/Source/StateTreeTestSuite/Private/StateTreeTestTypes.h
@@ -201,13 +201,13 @@ namespace UE::StateTree::Test
 		GENERATED_BODY()
 
 		UPROPERTY()
-		ECppClassEnumByte CppClassEnumByte;
+		ECppClassEnumByte CppClassEnumByte = ECppClassEnumByte::A;
 		UPROPERTY()
-		ECppClassEnumInt CppClassEnumInt;
+		ECppClassEnumInt CppClassEnumInt = ECppClassEnumInt::A;
 		UPROPERTY()
-		TEnumAsByte<ENamespaceEnumByte::Type> NamespaceEnumByte;
+		TEnumAsByte<ENamespaceEnumByte::Type> NamespaceEnumByte = ENamespaceEnumByte::A;
 		UPROPERTY()
-		TEnumAsByte<EEnumByte> EnumByte;
+		TEnumAsByte<EEnumByte> EnumByte = EEnumByte_A;
 	};
 
 	USTRUCT()
