[StateTree] Capture condition / task / utility names in ST rewind debugger
Add IStateTreeDebugInfoProvider to provide editor data driven debug info when possible, with default runtime only debug info as fallback.

#tests editor, PIE

#Synced-CL 50714302

diff --git a/Engine/Plugins/Runtime/StateTree/Source/StateTreeEditorModule/Private/StateTreeEditorModule.cpp b/Engine/Plugins/Runtime/StateTree/Source/StateTreeEditorModule/Private/StateTreeEditorModule.cpp
--- a/Engine/Plugins/Runtime/StateTree/Source/StateTreeEditorModule/Private/StateTreeEditorModule.cpp
+++ b/Engine/Plugins/Runtime/StateTree/Source/StateTreeEditorModule/Private/StateTreeEditorModule.cpp
@@ -37,6 +37,7 @@
 #include "StateTreeEditorDataExtension.h"
 #include "StateTreeEditorSchema.h"
 #include "StateTreeEditorStyle.h"
+#include "StateTreeModule.h"
 #include "StateTreeNodeClassCache.h"
 #include "StateTreePropertyFunctionBase.h"
 #include "Templates/GuardValueAccessors.h"
@@ -83,6 +84,78 @@ namespace UE::StateTree::Editor
 
 }; // UE::StateTree::Editor
 
+
+namespace UE::StateTree::Editor::Private
+{
+
+#if WITH_STATETREE_TRACE_DEBUGGER
+	class FStateTreeEditorDebugInfoProvider : public IStateTreeDebugInfoProvider
+	{
+		virtual FText GetNodeDescription(TNotNull<const UStateTree*> StateTree, FStateTreeIndex16 NodeIndex) const override
+		{
+			const FStateTreeNodeBase* Node = StateTree->GetNode(NodeIndex.Get()).GetPtr<FStateTreeNodeBase>();
+
+			if (Node)
+			{
+				if (UStateTreeEditorData* EditorData = Cast<UStateTreeEditorData>(StateTree->EditorData))
+				{
+					FConstStructView StructView;
+
+					if (Node->InstanceDataHandle.GetSource() == EStateTreeDataSourceType::EvaluationScopeInstanceData
+						|| Node->InstanceDataHandle.GetSource() == EStateTreeDataSourceType::EvaluationScopeInstanceDataObject)
+					{
+						StructView = StateTree->GetDefaultEvaluationScopeInstanceData().GetStruct(Node->InstanceTemplateIndex.Get());
+					}
+					else if (Node->InstanceDataHandle.GetSource() == EStateTreeDataSourceType::SharedInstanceData
+						|| Node->InstanceDataHandle.GetSource() == EStateTreeDataSourceType::SharedInstanceDataObject)
+					{
+						StructView = StateTree->GetSharedInstanceData()->GetStruct(Node->InstanceTemplateIndex.Get());
+					}
+					else
+					{
+						StructView = StateTree->GetDefaultInstanceData().GetStruct(Node->InstanceTemplateIndex.Get());
+					}
+
+					if (StructView.IsValid())
+					{
+						const FStateTreeDataView StateTreeDataView = FStateTreeDataView(StructView.GetScriptStruct(), const_cast<uint8*>(StructView.GetMemory()));
+						const FGuid NodeID = StateTree->GetNodeIdFromIndex(NodeIndex);
+						const FStateTreeBindingLookup BindingLookup(EditorData);
+
+						return Node->GetDescription(NodeID, StateTreeDataView, BindingLookup, EStateTreeNodeFormatting::Text);
+					}
+				}
+			}
+
+			return FText::FromString(FString::Printf(TEXT("%s"), Node ? *Node->Name.ToString() : *("Invalid Node Index: " + LexToString(NodeIndex.Get()))));
+		}
+	};
+
+	bool bEnableEditorDebugInfoProvider = true;
+
+	static void HandleEnableEditorDebugInfoProvider(IConsoleVariable*)
+	{
+		if (bEnableEditorDebugInfoProvider)
+		{
+			IStateTreeModule::Get().SetDebugInfoProvider(MakeShared<UE::StateTree::Editor::Private::FStateTreeEditorDebugInfoProvider>());
+		}
+		else
+		{
+			IStateTreeModule::Get().SetDebugInfoProvider(nullptr);
+		}
+	};
+
+	FAutoConsoleVariableRef CVarEnableEditorDebugInfoProvider(
+		TEXT("StateTree.Debugger.EnableEditorDebugInfoProvider"),
+		bEnableEditorDebugInfoProvider,
+		TEXT("True means we will use editor info for runtime debugging views."),
+		FConsoleVariableDelegate::CreateStatic(HandleEnableEditorDebugInfoProvider)
+	);
+#endif // WITH_STATETREE_TRACE_DEBUGGER
+
+} // namespace UE::StateTree::Editor::Private
+
+
 FStateTreeEditorModule& FStateTreeEditorModule::GetModule()
 {
 	return FModuleManager::LoadModuleChecked<FStateTreeEditorModule>("StateTreeEditorModule");
@@ -120,6 +193,11 @@ void FStateTreeEditorModule::StartupModule()
 
 	RewindDebuggerTrackCreator = MakePimpl<UE::StateTreeDebugger::FRewindDebuggerTrackCreator>();
 	IModularFeatures::Get().RegisterModularFeature(RewindDebugger::IRewindDebuggerTrackCreator::ModularFeatureName, RewindDebuggerTrackCreator.Get());
+
+	if (UE::StateTree::Editor::Private::bEnableEditorDebugInfoProvider)
+	{
+		IStateTreeModule::Get().SetDebugInfoProvider(MakeShared<UE::StateTree::Editor::Private::FStateTreeEditorDebugInfoProvider>());
+	}
 #endif // WITH_STATETREE_TRACE_DEBUGGER
 
 	MenuExtensibilityManager = MakeShareable(new FExtensibilityManager);
diff --git a/Engine/Plugins/Runtime/StateTree/Source/StateTreeModule/Private/Debugger/StateTreeTraceTypes.cpp b/Engine/Plugins/Runtime/StateTree/Source/StateTreeModule/Private/Debugger/StateTreeTraceTypes.cpp
--- a/Engine/Plugins/Runtime/StateTree/Source/StateTreeModule/Private/Debugger/StateTreeTraceTypes.cpp
+++ b/Engine/Plugins/Runtime/StateTree/Source/StateTreeModule/Private/Debugger/StateTreeTraceTypes.cpp
@@ -2,7 +2,9 @@
 
 #include "Debugger/StateTreeTraceTypes.h"
 
+
 #include "StateTree.h"
+#include "StateTreeModule.h"
 #include "StateTreeNodeBase.h"
 
 #include UE_INLINE_GENERATED_CPP_BY_NAME(StateTreeTraceTypes)
@@ -299,9 +301,7 @@ FString FStateTreeTraceTaskEvent::ToFullString(const UStateTree& StateTree) cons
 
 FString FStateTreeTraceTaskEvent::GetValueString(const UStateTree& StateTree) const
 {
-	const FConstStructView NodeView = Index.IsValid() ? StateTree.GetNode(Index.Get()) : FConstStructView();
-	const FStateTreeNodeBase* Node = NodeView.GetPtr<const FStateTreeNodeBase>();
-	return FString::Printf(TEXT("%s"), Node != nullptr ? *Node->Name.ToString() : *LexToString(Index.Get()));
+	return IStateTreeModule::Get().GetDebugInfoProvider()->GetNodeDescription(&StateTree, Index).ToString();
 }
 
 FString FStateTreeTraceTaskEvent::GetTypeString(const UStateTree& StateTree) const
@@ -320,9 +320,7 @@ FString FStateTreeTraceEvaluatorEvent::ToFullString(const UStateTree& StateTree)
 
 FString FStateTreeTraceEvaluatorEvent::GetValueString(const UStateTree& StateTree) const
 {
-	const FConstStructView NodeView = Index.IsValid() ? StateTree.GetNode(Index.Get()) : FConstStructView();
-	const FStateTreeNodeBase* Node = NodeView.GetPtr<const FStateTreeNodeBase>();
-	return FString::Printf(TEXT("%s"), Node != nullptr ? *Node->Name.ToString() : *LexToString(Index.Get()));
+	return IStateTreeModule::Get().GetDebugInfoProvider()->GetNodeDescription(&StateTree, Index).ToString();
 }
 
 FString FStateTreeTraceEvaluatorEvent::GetTypeString(const UStateTree& StateTree) const
@@ -341,10 +339,7 @@ FString FStateTreeTraceConditionEvent::ToFullString(const UStateTree& StateTree)
 
 FString FStateTreeTraceConditionEvent::GetValueString(const UStateTree& StateTree) const
 {
-	const FConstStructView NodeView = Index.IsValid() ? StateTree.GetNode(Index.Get()) : FConstStructView();
-	const FStateTreeNodeBase* Node = NodeView.GetPtr<const FStateTreeNodeBase>();
-
-	return FString::Printf(TEXT("%s"), Node != nullptr ? *Node->Name.ToString() : *LexToString(Index.Get()));
+	return IStateTreeModule::Get().GetDebugInfoProvider()->GetNodeDescription(&StateTree, Index).ToString();
 }
 
 FString FStateTreeTraceConditionEvent::GetTypeString(const UStateTree& StateTree) const
diff --git a/Engine/Plugins/Runtime/StateTree/Source/StateTreeModule/Private/StateTreeModule.cpp b/Engine/Plugins/Runtime/StateTree/Source/StateTreeModule/Private/StateTreeModule.cpp
--- a/Engine/Plugins/Runtime/StateTree/Source/StateTreeModule/Private/StateTreeModule.cpp
+++ b/Engine/Plugins/Runtime/StateTree/Source/StateTreeModule/Private/StateTreeModule.cpp
@@ -48,6 +48,18 @@ namespace UE::StateTree::Debug::Private
 		bRuntimeValidationInstanceDataGC,
 		TEXT("Test after each GC if nodes were properly GCed.")
 	);
+
+#if WITH_STATETREE_TRACE_DEBUGGER
+	class FStateTreeDebugInfoProvider : public IStateTreeDebugInfoProvider
+	{
+		virtual FText GetNodeDescription(TNotNull<const UStateTree*> StateTree, FStateTreeIndex16 NodeIndex) const override
+		{
+			const FStateTreeNodeBase* Node = StateTree->GetNode(NodeIndex.Get()).GetPtr<FStateTreeNodeBase>();
+			return FText::FromString(FString::Printf(TEXT("%s"), Node ? *Node->Name.ToString() : *("Invalid Node Index: " + LexToString(NodeIndex.Get()))));
+		}
+	};
+#endif // WITH_STATETREE_TRACE_DEBUGGER
+
 } // namespace UE::StateTree::Debug::Private
 #endif //WITH_STATETREE_DEBUG
 
@@ -60,6 +72,22 @@ UE::Trace::FStoreClient* FStateTreeModule::GetStoreClient()
 	}
 	return StoreClient.Get();
 }
+
+TNotNull<const IStateTreeDebugInfoProvider*> FStateTreeModule::GetDebugInfoProvider()
+{
+	if (!DebugInfoProvider)
+	{
+		DebugInfoProvider = MakeShared<UE::StateTree::Debug::Private::FStateTreeDebugInfoProvider>();
+	}
+
+	return DebugInfoProvider.Get();
+}
+
+void FStateTreeModule::SetDebugInfoProvider(TSharedPtr<IStateTreeDebugInfoProvider> InDebugInfoProvider)
+{
+	DebugInfoProvider = InDebugInfoProvider;
+}
+
 #endif // WITH_STATETREE_TRACE_DEBUGGER
 
 FStateTreeModule::FStateTreeModule()
@@ -158,12 +186,15 @@ void FStateTreeModule::ShutdownModule()
 	}
 
 	IModularFeatures::Get().UnregisterModularFeature(TraceServices::ModuleFeatureName, &StateTreeTraceModule);
+
+	DebugInfoProvider = nullptr;
 #endif // WITH_STATETREE_TRACE_DEBUGGER
 
 #if UE_WITH_STATETREE_CRASHREPORTER
 	UE::StateTree::FCrashReporterHandler::Unregister();
 #endif
 
+
 }
 
 bool FStateTreeModule::StartTraces(int32& OutTraceId)
diff --git a/Engine/Plugins/Runtime/StateTree/Source/StateTreeModule/Private/StateTreeModuleImpl.h b/Engine/Plugins/Runtime/StateTree/Source/StateTreeModule/Private/StateTreeModuleImpl.h
--- a/Engine/Plugins/Runtime/StateTree/Source/StateTreeModule/Private/StateTreeModuleImpl.h
+++ b/Engine/Plugins/Runtime/StateTree/Source/StateTreeModule/Private/StateTreeModuleImpl.h
@@ -56,6 +56,12 @@ private:
 	TUniquePtr<UE::Trace::FStoreClient> StoreClient;
 
 	FStateTreeTraceModule StateTreeTraceModule;
+
+	virtual TNotNull<const IStateTreeDebugInfoProvider*> GetDebugInfoProvider() override;
+	virtual void SetDebugInfoProvider(TSharedPtr<IStateTreeDebugInfoProvider> InDebugInfoProvider) override;
+
+	TSharedPtr<IStateTreeDebugInfoProvider> DebugInfoProvider;
+
 #endif // WITH_STATETREE_TRACE_DEBUGGER
 
 #if WITH_STATETREE_TRACE
diff --git a/Engine/Plugins/Runtime/StateTree/Source/StateTreeModule/Public/StateTreeModule.h b/Engine/Plugins/Runtime/StateTree/Source/StateTreeModule/Public/StateTreeModule.h
--- a/Engine/Plugins/Runtime/StateTree/Source/StateTreeModule/Public/StateTreeModule.h
+++ b/Engine/Plugins/Runtime/StateTree/Source/StateTreeModule/Public/StateTreeModule.h
@@ -2,13 +2,20 @@
 
 #pragma  once
 
+#include "Misc/NotNull.h"
 #include "Modules/ModuleManager.h"
 
+class IStateTreeDebugInfoProvider;
+class UStateTree;
+struct FStateTreeIndex16;
+
 namespace UE::Trace
 {
 	class FStoreClient;
 }
 
+#define UE_API STATETREEMODULE_API
+
 /**
 * The public interface to this module
 */
@@ -61,9 +68,36 @@ public:
 	virtual bool IsTracing() const = 0;
 
 #if WITH_STATETREE_TRACE_DEBUGGER
+	/**
+	 * Returns debug info provider. Creates on get if needed.
+	 * @return ptr to debug info provider to query against. 
+	 */
+	virtual TNotNull<const IStateTreeDebugInfoProvider*> GetDebugInfoProvider() = 0;
+
+	/**
+	 * Sets debug info provider.
+	 * @param ptr to debug info provider to query against. 
+	 */
+	virtual void SetDebugInfoProvider(TSharedPtr<IStateTreeDebugInfoProvider> InDebugInfoProvider) = 0;
+
 	/**
 	 * Gets the store client.
 	 */
 	virtual UE::Trace::FStoreClient* GetStoreClient() = 0;
 #endif // WITH_STATETREE_TRACE_DEBUGGER
 };
+
+
+/**
+ * Allows runtime to use editor data for debugging wben available & fallback to runtime only if needed.
+ */
+class IStateTreeDebugInfoProvider
+{
+public:
+	virtual ~IStateTreeDebugInfoProvider() = default;
+
+	/** Get text desc for the node */
+	virtual FText GetNodeDescription(TNotNull<const UStateTree*> StateTree, FStateTreeIndex16 NodeIndex) const = 0;
+};
+
+#undef UE_API
