[State Tree] 5.6 deprecated functions and structs clean up for State Tree Property Binding

#Synced-CL 50700524

diff --git a/Engine/Plugins/Runtime/StateTree/Source/StateTreeEditorModule/Public/StateTreeEditorData.h b/Engine/Plugins/Runtime/StateTree/Source/StateTreeEditorModule/Public/StateTreeEditorData.h
--- a/Engine/Plugins/Runtime/StateTree/Source/StateTreeEditorModule/Public/StateTreeEditorData.h
+++ b/Engine/Plugins/Runtime/StateTree/Source/StateTreeEditorModule/Public/StateTreeEditorData.h
@@ -327,20 +327,6 @@ public:
 		EditorBindings.AddFunctionBinding(PropertyFunctionNodeStruct, SourcePathSegments, TargetPath);
 	}
 
-PRAGMA_DISABLE_DEPRECATION_WARNINGS
-	UE_DEPRECATED(5.6, "Use the version taking FPropertyBindingPath instead")
-	void AddPropertyBinding(const FStateTreePropertyPath& SourcePath, const FStateTreePropertyPath& TargetPath)
-	{
-		EditorBindings.AddBinding(SourcePath, TargetPath);
-	}
-
-	UE_DEPRECATED(5.6, "Use the version taking FPropertyBindingPath instead")
-	void AddPropertyBinding(const UScriptStruct* PropertyFunctionNodeStruct, TConstArrayView<FPropertyBindingPathSegment> SourcePathSegments, const FStateTreePropertyPath& TargetPath)
-	{
-		EditorBindings.AddFunctionBinding(PropertyFunctionNodeStruct, SourcePathSegments, TargetPath);
-	}
-PRAGMA_ENABLE_DEPRECATION_WARNINGS
-
 	/**
 	 * Adds property binding between two structs.
 	 */
diff --git a/Engine/Plugins/Runtime/StateTree/Source/StateTreeEditorModule/Public/StateTreeEditorPropertyBindings.h b/Engine/Plugins/Runtime/StateTree/Source/StateTreeEditorModule/Public/StateTreeEditorPropertyBindings.h
--- a/Engine/Plugins/Runtime/StateTree/Source/StateTreeEditorModule/Public/StateTreeEditorPropertyBindings.h
+++ b/Engine/Plugins/Runtime/StateTree/Source/StateTreeEditorModule/Public/StateTreeEditorPropertyBindings.h
@@ -71,89 +71,6 @@ struct FStateTreeEditorPropertyBindings : public FPropertyBindingBindingCollecti
 	 */
 	UE_API const FStateTreePropertyPathBinding* AddOutputBinding(const FPropertyBindingPath& InSourcePath, const FPropertyBindingPath& InTargetPath);
 
-PRAGMA_DISABLE_DEPRECATION_WARNINGS
-	/**
-	 * Adds binding between source and destination paths. Removes any bindings to TargetPath before adding the new one.
-	 * @param SourcePath Binding source property path.
-	 * @param TargetPath Binding target property path.
-	 */
-	UE_DEPRECATED(5.6, "Use AddBinding taking FPropertyBindingPath instead")
-	void AddPropertyBinding(const FStateTreePropertyPath& SourcePath, const FStateTreePropertyPath& TargetPath)
-	{
-		AddBinding(SourcePath, TargetPath);
-	}
-	
-	/**
-	 * Adds binding.
-	 * @param Binding Binding to be added.
-	*/
-	UE_DEPRECATED(5.6, "Use the version taking FPropertyBindingPath instead")
-	void AddPropertyBinding(const FStateTreePropertyPathBinding& Binding)
-	{
-	}
-
-	
-	UE_DEPRECATED(5.6, "Use AddFunctionBinding taking FPropertyBindingPath instead")
-	FStateTreePropertyPath AddFunctionPropertyBinding(const UScriptStruct* InPropertyFunctionNodeStruct, TConstArrayView<FPropertyBindingPathSegment> InSourcePathSegments, const FStateTreePropertyPath& InTargetPath)
-	{
-		return AddFunctionBinding(InPropertyFunctionNodeStruct, InSourcePathSegments, InTargetPath);
-	}
-
-	/**
-	 * Removes all bindings to target path.
-	 * @param TargetPath Target property path.
-	 */ 
-	UE_DEPRECATED(5.6, "Use RemoveBinding taking FPropertyBindingPath instead")
-	void RemovePropertyBindings(const FStateTreePropertyPath& TargetPath, ESearchMode SearchMode = ESearchMode::Exact)
-	{
-		RemoveBindings(TargetPath, SearchMode);
-	}
-	
-	/**
-	 * Has any binding to the target path.
-	 * @return True of the target path has any bindings.
-	 */
-	UE_DEPRECATED(5.6, "Use HasBinding taking FPropertyBindingPath instead")
-	bool HasPropertyBinding(const FStateTreePropertyPath& TargetPath, ESearchMode SearchMode = ESearchMode::Exact) const
-	{
-		return HasBinding(TargetPath, SearchMode);
-	}
-
-	UE_DEPRECATED(5.6, "Use FindBinding taking FPropertyBindingPath instead")
-	const FStateTreePropertyPathBinding* FindPropertyBinding(const FStateTreePropertyPath& TargetPath, ESearchMode SearchMode = ESearchMode::Exact) const
-	{
-		return static_cast<const FStateTreePropertyPathBinding*>(FindBinding(TargetPath, SearchMode));
-	}
-
-	/**
-	 * @return Source path for given target path, or null if binding does not exists.
-	 */
-	UE_DEPRECATED(5.6, "Use GetBindingSource taking FPropertyBindingPath instead")
-	const FStateTreePropertyPath* GetPropertyBindingSource(const FStateTreePropertyPath& TargetPath) const
-	{
-		return nullptr;
-	}
-
-	/**
-	 * Returns all pointers to bindings for a specified structs based in struct ID.
-	 * @param StructID ID of the struct to find bindings for.
-	 * @param OutBindings Bindings for specified struct.
-	 */
-	UE_DEPRECATED(5.6, "Use GetBindingsFor taking FPropertyBindingPath instead")
-	void GetPropertyBindingsFor(const FGuid StructID, TArray<const FStateTreePropertyPathBinding*>& OutBindings) const
-	{
-	}
-
-	/**
-	 * Removes bindings which do not point to valid structs IDs.
-	 * @param ValidStructs Set of struct IDs that are currently valid.
-	 */
-	UE_DEPRECATED(5.6, "Use RemoveInvalidBindings taking FPropertyBindingDataView instead")
-	void RemoveUnusedBindings(const TMap<FGuid, const FStateTreeDataView>& ValidStructs)
-	{
-	}
-PRAGMA_ENABLE_DEPRECATION_WARNINGS
-
 protected:
 	//~ Begin FPropertyBindingBindingCollection overrides
 	virtual FPropertyBindingBinding* AddBindingInternal(const FPropertyBindingPath& InSourcePath, const FPropertyBindingPath& InTargetPath) override;
@@ -175,47 +92,11 @@ class UStateTreeEditorPropertyBindingsOwner : public UPropertyBindingBindingColl
 	GENERATED_UINTERFACE_BODY()
 };
 
-/** Struct of Parameters used to Create a Property */
-struct UE_DEPRECATED(5.6, "Use FPropertyCreationDesc instead") FStateTreeEditorPropertyCreationDesc
-{
-	/** Property Bag Description of the Property to Create */
-	FPropertyBagPropertyDesc PropertyDesc;
-
-	/** Optional: property to copy into the new created property */
-	const FProperty* SourceProperty = nullptr;
-
-	/** Optional: container address of the property to copy */
-	const void* SourceContainerAddress = nullptr;
-};
-
 class IStateTreeEditorPropertyBindingsOwner : public IPropertyBindingBindingCollectionOwner
 {
 	GENERATED_BODY()
-
 public:
 
-	/**
-	 * Returns structs within the owner that are visible for target struct.
-	 * @param TargetStructID Target struct ID
-	 * @param OutStructDescs Result descriptors of the visible structs.
-	 */
-	UE_DEPRECATED(5.6, "Use version taking FPropertyBindingBindableStructDescriptor instead")
-	virtual void GetAccessibleStructs(const FGuid TargetStructID, TArray<FStateTreeBindableStructDesc>& OutStructDescs) const final
-	{
-	}
-
-	/**
-	 * Returns struct descriptor based on struct ID.
-	 * @param StructID Target struct ID
-	 * @param OutStructDesc Result descriptor.
-	 * @return True if struct found.
-	 */
-	UE_DEPRECATED(5.6, "Use version taking FPropertyBindingBindableStructDescriptor instead")
-	virtual bool GetStructByID(const FGuid StructID, FStateTreeBindableStructDesc& OutStructDesc) const final
-	{
-		return false;
-	}
-
 	/**
 	 * Finds a bindable context struct based on name and type.
 	 * @param ObjectType Object type to match
@@ -223,18 +104,6 @@ public:
 	 */
 	virtual FStateTreeBindableStructDesc FindContextData(const UStruct* ObjectType, const FString ObjectNameHint) const PURE_VIRTUAL(IStateTreeEditorPropertyBindingsOwner::FindContextData, return {}; );
 
-	/**
-	 * Returns data view based on struct ID.
-	 * @param StructID Target struct ID
-	 * @param OutDataView Result data view.
-	 * @return True if struct found.
-	 */
-	UE_DEPRECATED(5.6, "Use version taking FPropertyBindingDataView instead")
-	virtual bool GetDataViewByID(const FGuid StructID, FStateTreeDataView& OutDataView) const final
-	{
-		return false;
-	}
-
 	/** @return Pointer to editor property bindings. */
 	virtual FStateTreeEditorPropertyBindings* GetPropertyEditorBindings() PURE_VIRTUAL(IStateTreeEditorPropertyBindingsOwner::GetPropertyEditorBindings, return nullptr; );
 
@@ -242,18 +111,6 @@ public:
 	virtual const FStateTreeEditorPropertyBindings* GetPropertyEditorBindings() const PURE_VIRTUAL(IStateTreeEditorPropertyBindingsOwner::GetPropertyEditorBindings, return nullptr; );
 
 	virtual EStateTreeVisitor EnumerateBindablePropertyFunctionNodes(TFunctionRef<EStateTreeVisitor(const UScriptStruct* NodeStruct, const FStateTreeBindableStructDesc& Desc, const FStateTreeDataView Value)> InFunc) const PURE_VIRTUAL(IStateTreeEditorPropertyBindingsOwner::EnumerateBindablePropertyFunctionNodes, return static_cast<EStateTreeVisitor>(0); );
-
-PRAGMA_DISABLE_DEPRECATION_WARNINGS
-	/**
-	 * Creates the given properties in the property bag of the struct matching the given struct ID
-	 * @param StructID Target struct ID
-	 * @param InOutCreationDescs the descriptions of the properties to create. This is modified to update the property names that actually got created
-	 */
-	UE_DEPRECATED(5.6, "Use version taking FPropertyCreationDesc instead")
-	virtual void CreateParameters(const FGuid StructID, TArrayView<FStateTreeEditorPropertyCreationDesc> InOutCreationDescs) final
-	{
-	}
-PRAGMA_ENABLE_DEPRECATION_WARNINGS
 };
 
 // TODO: We should merge this with IStateTreeEditorPropertyBindingsOwner and FStateTreeEditorPropertyBindings.
diff --git a/Engine/Plugins/Runtime/StateTree/Source/StateTreeModule/Private/StateTreePropertyBindings.cpp b/Engine/Plugins/Runtime/StateTree/Source/StateTreeModule/Private/StateTreePropertyBindings.cpp
--- a/Engine/Plugins/Runtime/StateTree/Source/StateTreeModule/Private/StateTreePropertyBindings.cpp
+++ b/Engine/Plugins/Runtime/StateTree/Source/StateTreeModule/Private/StateTreePropertyBindings.cpp
@@ -337,11 +337,6 @@ const FPropertyBindingBinding* FStateTreePropertyBindings::FindBindingInternal(T
 }
 #endif // WITH_EDITOR
 
-bool FStateTreePropertyBindings::ResetObjects(const FStateTreeIndex16 TargetBatchIndex, FStateTreeDataView TargetStructView) const
-{
-	return Super::ResetObjects(TargetBatchIndex, TargetStructView);
-}
-
 const FStateTreePropertyAccess* FStateTreePropertyBindings::GetPropertyAccess(const FStateTreePropertyRef& InPropertyReference) const
 {
 	if (!InPropertyReference.GetRefAccessIndex().IsValid())
@@ -356,72 +351,3 @@ const FStateTreePropertyAccess* FStateTreePropertyBindings::GetPropertyAccess(co
 
 	return &PropertyAccesses[InPropertyReference.GetRefAccessIndex().Get()];
 }
-
-PRAGMA_DISABLE_DEPRECATION_WARNINGS
-int32 FStateTreePropertyBindings::GetSourceStructNum() const
-{
-	return SourceStructs.Num();
-}
-
-const FStateTreePropertyCopyBatch& FStateTreePropertyBindings::GetBatch(const FStateTreeIndex16 TargetBatchIndex) const
-{
-	check(TargetBatchIndex.IsValid());
-	static FStateTreePropertyCopyBatch Batch;
-	return Batch;
-}
-
-TConstArrayView<FStateTreePropertyCopy> FStateTreePropertyBindings::GetBatchCopies(const FStateTreeIndex16 TargetBatchIndex) const
-{
-	return GetBatchCopies(GetBatch(TargetBatchIndex));
-}
-
-TConstArrayView<FStateTreePropertyCopy> FStateTreePropertyBindings::GetBatchCopies(const FStateTreePropertyCopyBatch& Batch) const
-{
-	return {};
-}
-
-bool FStateTreePropertyBindings::ResolveCopyType(const FStateTreePropertyPathIndirection& SourceIndirection, const FStateTreePropertyPathIndirection& TargetIndirection, FStateTreePropertyCopy& OutCopy)
-{
-	return false;
-}
-
-bool FStateTreePropertyBindings::ResolveCopyType(const FPropertyBindingPathIndirection& SourceIndirection, const FPropertyBindingPathIndirection& TargetIndirection, FStateTreePropertyCopy& OutCopy)
-{
-	return false;
-}
-
-EStateTreePropertyAccessCompatibility FStateTreePropertyBindings::GetPropertyCompatibility(const FProperty* FromProperty, const FProperty* ToProperty)
-{
-	// Temporary using cast as this method is no longer used and will be removed so we don't care about matching values in both enums
-	//const UE::PropertyBinding::EPropertyCompatibility Value = UE::PropertyBinding::GetPropertyCompatibility(FromProperty, ToProperty);
-	//check(UEnum::GetValueAsName<EStateTreePropertyAccessCompatibility>((EStateTreePropertyAccessCompatibility)Value) ==
-	//	  UEnum::GetValueAsName<UE::PropertyBinding::EPropertyCompatibility>(Value));
-	return static_cast<EStateTreePropertyAccessCompatibility>(UE::PropertyBinding::GetPropertyCompatibility(FromProperty, ToProperty));
-}
-
-bool FStateTreePropertyBindings::ResolvePath(const UStruct* Struct, const FPropertyBindingPath& Path
-	, FStateTreePropertyIndirection& OutFirstIndirection, FPropertyBindingPathIndirection& OutLeafIndirection)
-{
-	return false;
-}
-
-const FStateTreeBindableStructDesc* FStateTreePropertyBindings::GetSourceDescByHandle(const FStateTreeDataHandle SourceDataHandle)
-{
-	return nullptr;
-}
-
-void FStateTreePropertyBindings::PerformCopy(const FStateTreePropertyCopy& Copy, uint8* SourceAddress, uint8* TargetAddress) const
-{
-}
-
-void FStateTreePropertyBindings::PerformResetObjects(const FStateTreePropertyCopy& Copy, uint8* TargetAddress) const
-{
-}
-
-uint8* FStateTreePropertyBindings::GetAddress(FStateTreeDataView InStructView, const FStateTreePropertyIndirection& FirstIndirection
-	, const FProperty* LeafProperty) const
-{
-	return nullptr;
-}
-
-PRAGMA_ENABLE_DEPRECATION_WARNINGS
diff --git a/Engine/Plugins/Runtime/StateTree/Source/StateTreeModule/Public/StateTreePropertyBindings.h b/Engine/Plugins/Runtime/StateTree/Source/StateTreeModule/Public/StateTreePropertyBindings.h
--- a/Engine/Plugins/Runtime/StateTree/Source/StateTreeModule/Public/StateTreePropertyBindings.h
+++ b/Engine/Plugins/Runtime/StateTree/Source/StateTreeModule/Public/StateTreePropertyBindings.h
@@ -56,73 +56,6 @@ namespace UE::StateTree
 	[[nodiscard]] UE_API bool AcceptTaskInstanceData(EStateTreeBindableStructSource Target);
 }
 
-
-/**
- * Describes how the copy should be performed.
- */
-enum class UE_DEPRECATED(5.6, "Use EPropertyCopyType instead")  EStateTreePropertyCopyType : uint8
-{
-	None,						// No copying
-	
-	CopyPlain,					// For plain old data types, we do a simple memcpy.
-	CopyComplex,				// For more complex data types, we need to call the properties copy function
-	CopyBool,					// Read and write properties using bool property helpers, as source/dest could be bitfield or boolean
-	CopyStruct,					// Use struct copy operation, as this needs to correctly handle CPP struct ops
-	CopyObject,					// Read and write properties using object property helpers, as source/dest could be regular/weak/soft etc.
-	CopyName,					// FName needs special case because its size changes between editor/compiler and runtime.
-	CopyFixedArray,				// Array needs special handling for fixed size TArrays
-
-	StructReference,			// Copies pointer to a source struct into a FStateTreeStructRef.
-
-	/* Promote the type during the copy */
-
-	/* Bool promotions */
-	PromoteBoolToByte,
-	PromoteBoolToInt32,
-	PromoteBoolToUInt32,
-	PromoteBoolToInt64,
-	PromoteBoolToFloat,
-	PromoteBoolToDouble,
-
-	/* Byte promotions */
-	PromoteByteToInt32,
-	PromoteByteToUInt32,
-	PromoteByteToInt64,
-	PromoteByteToFloat,
-	PromoteByteToDouble,
-
-	/* Int32 promotions */
-	PromoteInt32ToInt64,
-	PromoteInt32ToFloat,	// This is strictly sketchy because of potential data loss, but it is usually OK in the general case
-	PromoteInt32ToDouble,
-
-	/* UInt32 promotions */
-	PromoteUInt32ToInt64,
-	PromoteUInt32ToFloat,	// This is strictly sketchy because of potential data loss, but it is usually OK in the general case
-	PromoteUInt32ToDouble,
-
-	/* Float promotions */
-	PromoteFloatToInt32,
-	PromoteFloatToInt64,
-	PromoteFloatToDouble,
-
-	/* Double promotions */
-	DemoteDoubleToInt32,
-	DemoteDoubleToInt64,
-	DemoteDoubleToFloat,
-};
-
-/** Enum describing property compatibility */
-enum class UE_DEPRECATED(5.6, "Use UE::PropertyBinding::EPropertyCompatibility instead") EStateTreePropertyAccessCompatibility : uint8
-{
-	/** Properties are incompatible */
-	Incompatible,
-	/** Properties are directly compatible */
-	Compatible,	
-	/** Properties can be copied with a simple type promotion */
-	Promotable,
-};
-
 /**
  * Descriptor for a struct or class that can be a binding source or target.
  * Each struct has unique identifier, which is used to distinguish them, and name that is mostly for debugging and UI.
@@ -166,25 +99,6 @@ struct FStateTreeBindableStructDesc : public FPropertyBindingBindableStructDescr
 #endif
 };
 
-struct UE_DEPRECATED(5.6, "Use FPropertyBindingPathIndirection instead") FStateTreePropertyPathIndirection : FPropertyBindingPathIndirection
-{
-};
-
-/**
- * Representation of a property path used for property binding in StateTree.
- */
-struct UE_DEPRECATED(5.6, "Use FPropertyBindingPath instead") STATETREEMODULE_API FStateTreePropertyPath final : FPropertyBindingPath
-{
-	using FPropertyBindingPath::FPropertyBindingPath;
-
-	//~ Intentionally not explicit to facilitate introduction of base class FPropertyBindingPath
-	FStateTreePropertyPath(const FPropertyBindingPath& Other)
-		: FPropertyBindingPath(Other)
-	{
-	}
-};
-
-
 USTRUCT()
 struct UE_DEPRECATED(all, "Use FPropertyBindingPath instead.") FStateTreeEditorPropertyPath
 {
@@ -345,89 +259,6 @@ private:
 	FStateTreeDataHandle SourceDataHandle = FStateTreeDataHandle::Invalid;
 };
 
-PRAGMA_DISABLE_DEPRECATION_WARNINGS
-/**
- * Used internally.
- * Property indirection is a resolved property path segment, used for accessing properties in structs.
- */
-struct UE_DEPRECATED(5.6, "Use FPropertyBindingPropertyIndirection instead") STATETREEMODULE_API FStateTreePropertyIndirection
-{
-	/** Index in the array the property points at. */
-	FStateTreeIndex16 ArrayIndex = FStateTreeIndex16::Invalid;
-
-	/** Cached offset of the property */
-	uint16 Offset = 0;
-
-	/** Cached offset of the property */
-	FStateTreeIndex16 NextIndex = FStateTreeIndex16::Invalid;
-
-	/** Type of access/indirection. */
-	EPropertyBindingPropertyAccessType Type = EPropertyBindingPropertyAccessType::Offset;
-
-	/** Type of the struct or object instance in case the segment is pointing into an instanced data. */
-	TObjectPtr<const UStruct> InstanceStruct = nullptr;
-
-	/** Cached array property. */
-	const FArrayProperty* ArrayProperty = nullptr;
-};
-
-
-/**
- * Used internally.
- * Describes property copy, the property from source is copied into the property at the target.
- * Copy target struct is described in the property copy batch.
- */
-struct UE_DEPRECATED(5.6, "Use FPropertyBindingCopyInfo instead") STATETREEMODULE_API FStateTreePropertyCopy
-{
-	/** Source property access. */
-	FStateTreePropertyIndirection SourceIndirection;
-
-	/** Target property access. */
-	FStateTreePropertyIndirection TargetIndirection;
-
-	/** Cached pointer to the leaf property of the access. */
-	const FProperty* SourceLeafProperty = nullptr;
-
-	/** Cached pointer to the leaf property of the access. */
-	const FProperty* TargetLeafProperty = nullptr;
-
-	/** Type of the source data, used for validation. */
-	TObjectPtr<const UStruct> SourceStructType = nullptr;
-
-	/** Cached property element size * dim. */
-	int32 CopySize = 0;
-
-	/** Describes how to get the source data pointer for the copy. */
-	FStateTreeDataHandle SourceDataHandle = FStateTreeDataHandle::Invalid;
-
-	/** Type of the copy */
-	EStateTreePropertyCopyType Type = EStateTreePropertyCopyType::None;
-};
-
-
-/**
- * Describes a batch of property copies from many sources to one target struct.
- * Note: The batch is used to reference both bindings and copies (a binding turns into copy when resolved).
- */
-struct UE_DEPRECATED(5.6, "Use FPropertyBindingCopyInfoBatch instead") STATETREEMODULE_API FStateTreePropertyCopyBatch
-{
-	/** Expected target struct */
-	FStateTreeBindableStructDesc TargetStruct;
-
-	/** Index to first binding/copy. */
-	FStateTreeIndex16 BindingsBegin;
-
-	/** Index to one past the last binding/copy. */
-	FStateTreeIndex16 BindingsEnd;
-
-	/** Index to first property function. */
-	FStateTreeIndex16 PropertyFunctionsBegin;
-
-	/** Index to one past the last property function. */
-	FStateTreeIndex16 PropertyFunctionsEnd;
-};
-PRAGMA_ENABLE_DEPRECATION_WARNINGS
-
 /**
  * Describes access to referenced property.
  */
@@ -467,40 +298,10 @@ struct FStateTreePropertyBindings : public FPropertyBindingBindingCollection
 	 */
 	UE_API virtual void OnReset() override;
 
-PRAGMA_DISABLE_DEPRECATION_WARNINGS
-	/**
-	 * @return Number of source structs the copy expects.
-	 */
-	UE_DEPRECATED(5.6, "Use GetNumBindableStructDescriptors instead")
-	UE_API int32 GetSourceStructNum() const;
-
-	/**
-	 * Copies a property from Source to Target based on the provided Copy.
-	 * @param Copy Describes which parameter and how is copied.
-	 * @param SourceStructView Pointer and type for the source containing the property to be copied.
-	 * @param TargetStructView Pointer and type for the target containing the property to be copied.
-	 * @return true if the property was copied successfully.
-	 */
-	UE_DEPRECATED(5.6, "Use the overload taking FPropertyBindingCopyInfo instead")
-	UE_API bool CopyProperty(const FStateTreePropertyCopy& Copy, FStateTreeDataView SourceStructView, FStateTreeDataView TargetStructView) const;
-
-	/** @return copy batch at specified index. */
-	UE_DEPRECATED(5.6, "Use the overload taking FPropertyBindingCopyInfo instead")
-	UE_API const FStateTreePropertyCopyBatch& GetBatch(const FStateTreeIndex16 TargetBatchIndex) const;
-
-	/** @return All the property copies for a specific batch. */
-	UE_DEPRECATED(5.6, "Use the overload return FPropertyBindingCopyInfo and taking FPropertyBindingIndex16 instead")
-	UE_API TConstArrayView<FStateTreePropertyCopy> GetBatchCopies(const FStateTreeIndex16 TargetBatchIndex) const;
-
-	/** @return All the property copies for a specific batch. */
-	UE_DEPRECATED(5.6, "Use the overload taking FPropertyBindingCopyInfoBatch instead")
-	UE_API TConstArrayView<FStateTreePropertyCopy> GetBatchCopies(const FStateTreePropertyCopyBatch& Batch) const;
-UE_API PRAGMA_ENABLE_DEPRECATION_WARNINGS
-
 	/**
 	 * @return Referenced property access for provided PropertyRef.
 	 */
-	const FStateTreePropertyAccess* GetPropertyAccess(const FStateTreePropertyRef& Reference) const;
+	UE_API const FStateTreePropertyAccess* GetPropertyAccess(const FStateTreePropertyRef& Reference) const;
 
 	/**
 	 * Pointer to referenced property 
@@ -538,46 +339,7 @@ protected:
 	UE_API virtual const FPropertyBindingBinding* FindBindingInternal(TFunctionRef<bool(const FPropertyBindingBinding&)> InPredicate) const override;
 #endif // WITH_EDITOR
 	//~ End FPropertyBindingBindingCollection overrides
-
-PRAGMA_DISABLE_DEPRECATION_WARNINGS
-	UE_DEPRECATED(5.6, "Use the overload taking FPropertyBindingIndex16 and FPropertyBindingDataView instead")
-	UE_API bool ResetObjects(const FStateTreeIndex16 TargetBatchIndex, FStateTreeDataView TargetStructView) const;
-
-	/**
-	 * Resolves what kind of copy type to use between specified property indirections.
-	 * @param SourceIndirection Property path indirections of the copy source,
-	 * @param TargetIndirection Property path indirections of the copy target.
-	 * @param OutCopy Resulting copy type.
-	 * @return true if copy was resolved, or false if no copy could be resolved between paths.
-	 */
-	UE_DEPRECATED(5.6, "Use the overload taking FPropertyBindingCopyInfo instead")
-	[[nodiscard]] static UE_API bool ResolveCopyType(const FPropertyBindingPathIndirection& SourceIndirection, const FPropertyBindingPathIndirection& TargetIndirection, FStateTreePropertyCopy& OutCopy);
-
-	UE_DEPRECATED(5.6, "Use UE::PropertyBinding::GetPropertyCompatibility instead")
-	static UE_API EStateTreePropertyAccessCompatibility GetPropertyCompatibility(const FProperty* FromProperty, const FProperty* ToProperty);
-
-	UE_DEPRECATED(5.6, "Use the overload taking FPropertyBindingPathIndirection instead")
-	[[nodiscard]] static UE_API bool ResolveCopyType(const FStateTreePropertyPathIndirection& SourceIndirection, const FStateTreePropertyPathIndirection& TargetIndirection, FStateTreePropertyCopy& OutCopy);
-
 private:
-	UE_DEPRECATED(5.6, "Use the version taking FPropertyBindingPropertyIndirection instead.")
-	[[nodiscard]] UE_API bool ResolvePath(const UStruct* Struct, const FPropertyBindingPath& Path, FStateTreePropertyIndirection& OutFirstIndirection, FPropertyBindingPathIndirection& OutLeafIndirection);
-
-	UE_DEPRECATED(5.6, "Use GetBindableStructDescriptorFromHandle instead.")
-	UE_API const FStateTreeBindableStructDesc* GetSourceDescByHandle(const FStateTreeDataHandle SourceDataHandle);
-
-	UE_DEPRECATED(5.6, "Use the version taking FPropertyBindingCopyInfo instead.")
-	UE_API void PerformCopy(const FStateTreePropertyCopy& Copy, uint8* SourceAddress, uint8* TargetAddress) const;
-
-	UE_DEPRECATED(5.6, "Use the version taking FPropertyBindingCopyInfo instead.")
-	UE_API void PerformResetObjects(const FStateTreePropertyCopy& Copy, uint8* TargetAddress) const;
-
-	UE_DEPRECATED(5.6, "Use the version taking FPropertyBinding types instead.")
-	UE_API uint8* GetAddress(FStateTreeDataView InStructView, const FStateTreePropertyIndirection& FirstIndirection, const FProperty* LeafProperty) const;
-
-private:
-	PRAGMA_ENABLE_DEPRECATION_WARNINGS
-
 	/** Array of expected source structs. */
 	UPROPERTY()
 	TArray<FStateTreeBindableStructDesc> SourceStructs;
@@ -629,33 +391,6 @@ struct IStateTreeBindingLookup
 
 	/** @return Display name of binding source, or empty if binding does not exists. */
 	virtual FText GetBindingSourceDisplayName(const FPropertyBindingPath& InTargetPath, EStateTreeNodeFormatting Formatting = EStateTreeNodeFormatting::Text) const = 0;
-
-PRAGMA_DISABLE_DEPRECATION_WARNINGS
-
-	UE_DEPRECATED(5.6, "Use the version taking FPropertyBindingPath instead")
-	virtual const FStateTreePropertyPath* GetPropertyBindingSource(const FStateTreePropertyPath& InTargetPath) const final
-	{
-		return nullptr;
-	}
-
-	UE_DEPRECATED(5.6, "Use the version taking FPropertyBindingPath instead")
-	virtual FText GetPropertyPathDisplayName(const FStateTreePropertyPath& InPath, EStateTreeNodeFormatting Formatting = EStateTreeNodeFormatting::Text) const final
-	{
-		return FText::GetEmpty();
-	}
-
-	UE_DEPRECATED(5.6, "Use the version taking FPropertyBindingPath instead")
-	virtual FText GetBindingSourceDisplayName(const FStateTreePropertyPath& InTargetPath, EStateTreeNodeFormatting Formatting = EStateTreeNodeFormatting::Text) const final
-	{
-		return FText::GetEmpty();
-	}
-
-	UE_DEPRECATED(5.6, "Use the version taking FPropertyBindingPath instead")
-	virtual const FProperty* GetPropertyPathLeafProperty(const FStateTreePropertyPath& InPath) const final
-	{
-		return nullptr;
-	}
-PRAGMA_ENABLE_DEPRECATION_WARNINGS
 };
 
 
