[StateTree] Reset TriggerTransitionsFromFrameIndex after it was used. Fix if more than 1 linked state and completed in cascade. Added unittest to confirm the fix.

#Synced-CL 50826644

diff --git a/Engine/Plugins/Runtime/StateTree/Source/StateTreeModule/Private/StateTreeExecutionContext.cpp b/Engine/Plugins/Runtime/StateTree/Source/StateTreeModule/Private/StateTreeExecutionContext.cpp
--- a/Engine/Plugins/Runtime/StateTree/Source/StateTreeModule/Private/StateTreeExecutionContext.cpp
+++ b/Engine/Plugins/Runtime/StateTree/Source/StateTreeModule/Private/StateTreeExecutionContext.cpp
@@ -5740,6 +5740,7 @@ bool FStateTreeExecutionContext::TriggerTransitions()
 		// Transition() can TriggerTransitions() in a loop when a sub-frame completes.
 		//We do not want to evaluate the transition from that sub-frame.
 		const int32 EndFrameIndex = TriggerTransitionsFromFrameIndex.Get(Exec.ActiveFrames.Num() - 1);
+		TriggerTransitionsFromFrameIndex.Reset();
 		for (int32 FrameIndex = EndFrameIndex; FrameIndex >= 0; FrameIndex--)
 		{
 			FStateTreeExecutionFrame& CurrentFrame = Exec.ActiveFrames[FrameIndex];
diff --git a/Engine/Plugins/Runtime/StateTree/Source/StateTreeTestSuite/Private/StateTreeLinkedStateTest.cpp b/Engine/Plugins/Runtime/StateTree/Source/StateTreeTestSuite/Private/StateTreeLinkedStateTest.cpp
--- a/Engine/Plugins/Runtime/StateTree/Source/StateTreeTestSuite/Private/StateTreeLinkedStateTest.cpp
+++ b/Engine/Plugins/Runtime/StateTree/Source/StateTreeTestSuite/Private/StateTreeLinkedStateTest.cpp
@@ -1202,14 +1202,14 @@ struct FStateTreeTest_Linked_GlobalParameter : FStateTreeTestBase
 		//Tree 1
 		//  Global task and parameter
 		//	Root
-		//		StateLinkedTree1 (Tree2) -> Next
-		//		SubTree2
-		//	SubTree
-		//		State3 (with transition OnTick to succeeded)
+		//		State1 LinkedAsset(Tree2) [OnSucceeded goto Next]
+		//		State2 Linked(StateSub1) [OnSucceeded goto succeeded]
+		//	StateSub1
+		//		State3 [OnTick goto succeeded]
 		//Tree 2
 		//  Global task and parameter
 		//	Root
-		//		State1 (with transition OnTick to succeeded)
+		//		State1 [OnTick goto succeeded]
 
 		auto AddInt = [](FInstancedPropertyBag& PropertyBag, FName VarName)
 			{
@@ -2078,6 +2078,106 @@ struct FStateTreeTest_ReenterLinkedAssetWithDifferentAsset : FStateTreeTestBase
 IMPLEMENT_STATE_TREE_INSTANT_TEST(FStateTreeTest_ReenterLinkedAssetWithDifferentAsset, "System.StateTree.LinkedAsset.ReenterWithDifferentAsset");
 
 
+struct FStateTreeTest_AutoCompletedLinkedAsset : FStateTreeTestBase
+{
+	// Test going back to the same linked state after it completed via another transition.
+	virtual bool InstantTest() override
+	{
+		FStateTreeCompilerLog Log;
+
+		// Tree1
+		//  Root
+		//	  State1 (linked to Tree2) [OnCompleted go to State2]
+		//    State2 (auto complete on enter)[OnCompleted go to State1]
+		// Tree2
+			// State1 (linked to Tree3)
+		// Tree3
+			// State1 => 1 tick to completion
+
+		// Tree 3
+		UStateTree& StateTree3 = NewStateTree();
+		{
+			UStateTreeEditorData& EditorData = *Cast<UStateTreeEditorData>(StateTree3.EditorData);
+
+			UStateTreeState& State1 = EditorData.AddSubTree("Tree3State1");
+			TStateTreeEditorNode<FTestTask_Stand>& State1Task1 = State1.AddTask<FTestTask_Stand>("Tree3State1Task1");
+			State1Task1.GetNode().TicksToCompletion = 1;
+			State1.AddTransition(EStateTreeTransitionTrigger::OnStateCompleted, EStateTreeTransitionType::Succeeded);
+
+			{
+				FStateTreeCompiler Compiler(Log);
+				const bool bResult = Compiler.Compile(StateTree3);
+				AITEST_TRUE(TEXT("StateTree3 should get compiled"), bResult);
+			}
+		}
+
+		// Tree 2
+		UStateTree& StateTree2 = NewStateTree();
+		{
+			UStateTreeEditorData& EditorData = *Cast<UStateTreeEditorData>(StateTree2.EditorData);
+
+			UStateTreeState& State1 = EditorData.AddSubTree("Tree2State1", EStateTreeStateType::LinkedAsset);
+			State1.SetLinkedStateAsset(&StateTree3);
+
+			{
+				FStateTreeCompiler Compiler(Log);
+				const bool bResult = Compiler.Compile(StateTree2);
+				AITEST_TRUE(TEXT("StateTree2 should get compiled"), bResult);
+			}
+		}
+
+		// Tree 1
+		UStateTree& StateTree = NewStateTree();
+		{
+			UStateTreeEditorData& EditorData = *Cast<UStateTreeEditorData>(StateTree.EditorData);
+
+			UStateTreeState& Root = EditorData.AddSubTree("Tree1StateRoot");
+			UStateTreeState& State1 = Root.AddChildState("Tree1State1", EStateTreeStateType::LinkedAsset);
+			State1.SetLinkedStateAsset(&StateTree2);
+
+			UStateTreeState& State2 = Root.AddChildState("Tree1State2");
+			TStateTreeEditorNode<FTestTask_Stand>& State2Task1 = State2.AddTask<FTestTask_Stand>("Tree1State2Task1");
+			State2Task1.GetNode().EnterStateResult = EStateTreeRunStatus::Succeeded;
+
+			State1.AddTransition(EStateTreeTransitionTrigger::OnStateCompleted, EStateTreeTransitionType::GotoState, &State2);
+			State2.AddTransition(EStateTreeTransitionTrigger::OnStateCompleted, EStateTreeTransitionType::GotoState, &State1);
+			{
+				FStateTreeCompiler Compiler(Log);
+				const bool bResult = Compiler.Compile(StateTree);
+				AITEST_TRUE(TEXT("StateTree should get compiled"), bResult);
+			}
+		}
+
+		{
+			FStateTreeInstanceData InstanceData;
+			FTestStateTreeExecutionContext Exec(StateTree, StateTree, InstanceData);
+			const bool bInitSucceeded = Exec.IsValid();
+			AITEST_TRUE(TEXT("StateTree should init"), bInitSucceeded);
+			{
+				EStateTreeRunStatus Status = Exec.Start();
+				AITEST_EQUAL("Start should complete with Running", Status, EStateTreeRunStatus::Running);
+				AITEST_TRUE(TEXT("StateTree should tick Tree2State1Task1"), Exec.Expect("Tree3State1Task1", TEXT("EnterState")));
+				AITEST_TRUE(TEXT("StateTree should tick Tree2State1Task1"), Exec.ExpectInActiveStates("Tree1StateRoot", "Tree1State1", "Tree2State1", "Tree3State1"));
+				Exec.LogClear();
+			}
+			{
+				EStateTreeRunStatus Status = Exec.Tick(0.1f);
+				AITEST_EQUAL("Tick should complete with Running", Status, EStateTreeRunStatus::Running);
+				AITEST_TRUE(TEXT("StateTree should tick Tree2State1Task1"), Exec.Expect("Tree3State1Task1", TEXT("ExitState")));
+				AITEST_TRUE(TEXT("StateTree should tick Tree2State1Task1"), Exec.Expect("Tree1State2Task1", TEXT("EnterState")));
+				AITEST_TRUE(TEXT("StateTree should tick Tree2State1Task1"), Exec.Expect("Tree1State2Task1", TEXT("ExitState")));
+				AITEST_TRUE(TEXT("StateTree should tick Tree2State1Task1"), Exec.Expect("Tree3State1Task1", TEXT("EnterState")));
+				AITEST_TRUE(TEXT("StateTree should tick Tree2State1Task1"), Exec.ExpectInActiveStates("Tree1StateRoot", "Tree1State1", "Tree2State1", "Tree3State1"));
+				Exec.LogClear();
+			}
+			Exec.Stop();
+		}
+
+		return true;
+	}
+};
+IMPLEMENT_STATE_TREE_INSTANT_TEST(FStateTreeTest_AutoCompletedLinkedAsset, "System.StateTree.LinkedAsset.AutoCompleted");
+
 } // namespace UE::StateTree::Tests
 
 UE_ENABLE_OPTIMIZATION_SHIP
