[StateTree] Correctly Scope Start and Stop global tasks debug phase. Global tasks are entered/exited hierarchically and each frame should have its own scope to reflect the correct ordering.

#Synced-CL 50700524

diff --git a/Engine/Plugins/Runtime/StateTree/Source/StateTreeModule/Private/StateTreeExecutionContext.cpp b/Engine/Plugins/Runtime/StateTree/Source/StateTreeModule/Private/StateTreeExecutionContext.cpp
--- a/Engine/Plugins/Runtime/StateTree/Source/StateTreeModule/Private/StateTreeExecutionContext.cpp
+++ b/Engine/Plugins/Runtime/StateTree/Source/StateTreeModule/Private/StateTreeExecutionContext.cpp
@@ -4250,14 +4250,12 @@ EStateTreeRunStatus FStateTreeExecutionContext::StartGlobalsForFrameInternal(con
 	}
 
 	CSV_SCOPED_TIMING_STAT_EXCLUSIVE(StateTree_StartEvaluators);
-	UE_STATETREE_DEBUG_SCOPED_PHASE(this, EStateTreeUpdatePhase::StartGlobalTasks);
 
-	FStateTreeExecutionState& Exec = GetExecState();
 	FCurrentlyProcessedFrameScope FrameScope(*this, CurrentParentFrame, CurrentFrame);
-	if constexpr (!bOnActiveInstances)
-	{
-		UE_STATETREE_DEBUG_ENTER_PHASE(this, EStateTreeUpdatePhase::StartGlobalTasksForSelection);
-	}
+	
+	constexpr EStateTreeUpdatePhase StartGlobalTasksPhase = bOnActiveInstances ? EStateTreeUpdatePhase::StartGlobalTasks : EStateTreeUpdatePhase::StartGlobalTasksForSelection;
+
+	UE_STATETREE_DEBUG_ENTER_PHASE(this, StartGlobalTasksPhase);
 
 	CurrentFrame.bHaveEntered = true;
 
@@ -4422,10 +4420,7 @@ EStateTreeRunStatus FStateTreeExecutionContext::StartGlobalsForFrameInternal(con
 		}
 	}
 
-	if constexpr (!bOnActiveInstances)
-	{
-		UE_STATETREE_DEBUG_EXIT_PHASE(this, EStateTreeUpdatePhase::StartGlobalTasksForSelection);
-	}
+	UE_STATETREE_DEBUG_EXIT_PHASE(this, StartGlobalTasksPhase);
 
 	return Result;
 }
@@ -4438,7 +4433,7 @@ EStateTreeRunStatus FStateTreeExecutionContext::StartEvaluatorsAndGlobalTasks(FS
 
 EStateTreeRunStatus FStateTreeExecutionContext::StartEvaluatorsAndGlobalTasks()
 {
-	UE_STATETREE_DEBUG_SCOPED_PHASE(this, EStateTreeUpdatePhase::StartGlobalTasks);
+	CSV_SCOPED_TIMING_STAT_EXCLUSIVE(StateTree_StartEvaluatorsAndGlobalTasks);
 
 	STATETREE_LOG(Verbose, TEXT("Start Evaluators & Global tasks"));
 
@@ -4482,8 +4477,7 @@ void FStateTreeExecutionContext::StopEvaluatorsAndGlobalTasks(const EStateTreeRu
 
 void FStateTreeExecutionContext::StopEvaluatorsAndGlobalTasks(const EStateTreeRunStatus CompletionStatus)
 {
-	CSV_SCOPED_TIMING_STAT_EXCLUSIVE(StateTree_StopEvaluators);
-	UE_STATETREE_DEBUG_SCOPED_PHASE(this, EStateTreeUpdatePhase::StopGlobalTasks);
+	CSV_SCOPED_TIMING_STAT_EXCLUSIVE(StateTree_StopEvaluatorsAndGlobalTasks);
 
 	STATETREE_LOG(Verbose, TEXT("Stop Evaluators & Global Tasks"));
 
@@ -4530,10 +4524,9 @@ void FStateTreeExecutionContext::StopGlobalsForFrameInternal(const FStateTreeExe
 	}
 
 	FCurrentlyProcessedFrameScope FrameScope(*this, ParentFrame, Frame);
-	if constexpr (bOnActiveInstances)
-	{
-		UE_STATETREE_DEBUG_ENTER_PHASE(this, EStateTreeUpdatePhase::StopGlobalTasksForSelection);
-	}
+
+	constexpr EStateTreeUpdatePhase StopGlobalTasksPhase = bOnActiveInstances ? EStateTreeUpdatePhase::StopGlobalTasks : EStateTreeUpdatePhase::StopGlobalTasksForSelection;
+	UE_STATETREE_DEBUG_ENTER_PHASE(this, StopGlobalTasksPhase);
 
 	const UStateTree* CurrentStateTree = Frame.StateTree;
 
@@ -4618,10 +4611,7 @@ void FStateTreeExecutionContext::StopGlobalsForFrameInternal(const FStateTreeExe
 	Frame.ActiveNodeIndex = FStateTreeIndex16::Invalid;
 	Frame.bHaveEntered = false;
 
-	if constexpr (bOnActiveInstances)
-	{
-		UE_STATETREE_DEBUG_EXIT_PHASE(this, EStateTreeUpdatePhase::StopGlobalTasksForSelection);
-	}
+	UE_STATETREE_DEBUG_EXIT_PHASE(this, StopGlobalTasksPhase);
 }
 
 // Deprecated
