[StateTree] Fix RWD crash on object instance data saved as struct data.

#tests PIE, editor

#Synced-CL 50877637

diff --git a/Engine/Plugins/Runtime/StateTree/Source/StateTreeEditorModule/Private/StateTreeEditorModule.cpp b/Engine/Plugins/Runtime/StateTree/Source/StateTreeEditorModule/Private/StateTreeEditorModule.cpp
--- a/Engine/Plugins/Runtime/StateTree/Source/StateTreeEditorModule/Private/StateTreeEditorModule.cpp
+++ b/Engine/Plugins/Runtime/StateTree/Source/StateTreeEditorModule/Private/StateTreeEditorModule.cpp
@@ -99,30 +99,43 @@ namespace UE::StateTree::Editor::Private
 			{
 				if (UStateTreeEditorData* EditorData = Cast<UStateTreeEditorData>(StateTree->EditorData))
 				{
-					FConstStructView StructView;
+					FStateTreeDataView NodeDataView;
+
+					auto SetNodeDataView = [&NodeDataView, Node](auto& DataContainer)
+					{
+						// const casts needed as FStateTreeDataView doesn't accept const UObject* & FConstStructView
+						if (Node->InstanceDataHandle.IsObjectSource())
+						{
+							NodeDataView = FStateTreeDataView(const_cast<UObject*>(DataContainer.GetObject(Node->InstanceTemplateIndex.Get())));
+						}
+						else
+						{
+							FConstStructView StructView = DataContainer.GetStruct(Node->InstanceTemplateIndex.Get());
+							NodeDataView = FStateTreeDataView(StructView.GetScriptStruct(), const_cast<uint8*>(StructView.GetMemory()));
+						}
+					};
 
 					if (Node->InstanceDataHandle.GetSource() == EStateTreeDataSourceType::EvaluationScopeInstanceData
 						|| Node->InstanceDataHandle.GetSource() == EStateTreeDataSourceType::EvaluationScopeInstanceDataObject)
 					{
-						StructView = StateTree->GetDefaultEvaluationScopeInstanceData().GetStruct(Node->InstanceTemplateIndex.Get());
+						SetNodeDataView(StateTree->GetDefaultEvaluationScopeInstanceData());
 					}
 					else if (Node->InstanceDataHandle.GetSource() == EStateTreeDataSourceType::SharedInstanceData
 						|| Node->InstanceDataHandle.GetSource() == EStateTreeDataSourceType::SharedInstanceDataObject)
 					{
-						StructView = StateTree->GetSharedInstanceData()->GetStruct(Node->InstanceTemplateIndex.Get());
+						SetNodeDataView(*StateTree->GetSharedInstanceData());
 					}
 					else
 					{
-						StructView = StateTree->GetDefaultInstanceData().GetStruct(Node->InstanceTemplateIndex.Get());
+						SetNodeDataView(StateTree->GetDefaultInstanceData());
 					}
 
-					if (StructView.IsValid())
+					if (NodeDataView.IsValid())
 					{
-						const FStateTreeDataView StateTreeDataView = FStateTreeDataView(StructView.GetScriptStruct(), const_cast<uint8*>(StructView.GetMemory()));
 						const FGuid NodeID = StateTree->GetNodeIdFromIndex(NodeIndex);
 						const FStateTreeBindingLookup BindingLookup(EditorData);
 
-						return Node->GetDescription(NodeID, StateTreeDataView, BindingLookup, EStateTreeNodeFormatting::Text);
+						return Node->GetDescription(NodeID, NodeDataView, BindingLookup, EStateTreeNodeFormatting::Text);
 					}
 				}
 			}
