[StateTree] Fix bad access in TStateTreeInstanceDataStructRef.

#Synced-CL 50157473
#rnx

#Synced-CL 50826644

diff --git a/Engine/Plugins/Runtime/StateTree/Source/StateTreeModule/Public/StateTreeInstanceData.h b/Engine/Plugins/Runtime/StateTree/Source/StateTreeModule/Public/StateTreeInstanceData.h
--- a/Engine/Plugins/Runtime/StateTree/Source/StateTreeModule/Public/StateTreeInstanceData.h
+++ b/Engine/Plugins/Runtime/StateTree/Source/StateTreeModule/Public/StateTreeInstanceData.h
@@ -737,7 +737,7 @@ PRAGMA_ENABLE_DEPRECATION_WARNINGS
 
 			if constexpr (!std::is_same_v<T, void>)
 			{
-				check(DataView.GetStruct() == TBaseStructure<T>::Get());
+				check(DataView.GetStruct() == nullptr || DataView.GetStruct()->IsChildOf<T>());
 			}
 		}
 		else
