[StateTree] Fix wrong use of ActiveNodeIndex in ExecutionContext. A task can be deactivated. For StateExit, since we run in reverse order, we can't set the index until it reach. Added unit test "System.StateTree.TasksCompletion.StartGlobalTaskFailure" to confirm the logic for global tasks

#Synced-CL 50933965

diff --git a/Engine/Plugins/Runtime/StateTree/Source/StateTreeModule/Private/StateTreeExecutionContext.cpp b/Engine/Plugins/Runtime/StateTree/Source/StateTreeModule/Private/StateTreeExecutionContext.cpp
--- a/Engine/Plugins/Runtime/StateTree/Source/StateTreeModule/Private/StateTreeExecutionContext.cpp
+++ b/Engine/Plugins/Runtime/StateTree/Source/StateTreeModule/Private/StateTreeExecutionContext.cpp
@@ -1123,20 +1123,33 @@ FStateTreeExecutionContext::FStateTreeExecutionContext(const FStateTreeExecution
 FStateTreeExecutionContext::FStateTreeExecutionContext(const FStateTreeExecutionContext& InContextToCopy, TNotNull<const UStateTree*> InStateTree, FStateTreeInstanceData& InInstanceData)
 	: FStateTreeExecutionContext(&InContextToCopy.Owner, InStateTree, InInstanceData, InContextToCopy.CollectExternalDataDelegate)
 {
-	SetLinkedStateTreeOverrides(InContextToCopy.LinkedAssetStateTreeOverrides);
-	const bool bIsSameSchema = RootStateTree.GetSchema()->GetClass() == InContextToCopy.GetStateTree()->GetSchema()->GetClass();
-	if (bIsSameSchema)
+	if (IsValid())
 	{
-		for (const FStateTreeExternalDataDesc& TargetDataDesc : GetContextDataDescs())
+		SetLinkedStateTreeOverrides(InContextToCopy.LinkedAssetStateTreeOverrides);
+		const bool bIsSameSchema = RootStateTree.GetSchema()->GetClass() == InContextToCopy.GetStateTree()->GetSchema()->GetClass();
+		if (bIsSameSchema)
 		{
-			const int32 TargetIndex = TargetDataDesc.Handle.DataHandle.GetIndex();
-			ContextAndExternalDataViews[TargetIndex] = InContextToCopy.ContextAndExternalDataViews[TargetIndex];
+			for (const FStateTreeExternalDataDesc& TargetDataDesc : GetContextDataDescs())
+			{
+				const int32 TargetIndex = TargetDataDesc.Handle.DataHandle.GetIndex();
+				const bool bViewsValidIndex = ContextAndExternalDataViews.IsValidIndex(TargetIndex);
+				const bool bCopyViewsValidIndex = InContextToCopy.ContextAndExternalDataViews.IsValidIndex(TargetIndex);
+				if (ensure(bViewsValidIndex && bCopyViewsValidIndex)) // same context, should share the same ContextAndExternalDataViews.
+				{
+					ContextAndExternalDataViews[TargetIndex] = InContextToCopy.ContextAndExternalDataViews[TargetIndex];
+				}
+			}
+		}
+		else
+		{
+			STATETREE_LOG(Error, TEXT("%hs: '%s' using StateTree '%s' trying to run subtree '%s' but their schemas don't match"),
+				__FUNCTION__, *GetNameSafe(&Owner), *GetFullNameSafe(InContextToCopy.GetStateTree()), *GetFullNameSafe(&RootStateTree));
 		}
 	}
 	else
 	{
-		STATETREE_LOG(Error, TEXT("%hs: '%s' using StateTree '%s' trying to run subtree '%s' but their schemas don't match"),
-			__FUNCTION__, *GetNameSafe(&Owner), *GetFullNameSafe(InContextToCopy.GetStateTree()), *GetFullNameSafe(&RootStateTree));
+		STATETREE_LOG(Warning, TEXT("%hs: StateTree asset is not valid ('%s' using StateTree '%s')"),
+			__FUNCTION__, *GetNameSafe(&Owner), *GetFullNameSafe(&RootStateTree));
 	}
 }
 
@@ -3906,40 +3919,38 @@ void FStateTreeExecutionContext::ExitState(const TSharedPtr<const FSelectStateRe
 			FCurrentlyProcessedStateScope StateScope(*this, CurrentHandle);
 			UE_STATETREE_DEBUG_STATE_EVENT(this, CurrentHandle, EStateTreeTraceEventType::OnExiting);
 
-			for (int32 TaskIndex = (CurrentState.TasksBegin + CurrentState.TasksNum) - 1; TaskIndex >= CurrentState.TasksBegin; --TaskIndex)
+			// Call task completed only if EnterState() was called.
+			// The task order in the tree (BF) allows us to use the comparison.
+			const int32 TaskEnd = FMath::Min(CurrentFrame.ActiveNodeIndex.AsInt32() + 1, CurrentState.TasksBegin + CurrentState.TasksNum);
+			for (int32 TaskIndex = TaskEnd - 1; TaskIndex >= CurrentState.TasksBegin; --TaskIndex)
 			{
 				const FStateTreeTaskBase& Task = CurrentStateTree->Nodes[TaskIndex].Get<const FStateTreeTaskBase>();
 
 				// Ignore disabled task
 				if (Task.bTaskEnabled)
 				{
-					// Call task completed only if EnterState() was called.
-					// The task order in the tree (BF) allows us to use the comparison.
-					if (TaskIndex <= CurrentFrame.ActiveNodeIndex.AsInt32())
-					{
-						const FStateTreeDataView TaskInstanceView = GetDataView(CurrentParentFrame, CurrentFrame, Task.InstanceDataHandle);
-						FNodeInstanceDataScope DataScope(*this, &Task, TaskIndex, Task.InstanceDataHandle, TaskInstanceView);
+					const FStateTreeDataView TaskInstanceView = GetDataView(CurrentParentFrame, CurrentFrame, Task.InstanceDataHandle);
+					FNodeInstanceDataScope DataScope(*this, &Task, TaskIndex, Task.InstanceDataHandle, TaskInstanceView);
 
-						const bool bShouldCallStateChange = CurrentTransition.ChangeType == EStateTreeStateChangeType::Changed
-							|| (CurrentTransition.ChangeType == EStateTreeStateChangeType::Sustained && Task.bShouldStateChangeOnReselect);
+					const bool bShouldCallStateChange = CurrentTransition.ChangeType == EStateTreeStateChangeType::Changed
+						|| (CurrentTransition.ChangeType == EStateTreeStateChangeType::Sustained && Task.bShouldStateChangeOnReselect);
 
-						if (bShouldCallStateChange)
+					if (bShouldCallStateChange)
+					{
+						STATETREE_LOG(Verbose, TEXT("%*sTask '%s'.ExitState()"), (FrameIndex + StateIndex + 1) * UE::StateTree::Debug::IndentSize, TEXT(""), *Task.Name.ToString());
+						UE_STATETREE_DEBUG_TASK_EXIT_STATE(this, CurrentStateTree, FStateTreeIndex16(TaskIndex));
 						{
-							STATETREE_LOG(Verbose, TEXT("%*sTask '%s'.ExitState()"), (FrameIndex + StateIndex + 1) * UE::StateTree::Debug::IndentSize, TEXT(""), *Task.Name.ToString());
-							UE_STATETREE_DEBUG_TASK_EXIT_STATE(this, CurrentStateTree, FStateTreeIndex16(TaskIndex));
-							{
-								QUICK_SCOPE_CYCLE_COUNTER(StateTree_Task_ExitState);
-								CSV_SCOPED_TIMING_STAT_EXCLUSIVE(StateTree_Task_ExitState);
-								Task.ExitState(*this, CurrentTransition);
-							}
-
-							if (Task.OutputBindingsBatch.IsValid())
-							{
-								CopyBatchOnActiveInstances(CurrentlyProcessedParentFrame, CurrentFrame, TaskInstanceView, Task.OutputBindingsBatch);
-							}
+							QUICK_SCOPE_CYCLE_COUNTER(StateTree_Task_ExitState);
+							CSV_SCOPED_TIMING_STAT_EXCLUSIVE(StateTree_Task_ExitState);
+							Task.ExitState(*this, CurrentTransition);
+						}
 
-							UE_STATETREE_DEBUG_TASK_EVENT(this, TaskIndex, TaskInstanceView, EStateTreeTraceEventType::OnExited, Transition.CurrentRunStatus);
+						if (Task.OutputBindingsBatch.IsValid())
+						{
+							CopyBatchOnActiveInstances(CurrentlyProcessedParentFrame, CurrentFrame, TaskInstanceView, Task.OutputBindingsBatch);
 						}
+
+						UE_STATETREE_DEBUG_TASK_EVENT(this, TaskIndex, TaskInstanceView, EStateTreeTraceEventType::OnExited, Transition.CurrentRunStatus);
 					}
 				}
 				else
@@ -4530,80 +4541,74 @@ void FStateTreeExecutionContext::StopGlobalsForFrameInternal(const FStateTreeExe
 
 	const UStateTree* CurrentStateTree = Frame.StateTree;
 
-	const int32 GlobalTasksEnd = CurrentStateTree->GlobalTasksBegin + CurrentStateTree->GlobalTasksNum;
+	const int32 GlobalTasksEnd = FMath::Min(Frame.ActiveNodeIndex.AsInt32() + 1, CurrentStateTree->GlobalTasksBegin + CurrentStateTree->GlobalTasksNum);
 	for (int32 TaskIndex = GlobalTasksEnd - 1; TaskIndex >= CurrentStateTree->GlobalTasksBegin; --TaskIndex)
 	{
-		if (TaskIndex <= Frame.ActiveNodeIndex.AsInt32())
-		{
-			const FStateTreeTaskBase& Task = CurrentStateTree->Nodes[TaskIndex].Get<const FStateTreeTaskBase>();
-			const FStateTreeDataView TaskInstanceView = bOnActiveInstances
-				? GetDataView(ParentFrame, Frame, Task.InstanceDataHandle)
-				: GetDataViewOrTemporary(ParentFrame, Frame, Task.InstanceDataHandle);
-			FNodeInstanceDataScope DataScope(*this, &Task, TaskIndex, Task.InstanceDataHandle, TaskInstanceView);
+		const FStateTreeTaskBase& Task = CurrentStateTree->Nodes[TaskIndex].Get<const FStateTreeTaskBase>();
+		const FStateTreeDataView TaskInstanceView = bOnActiveInstances
+			? GetDataView(ParentFrame, Frame, Task.InstanceDataHandle)
+			: GetDataViewOrTemporary(ParentFrame, Frame, Task.InstanceDataHandle);
+		FNodeInstanceDataScope DataScope(*this, &Task, TaskIndex, Task.InstanceDataHandle, TaskInstanceView);
 
-			// Ignore disabled task
-			if (Task.bTaskEnabled == false)
+		// Ignore disabled task
+		if (Task.bTaskEnabled == false)
+		{
+			STATETREE_LOG(VeryVerbose, TEXT("%*sSkipped 'ExitState' for disabled Task: '%s'"), UE::StateTree::Debug::IndentSize, TEXT(""), *Task.Name.ToString());
+		}
+		else
+		{
+			STATETREE_LOG(Verbose, TEXT("  Stop: '%s'"), *Task.Name.ToString());
+			UE_STATETREE_DEBUG_TASK_EXIT_STATE(this, CurrentStateTree, FStateTreeIndex16(TaskIndex));
 			{
-				STATETREE_LOG(VeryVerbose, TEXT("%*sSkipped 'ExitState' for disabled Task: '%s'"), UE::StateTree::Debug::IndentSize, TEXT(""), *Task.Name.ToString());
+				QUICK_SCOPE_CYCLE_COUNTER(StateTree_Task_TreeStop);
+				Task.ExitState(*this, Transition);
 			}
-			else
+
+			if (Task.OutputBindingsBatch.IsValid())
 			{
-				STATETREE_LOG(Verbose, TEXT("  Stop: '%s'"), *Task.Name.ToString());
-				UE_STATETREE_DEBUG_TASK_EXIT_STATE(this, CurrentStateTree, FStateTreeIndex16(TaskIndex));
+				if constexpr (bOnActiveInstances)
 				{
-					QUICK_SCOPE_CYCLE_COUNTER(StateTree_Task_TreeStop);
-					Task.ExitState(*this, Transition);
+					CopyBatchOnActiveInstances(ParentFrame, Frame, TaskInstanceView, Task.OutputBindingsBatch);
 				}
-
-				if (Task.OutputBindingsBatch.IsValid())
+				else
 				{
-					if constexpr (bOnActiveInstances)
-					{
-						CopyBatchOnActiveInstances(ParentFrame, Frame, TaskInstanceView, Task.OutputBindingsBatch);
-					}
-					else
-					{
-						CopyBatchWithValidation(ParentFrame, Frame, TaskInstanceView, Task.OutputBindingsBatch);
-					}
+					CopyBatchWithValidation(ParentFrame, Frame, TaskInstanceView, Task.OutputBindingsBatch);
 				}
-
-				UE_STATETREE_DEBUG_TASK_EVENT(this, TaskIndex, TaskInstanceView, EStateTreeTraceEventType::OnExited, Transition.CurrentRunStatus);
 			}
+
+			UE_STATETREE_DEBUG_TASK_EVENT(this, TaskIndex, TaskInstanceView, EStateTreeTraceEventType::OnExited, Transition.CurrentRunStatus);
 		}
 		Frame.ActiveNodeIndex = FStateTreeIndex16(TaskIndex - 1);
 	}
 
-	const int32 EvaluatorsEnd = CurrentStateTree->EvaluatorsBegin + CurrentStateTree->EvaluatorsNum;
+	const int32 EvaluatorsEnd = FMath::Min(Frame.ActiveNodeIndex.AsInt32() + 1, CurrentStateTree->EvaluatorsBegin + CurrentStateTree->EvaluatorsNum);
 	for (int32 EvalIndex = EvaluatorsEnd - 1; EvalIndex >= CurrentStateTree->EvaluatorsBegin; --EvalIndex)
 	{
-		if (EvalIndex <= Frame.ActiveNodeIndex.AsInt32())
+		const FStateTreeEvaluatorBase& Eval = CurrentStateTree->Nodes[EvalIndex].Get<const FStateTreeEvaluatorBase>();
+		const FStateTreeDataView EvalInstanceView = bOnActiveInstances
+			? GetDataView(ParentFrame, Frame, Eval.InstanceDataHandle)
+			: GetDataViewOrTemporary(ParentFrame, Frame, Eval.InstanceDataHandle);
+		FNodeInstanceDataScope DataScope(*this, &Eval, EvalIndex, Eval.InstanceDataHandle, EvalInstanceView);
+
+		STATETREE_LOG(Verbose, TEXT("  Stop: '%s'"), *Eval.Name.ToString());
+		UE_STATETREE_DEBUG_EVALUATOR_EXIT_TREE(this, CurrentStateTree, FStateTreeIndex16(EvalIndex));
 		{
-			const FStateTreeEvaluatorBase& Eval = CurrentStateTree->Nodes[EvalIndex].Get<const FStateTreeEvaluatorBase>();
-			const FStateTreeDataView EvalInstanceView = bOnActiveInstances
-				? GetDataView(ParentFrame, Frame, Eval.InstanceDataHandle)
-				: GetDataViewOrTemporary(ParentFrame, Frame, Eval.InstanceDataHandle);
-			FNodeInstanceDataScope DataScope(*this, &Eval, EvalIndex, Eval.InstanceDataHandle, EvalInstanceView);
+			QUICK_SCOPE_CYCLE_COUNTER(StateTree_Eval_TreeStop);
+			Eval.TreeStop(*this);
 
-			STATETREE_LOG(Verbose, TEXT("  Stop: '%s'"), *Eval.Name.ToString());
-			UE_STATETREE_DEBUG_EVALUATOR_EXIT_TREE(this, CurrentStateTree, FStateTreeIndex16(EvalIndex));
+			if (Eval.OutputBindingsBatch.IsValid())
 			{
-				QUICK_SCOPE_CYCLE_COUNTER(StateTree_Eval_TreeStop);
-				Eval.TreeStop(*this);
-
-				if (Eval.OutputBindingsBatch.IsValid())
+				if constexpr (bOnActiveInstances)
 				{
-					if constexpr (bOnActiveInstances)
-					{
-						CopyBatchOnActiveInstances(ParentFrame, Frame, EvalInstanceView, Eval.OutputBindingsBatch);
-					}
-					else
-					{
-						CopyBatchWithValidation(ParentFrame, Frame, EvalInstanceView, Eval.OutputBindingsBatch);
-					}
+					CopyBatchOnActiveInstances(ParentFrame, Frame, EvalInstanceView, Eval.OutputBindingsBatch);
+				}
+				else
+				{
+					CopyBatchWithValidation(ParentFrame, Frame, EvalInstanceView, Eval.OutputBindingsBatch);
 				}
-
-				UE_STATETREE_DEBUG_EVALUATOR_EVENT(this, EvalIndex, EvalInstanceView, EStateTreeTraceEventType::OnTreeStopped);
 			}
+
+			UE_STATETREE_DEBUG_EVALUATOR_EVENT(this, EvalIndex, EvalInstanceView, EStateTreeTraceEventType::OnTreeStopped);
 		}
 		Frame.ActiveNodeIndex = FStateTreeIndex16(EvalIndex - 1);
 	}
@@ -5737,7 +5742,7 @@ bool FStateTreeExecutionContext::TriggerTransitions()
 			const UStateTree* CurrentStateTree = CurrentFrame.StateTree;
 			const int32 CurrentActiveNodeIndex = CurrentFrame.ActiveNodeIndex.AsInt32();
 
-			for (int32 StateIndex = CurrentFrame.ActiveStates.Num() - 1; StateIndex >= 0; StateIndex--)
+			for (int32 StateIndex = CurrentFrame.ActiveStates.Num() - 1; StateIndex >= 0; --StateIndex)
 			{
 				const FStateTreeStateHandle StateHandle = CurrentFrame.ActiveStates[StateIndex];
 				const UE::StateTree::FActiveStateID StateID = CurrentFrame.ActiveStates.StateIDs[StateIndex];
@@ -5749,10 +5754,11 @@ bool FStateTreeExecutionContext::TriggerTransitions()
 				if (State.bHasTransitionTasks)
 				{
 					bool bAdded = false;
-					for (int32 TaskIndex = (State.TasksBegin + State.TasksNum) - 1; TaskIndex >= State.TasksBegin; TaskIndex--)
+					const int32 TaskEnd = FMath::Min(CurrentActiveNodeIndex + 1, State.TasksBegin + State.TasksNum);
+					for (int32 TaskIndex = TaskEnd - 1; TaskIndex >= State.TasksBegin; --TaskIndex)
 					{
 						const FStateTreeTaskBase& Task = CurrentStateTree->Nodes[TaskIndex].Get<const FStateTreeTaskBase>();
-						if (Task.bShouldAffectTransitions && Task.bTaskEnabled && TaskIndex <= CurrentActiveNodeIndex)
+						if (Task.bShouldAffectTransitions && Task.bTaskEnabled)
 						{
 							TransitionHandlers.Emplace(FrameIndex, StateHandle, StateID, FStateTreeIndex16(TaskIndex), Task.TransitionHandlingPriority);
 							bAdded = true;
@@ -5782,7 +5788,8 @@ bool FStateTreeExecutionContext::TriggerTransitions()
 				if (CurrentFrame.StateTree->bHasGlobalTransitionTasks)
 				{
 					bool bAdded = false;
-					for (int32 TaskIndex = (CurrentStateTree->GlobalTasksBegin + CurrentStateTree->GlobalTasksNum) - 1; TaskIndex >= CurrentFrame.StateTree->GlobalTasksBegin; TaskIndex--)
+					const int32 TaskEnd = FMath::Min(CurrentActiveNodeIndex + 1, CurrentStateTree->GlobalTasksBegin + CurrentStateTree->GlobalTasksNum);
+					for (int32 TaskIndex = TaskEnd - 1; TaskIndex >= CurrentFrame.StateTree->GlobalTasksBegin; --TaskIndex)
 					{
 						const FStateTreeTaskBase& Task = CurrentStateTree->Nodes[TaskIndex].Get<const FStateTreeTaskBase>();
 						if (Task.bShouldAffectTransitions && Task.bTaskEnabled)
diff --git a/Engine/Plugins/Runtime/StateTree/Source/StateTreeTestSuite/Private/StateTreeTaskStateTest.cpp b/Engine/Plugins/Runtime/StateTree/Source/StateTreeTestSuite/Private/StateTreeTaskStateTest.cpp
--- a/Engine/Plugins/Runtime/StateTree/Source/StateTreeTestSuite/Private/StateTreeTaskStateTest.cpp
+++ b/Engine/Plugins/Runtime/StateTree/Source/StateTreeTestSuite/Private/StateTreeTaskStateTest.cpp
@@ -253,6 +253,168 @@ struct FStateTreeTest_TasksCompletion_FailureTasks : FStateTreeTestBase
 };
 IMPLEMENT_STATE_TREE_INSTANT_TEST(FStateTreeTest_TasksCompletion_FailureTasks, "System.StateTree.TasksCompletion.Failure");
 
+struct FStateTreeTest_TasksCompletion_StartGlobalFailureTasks : FStateTreeTestBase
+{
+	// Test when global tasks fails when the tree starts.
+	virtual bool InstantTest() override
+	{
+		// Tree1 (3 global tasks)
+		//		Root
+		UStateTree& StateTree = NewStateTree();
+		TStateTreeEditorNode<FTestTask_Stand>* GlobalTask1Ptr = nullptr;
+		TStateTreeEditorNode<FTestTask_Stand>* GlobalTask2Ptr = nullptr;
+		TStateTreeEditorNode<FTestTask_Stand>* GlobalTask3Ptr = nullptr;
+
+		auto CompileTree = [&StateTree, this]()
+			{
+				FStateTreeCompilerLog Log;
+				FStateTreeCompiler Compiler(Log);
+				const bool bResult = Compiler.Compile(StateTree);
+				AITEST_TRUE(TEXT("StateTree should get compiled"), bResult);
+				return bResult;
+			};
+
+		{
+			UStateTreeEditorData& EditorData = *Cast<UStateTreeEditorData>(StateTree.EditorData);
+
+			TStateTreeEditorNode<FTestTask_Stand>& GlobalTask1 = EditorData.AddGlobalTask<FTestTask_Stand>("Tree1GlobalTask1");
+			TStateTreeEditorNode<FTestTask_Stand>& GlobalTask2 = EditorData.AddGlobalTask<FTestTask_Stand>("Tree1GlobalTask2");
+			TStateTreeEditorNode<FTestTask_Stand>& GlobalTask3 = EditorData.AddGlobalTask<FTestTask_Stand>("Tree1GlobalTask3");
+			GlobalTask1Ptr = &GlobalTask1;
+			GlobalTask2Ptr = &GlobalTask2;
+			GlobalTask3Ptr = &GlobalTask3;
+
+			UStateTreeState& Root = EditorData.AddSubTree(FName(TEXT("Tree1StateRoot")));
+			Root.AddTask<FTestTask_Stand>("Tree1StateRootTask1");
+		}
+
+		// test global task 1 failing
+		GlobalTask1Ptr->GetNode().EnterStateResult = EStateTreeRunStatus::Failed;
+		GlobalTask2Ptr->GetNode().EnterStateResult = EStateTreeRunStatus::Succeeded;
+		if (!CompileTree())
+		{
+			return false;
+		}
+		{
+			EStateTreeRunStatus Status = EStateTreeRunStatus::Unset;
+			FStateTreeInstanceData InstanceData;
+			FTestStateTreeExecutionContext Exec(StateTree, StateTree, InstanceData);
+			const bool bInitSucceeded = Exec.IsValid();
+			AITEST_TRUE(TEXT("StateTree should init"), bInitSucceeded);
+
+			{
+				Status = Exec.Start();
+				AITEST_EQUAL(TEXT("Start should complete with failed"), Status, EStateTreeRunStatus::Failed);
+				AITEST_TRUE(TEXT("Task1 should enter but fail."),
+					Exec.Expect("Tree1GlobalTask1", TEXT("EnterState"))
+					.Then("Tree1GlobalTask1", TEXT("ExitFailed"))
+					.Then("Tree1GlobalTask1", TEXT("ExitState"))
+				);
+				AITEST_FALSE(TEXT("Task2 should not enter."), Exec.Expect("Tree1GlobalTask2", TEXT("EnterState")));
+				AITEST_FALSE(TEXT("Task3 should not enter."), Exec.Expect("Tree1GlobalTask3", TEXT("EnterState")));
+				AITEST_FALSE(TEXT("Task2 should not exit."), Exec.Expect("Tree1GlobalTask2", TEXT("ExitState")));
+				AITEST_FALSE(TEXT("Task3 should not exit."), Exec.Expect("Tree1GlobalTask3", TEXT("ExitState")));
+				Exec.LogClear();
+				Exec.Stop();
+			}
+		}
+
+		// test global task 2 failing
+		GlobalTask1Ptr->GetNode().EnterStateResult = EStateTreeRunStatus::Succeeded;
+		GlobalTask2Ptr->GetNode().EnterStateResult = EStateTreeRunStatus::Failed;
+		if (!CompileTree())
+		{
+			return false;
+		}
+		{
+			EStateTreeRunStatus Status = EStateTreeRunStatus::Unset;
+			FStateTreeInstanceData InstanceData;
+			FTestStateTreeExecutionContext Exec(StateTree, StateTree, InstanceData);
+			const bool bInitSucceeded = Exec.IsValid();
+			AITEST_TRUE(TEXT("StateTree should init"), bInitSucceeded);
+
+			{
+				Status = Exec.Start();
+				AITEST_EQUAL(TEXT("Start should complete with failed"), Status, EStateTreeRunStatus::Failed);
+				AITEST_TRUE(TEXT("Task2 should enter."),
+					Exec.Expect("Tree1GlobalTask1", TEXT("EnterState"))
+					.Then("Tree1GlobalTask2", TEXT("EnterState"))
+					.Then("Tree1GlobalTask2", TEXT("ExitFailed"))
+					.Then("Tree1GlobalTask2", TEXT("ExitState"))
+					.Then("Tree1GlobalTask1", TEXT("ExitState"))
+				);
+				AITEST_FALSE(TEXT("Task3 should not enter."), Exec.Expect("Tree1GlobalTask3", TEXT("EnterState")));
+				AITEST_FALSE(TEXT("Task3 should not exit."), Exec.Expect("Tree1GlobalTask3", TEXT("ExitState")));
+				Exec.LogClear();
+				Exec.Stop();
+			}
+		}
+
+		// test global task 3 failing
+		GlobalTask2Ptr->GetNode().EnterStateResult = EStateTreeRunStatus::Succeeded;
+		GlobalTask3Ptr->GetNode().EnterStateResult = EStateTreeRunStatus::Failed;
+		if (!CompileTree())
+		{
+			return false;
+		}
+		{
+			EStateTreeRunStatus Status = EStateTreeRunStatus::Unset;
+			FStateTreeInstanceData InstanceData;
+			FTestStateTreeExecutionContext Exec(StateTree, StateTree, InstanceData);
+			const bool bInitSucceeded = Exec.IsValid();
+			AITEST_TRUE(TEXT("StateTree should init"), bInitSucceeded);
+
+			{
+				Status = Exec.Start();
+				AITEST_EQUAL(TEXT("Start should complete with failed"), Status, EStateTreeRunStatus::Failed);
+				AITEST_TRUE(TEXT("Task3 should enter."),
+					Exec.Expect("Tree1GlobalTask1", TEXT("EnterState"))
+					.Then("Tree1GlobalTask2", TEXT("EnterState"))
+					.Then("Tree1GlobalTask3", TEXT("EnterState"))
+					.Then("Tree1GlobalTask3", TEXT("ExitFailed"))
+					.Then("Tree1GlobalTask3", TEXT("ExitState"))
+					.Then("Tree1GlobalTask2", TEXT("ExitState"))
+					.Then("Tree1GlobalTask1", TEXT("ExitState"))
+				);
+				Exec.LogClear();
+				Exec.Stop();
+			}
+		}
+
+		// test global task 3 succeeded
+		GlobalTask3Ptr->GetNode().EnterStateResult = EStateTreeRunStatus::Succeeded;
+		if (!CompileTree())
+		{
+			return false;
+		}
+		{
+			EStateTreeRunStatus Status = EStateTreeRunStatus::Unset;
+			FStateTreeInstanceData InstanceData;
+			FTestStateTreeExecutionContext Exec(StateTree, StateTree, InstanceData);
+			const bool bInitSucceeded = Exec.IsValid();
+			AITEST_TRUE(TEXT("StateTree should init"), bInitSucceeded);
+
+			{
+				Status = Exec.Start();
+				AITEST_EQUAL(TEXT("Start should complete with succeeded."), Status, EStateTreeRunStatus::Succeeded);
+				AITEST_TRUE(TEXT("Task3 should enter."),
+					Exec.Expect("Tree1GlobalTask1", TEXT("EnterState"))
+					.Then("Tree1GlobalTask2", TEXT("EnterState"))
+					.Then("Tree1GlobalTask3", TEXT("EnterState"))
+					.Then("Tree1GlobalTask3", TEXT("ExitSucceeded"))
+					.Then("Tree1GlobalTask2", TEXT("ExitSucceeded"))
+					.Then("Tree1GlobalTask1", TEXT("ExitSucceeded"))
+				);
+				Exec.LogClear();
+				Exec.Stop();
+			}
+		}
+
+		return true;
+	}
+};
+IMPLEMENT_STATE_TREE_INSTANT_TEST(FStateTreeTest_TasksCompletion_StartGlobalFailureTasks, "System.StateTree.TasksCompletion.StartGlobalTaskFailure");
+
 // test set status with priority
 struct FStateTreeTest_TasksCompletion_Status : FStateTreeTestBase
 {
