Changed CSV StateTree stats to non exclusive.

Time will be added back to CSV Tickables. ~90us on Switch.

#tests replayrun

#Synced-CL 50986165

#Synced-CL 50882011

diff --git a/Engine/Plugins/AI/MassAI/Source/MassAIBehavior/Private/MassStateTreeExecutionContext.cpp b/Engine/Plugins/AI/MassAI/Source/MassAIBehavior/Private/MassStateTreeExecutionContext.cpp
--- a/Engine/Plugins/AI/MassAI/Source/MassAIBehavior/Private/MassStateTreeExecutionContext.cpp
+++ b/Engine/Plugins/AI/MassAI/Source/MassAIBehavior/Private/MassStateTreeExecutionContext.cpp
@@ -11,11 +11,13 @@
 
 #include UE_INLINE_GENERATED_CPP_BY_NAME(MassStateTreeExecutionContext)
 
+CSV_DECLARE_CATEGORY_EXTERN(MassStateTreeProcessor);
+
 namespace UE::MassBehavior
 {
 bool CollectExternalData(const FStateTreeExecutionContext& Context, const UStateTree* StateTree, TArrayView<const FStateTreeExternalDataDesc> ExternalDataDescs, TArrayView<FStateTreeDataView> OutDataViews)
 {
-	CSV_SCOPED_TIMING_STAT_EXCLUSIVE(StateTreeProcessorCollectExternalData);
+	CSV_SCOPED_TIMING_STAT(MassStateTreeProcessor, CollectExternalData);
 	
 	const FMassStateTreeExecutionContext& MassStateTreeContext = static_cast<const FMassStateTreeExecutionContext&>(Context); 
 	const FMassEntityManager& EntityManager = MassStateTreeContext.GetEntityManager();
diff --git a/Engine/Plugins/AI/MassAI/Source/MassAIBehavior/Private/MassStateTreeProcessors.cpp b/Engine/Plugins/AI/MassAI/Source/MassAIBehavior/Private/MassStateTreeProcessors.cpp
--- a/Engine/Plugins/AI/MassAI/Source/MassAIBehavior/Private/MassStateTreeProcessors.cpp
+++ b/Engine/Plugins/AI/MassAI/Source/MassAIBehavior/Private/MassStateTreeProcessors.cpp
@@ -23,7 +23,7 @@
 
 #include UE_INLINE_GENERATED_CPP_BY_NAME(MassStateTreeProcessors)
 
-CSV_DEFINE_CATEGORY(StateTreeProcessor, true);
+CSV_DEFINE_CATEGORY(MassStateTreeProcessor, false);
 
 namespace UE::MassBehavior
 {
@@ -60,7 +60,7 @@ void ForEachEntityInChunk(
 
 			// Make sure all required external data are set.
 			{
-				CSV_SCOPED_TIMING_STAT_EXCLUSIVE(StateTreeProcessorExternalDataValidation);
+				CSV_SCOPED_TIMING_STAT(MassStateTreeProcessor, ExternalDataValidation);
 				// TODO: disable this when not in debug.
 				if (!ensureMsgf(StateTreeContext.AreContextDataViewsValid(), TEXT("StateTree will not execute due to missing external data.")))
 				{
@@ -315,14 +315,14 @@ void UMassStateTreeProcessor::SignalEntities(FMassEntityManager& EntityManager,
 	UMassSignalSubsystem& SignalSubsystem = Context.GetMutableSubsystemChecked<UMassSignalSubsystem>();
 	
 	QUICK_SCOPE_CYCLE_COUNTER(StateTreeProcessor_Run);
-	CSV_SCOPED_TIMING_STAT_EXCLUSIVE(StateTreeProcessorExecute);
+	CSV_SCOPED_TIMING_STAT(MassStateTreeProcessor, Execute);
 
 	const double TimeInSeconds = EntityManager.GetWorld()->GetTimeSeconds();
 
 	auto TickFunc = [TimeInSeconds, &EntitySignals](FMassExecutionContext& ExecutionContext, TFunctionRef<void(FMassEntityHandle)> AddEntityToSignalFunc)
 		{
 			// Keep stats regarding the amount of tree instances ticked per frame
-			CSV_CUSTOM_STAT(StateTreeProcessor, NumTickedStateTree, ExecutionContext.GetNumEntities(), ECsvCustomStatOp::Accumulate);
+			CSV_CUSTOM_STAT(MassStateTreeProcessor, NumTickedStateTree, ExecutionContext.GetNumEntities(), ECsvCustomStatOp::Accumulate);
 
 			UMassStateTreeSubsystem& MassStateTreeSubsystem = ExecutionContext.GetMutableSubsystemChecked<UMassStateTreeSubsystem>();
 
diff --git a/Engine/Plugins/Runtime/StateTree/Source/StateTreeModule/Private/StateTreeExecutionContext.cpp b/Engine/Plugins/Runtime/StateTree/Source/StateTreeModule/Private/StateTreeExecutionContext.cpp
--- a/Engine/Plugins/Runtime/StateTree/Source/StateTreeModule/Private/StateTreeExecutionContext.cpp
+++ b/Engine/Plugins/Runtime/StateTree/Source/StateTreeModule/Private/StateTreeExecutionContext.cpp
@@ -25,6 +25,8 @@
 #define STATETREE_LOG(Verbosity, Format, ...) UE_VLOG_UELOG(GetOwner(), LogStateTree, Verbosity, TEXT("%s: ") Format, *GetInstanceDescriptionInternal(), ##__VA_ARGS__)
 #define STATETREE_CLOG(Condition, Verbosity, Format, ...) UE_CVLOG_UELOG((Condition), GetOwner(), LogStateTree, Verbosity, TEXT("%s: ") Format, *GetInstanceDescriptionInternal(), ##__VA_ARGS__)
 
+CSV_DEFINE_CATEGORY(StateTree, false);
+
 namespace UE::StateTree::Debug
 {
 	// Debug printing indent for hierarchical data.
@@ -907,7 +909,7 @@ void FStateTreeMinimalExecutionContext::RemoveScheduledTickRequest(UE::StateTree
 
 void FStateTreeMinimalExecutionContext::SendEvent(const FGameplayTag Tag, const FConstStructView Payload, const FName Origin)
 {
-	CSV_SCOPED_TIMING_STAT_EXCLUSIVE(StateTree_SendEvent);
+	CSV_SCOPED_TIMING_STAT(StateTree, SendEvent);
 
 	if (!IsValid())
 	{
@@ -1450,7 +1452,7 @@ void FStateTreeExecutionContext::SetUpdatePhaseInExecutionState(FStateTreeExecut
 
 EStateTreeRunStatus FStateTreeExecutionContext::Start(FStartParameters Parameters)
 {
-	CSV_SCOPED_TIMING_STAT_EXCLUSIVE(StateTree_Start);
+	CSV_SCOPED_TIMING_STAT(StateTree, Start);
 	TRACE_CPUPROFILER_EVENT_SCOPE(FStateTreeExecutionContext::Start);
 	UE_STATETREE_CRASH_REPORTER_SCOPE(&Owner, &RootStateTree, UE::StateTree::ExecutionContext::Private::Name_Start.Resolve());
 
@@ -1663,7 +1665,7 @@ EStateTreeRunStatus FStateTreeExecutionContext::Start(FStartParameters Parameter
 
 EStateTreeRunStatus FStateTreeExecutionContext::Stop(EStateTreeRunStatus CompletionStatus)
 {
-	CSV_SCOPED_TIMING_STAT_EXCLUSIVE(StateTree_Stop);
+	CSV_SCOPED_TIMING_STAT(StateTree, Stop);
 	TRACE_CPUPROFILER_EVENT_SCOPE(FStateTreeExecutionContext::Stop);
 	UE_STATETREE_CRASH_REPORTER_SCOPE(&Owner, &RootStateTree, UE::StateTree::ExecutionContext::Private::Name_Stop.Resolve());
 
@@ -1799,7 +1801,7 @@ EStateTreeRunStatus FStateTreeExecutionContext::TickPostlude()
 
 EStateTreeRunStatus FStateTreeExecutionContext::Tick(const float DeltaTime)
 {
-	CSV_SCOPED_TIMING_STAT_EXCLUSIVE(StateTree_Tick);
+	CSV_SCOPED_TIMING_STAT(StateTree, Tick);
 	TRACE_CPUPROFILER_EVENT_SCOPE(FStateTreeExecutionContext::Tick);
 	UE_STATETREE_CRASH_REPORTER_SCOPE(&Owner, &RootStateTree, UE::StateTree::ExecutionContext::Private::Name_Tick.Resolve());
 
@@ -1819,7 +1821,7 @@ EStateTreeRunStatus FStateTreeExecutionContext::Tick(const float DeltaTime)
 
 EStateTreeRunStatus FStateTreeExecutionContext::TickUpdateTasks(const float DeltaTime)
 {
-	CSV_SCOPED_TIMING_STAT_EXCLUSIVE(StateTree_Tick);
+	CSV_SCOPED_TIMING_STAT(StateTree, Tick);
 	TRACE_CPUPROFILER_EVENT_SCOPE(FStateTreeExecutionContext::TickUpdateTasks);
 	UE_STATETREE_CRASH_REPORTER_SCOPE(&Owner, &RootStateTree, UE::StateTree::ExecutionContext::Private::Name_Tick.Resolve());
 
@@ -1838,7 +1840,7 @@ EStateTreeRunStatus FStateTreeExecutionContext::TickUpdateTasks(const float Delt
 
 EStateTreeRunStatus FStateTreeExecutionContext::TickTriggerTransitions()
 {
-	CSV_SCOPED_TIMING_STAT_EXCLUSIVE(StateTree_Tick);
+	CSV_SCOPED_TIMING_STAT(StateTree, Tick);
 	TRACE_CPUPROFILER_EVENT_SCOPE(FStateTreeExecutionContext::TickTriggerTransitions);
 	UE_STATETREE_CRASH_REPORTER_SCOPE(&Owner, &RootStateTree, UE::StateTree::ExecutionContext::Private::Name_Tick.Resolve());
 
@@ -2120,7 +2122,7 @@ void FStateTreeExecutionContext::RequestTransition(const FStateTreeTransitionReq
 	using namespace UE::StateTree;
 	using namespace UE::StateTree::ExecutionContext;
 
-	CSV_SCOPED_TIMING_STAT_EXCLUSIVE(StateTree_RequestTransition);
+	CSV_SCOPED_TIMING_STAT(StateTree, RequestTransition);
 
 	if (!IsValid())
 	{
@@ -3557,7 +3559,7 @@ EStateTreeRunStatus FStateTreeExecutionContext::EnterState(const TSharedPtr<FSel
 	using namespace UE::StateTree::ExecutionContext;
 	using namespace UE::StateTree::ExecutionContext::Private;
 
-	CSV_SCOPED_TIMING_STAT_EXCLUSIVE(StateTree_EnterState);
+	CSV_SCOPED_TIMING_STAT(StateTree, EnterState);
 
 	ON_SCOPE_EXIT{ InstanceData.ResetTemporaryInstances(); };
 
@@ -3761,7 +3763,7 @@ EStateTreeRunStatus FStateTreeExecutionContext::EnterState(const TSharedPtr<FSel
 					EStateTreeRunStatus TaskRunStatus = EStateTreeRunStatus::Unset;
 					{
 						QUICK_SCOPE_CYCLE_COUNTER(StateTree_Task_EnterState);
-						CSV_SCOPED_TIMING_STAT_EXCLUSIVE(StateTree_Task_EnterState);
+						CSV_SCOPED_TIMING_STAT(StateTree, Task_EnterState);
 
 						TaskRunStatus = Task.EnterState(*this, CurrentTransition);
 					}
@@ -3825,7 +3827,7 @@ void FStateTreeExecutionContext::ExitState(const TSharedPtr<const FSelectStateRe
 	using namespace UE::StateTree::ExecutionContext;
 	using namespace UE::StateTree::ExecutionContext::Private;
 
-	CSV_SCOPED_TIMING_STAT_EXCLUSIVE(StateTree_ExitState);
+	CSV_SCOPED_TIMING_STAT(StateTree, ExitState);
 
 	FStateTreeExecutionState& Exec = GetExecState();
 	if (Exec.ActiveFrames.IsEmpty())
@@ -3941,7 +3943,7 @@ void FStateTreeExecutionContext::ExitState(const TSharedPtr<const FSelectStateRe
 						UE_STATETREE_DEBUG_TASK_EXIT_STATE(this, CurrentStateTree, FStateTreeIndex16(TaskIndex));
 						{
 							QUICK_SCOPE_CYCLE_COUNTER(StateTree_Task_ExitState);
-							CSV_SCOPED_TIMING_STAT_EXCLUSIVE(StateTree_Task_ExitState);
+							CSV_SCOPED_TIMING_STAT(StateTree, Task_ExitState);
 							Task.ExitState(*this, CurrentTransition);
 						}
 
@@ -3997,7 +3999,7 @@ void FStateTreeExecutionContext::RemoveAllDelegateListeners()
 
 void FStateTreeExecutionContext::StateCompleted()
 {
-	CSV_SCOPED_TIMING_STAT_EXCLUSIVE(StateTree_StateCompleted);
+	CSV_SCOPED_TIMING_STAT(StateTree, StateCompleted);
 
 	const FStateTreeExecutionState& Exec = GetExecState();
 
@@ -4134,7 +4136,7 @@ EStateTreeRunStatus FStateTreeExecutionContext::TickEvaluatorsAndGlobalTasks(con
 	// A global task can complete async. See CompletedStates.
 	// When a global task fails, stop ticking the following tasks.
 
-	CSV_SCOPED_TIMING_STAT_EXCLUSIVE(StateTree_TickEvaluators);
+	CSV_SCOPED_TIMING_STAT(StateTree, TickEvaluators);
 	UE_STATETREE_DEBUG_SCOPED_PHASE(this, EStateTreeUpdatePhase::TickingGlobalTasks);
 
 	STATETREE_LOG(VeryVerbose, TEXT("Ticking Evaluators & Global Tasks"));
@@ -4260,7 +4262,7 @@ EStateTreeRunStatus FStateTreeExecutionContext::StartGlobalsForFrameInternal(con
 		return EStateTreeRunStatus::Running;
 	}
 
-	CSV_SCOPED_TIMING_STAT_EXCLUSIVE(StateTree_StartEvaluators);
+	CSV_SCOPED_TIMING_STAT(StateTree, StartEvaluators);
 
 	FCurrentlyProcessedFrameScope FrameScope(*this, CurrentParentFrame, CurrentFrame);
 	
@@ -4444,7 +4446,7 @@ EStateTreeRunStatus FStateTreeExecutionContext::StartEvaluatorsAndGlobalTasks(FS
 
 EStateTreeRunStatus FStateTreeExecutionContext::StartEvaluatorsAndGlobalTasks()
 {
-	CSV_SCOPED_TIMING_STAT_EXCLUSIVE(StateTree_StartEvaluatorsAndGlobalTasks);
+	CSV_SCOPED_TIMING_STAT(StateTree, StartEvaluatorsAndGlobalTasks);
 
 	STATETREE_LOG(Verbose, TEXT("Start Evaluators & Global tasks"));
 
@@ -4488,7 +4490,7 @@ void FStateTreeExecutionContext::StopEvaluatorsAndGlobalTasks(const EStateTreeRu
 
 void FStateTreeExecutionContext::StopEvaluatorsAndGlobalTasks(const EStateTreeRunStatus CompletionStatus)
 {
-	CSV_SCOPED_TIMING_STAT_EXCLUSIVE(StateTree_StopEvaluatorsAndGlobalTasks);
+	CSV_SCOPED_TIMING_STAT(StateTree, StopEvaluatorsAndGlobalTasks);
 
 	STATETREE_LOG(Verbose, TEXT("Stop Evaluators & Global Tasks"));
 
@@ -4651,7 +4653,7 @@ EStateTreeRunStatus FStateTreeExecutionContext::TickTasks(const float DeltaTime)
 	// When a task fails, stop ticking the following tasks.
 	// When no task ticks, then the leaf completes.
 
-	CSV_SCOPED_TIMING_STAT_EXCLUSIVE(StateTree_TickTasks);
+	CSV_SCOPED_TIMING_STAT(StateTree, TickTasks);
 	UE_STATETREE_DEBUG_SCOPED_PHASE(this, EStateTreeUpdatePhase::TickingTasks);
 
 	using namespace UE::StateTree;
@@ -4895,7 +4897,7 @@ FStateTreeExecutionContext::FTickTaskResult FStateTreeExecutionContext::TickTask
 		EStateTreeRunStatus TaskRunStatus = EStateTreeRunStatus::Unset;
 		{
 			QUICK_SCOPE_CYCLE_COUNTER(StateTree_Task_Tick);
-			CSV_SCOPED_TIMING_STAT_EXCLUSIVE(StateTree_Task_Tick);
+			CSV_SCOPED_TIMING_STAT(StateTree, Task_Tick);
 
 			TaskRunStatus = Task.Tick(*this, Args.DeltaTime);
 		}
@@ -4949,7 +4951,7 @@ bool FStateTreeExecutionContext::TestAllConditionsWithValidation(const FStateTre
 template<bool bOnActiveInstances>
 bool FStateTreeExecutionContext::TestAllConditionsInternal(const FStateTreeExecutionFrame* CurrentParentFrame, const FStateTreeExecutionFrame& CurrentFrame, FStateTreeStateHandle CurrentStateHandle, const FMemoryRequirement& MemoryRequirement, const int32 ConditionsOffset, const int32 ConditionsNum, EStateTreeUpdatePhase Phase)
 {
-	CSV_SCOPED_TIMING_STAT_EXCLUSIVE(StateTree_TestConditions);
+	CSV_SCOPED_TIMING_STAT(StateTree, TestConditions);
 
 	using namespace UE::StateTree::InstanceData;
 	using namespace UE::StateTree::ExecutionContext::Private;
@@ -5172,7 +5174,7 @@ float FStateTreeExecutionContext::EvaluateUtility(const FStateTreeExecutionFrame
 float FStateTreeExecutionContext::EvaluateUtilityWithValidation(const FStateTreeExecutionFrame* CurrentParentFrame, const FStateTreeExecutionFrame& CurrentFrame, FStateTreeStateHandle CurrentStateHandle, const FMemoryRequirement& MemoryRequirement, const int32 ConsiderationsBegin, const int32 ConsiderationsNum, const float StateWeight)
 {
 	// @todo: Tracing support
-	CSV_SCOPED_TIMING_STAT_EXCLUSIVE(StateTree_EvaluateUtility);
+	CSV_SCOPED_TIMING_STAT(StateTree, EvaluateUtility);
 
 	using namespace UE::StateTree::InstanceData;
 	using namespace UE::StateTree::ExecutionContext::Private;
@@ -5629,7 +5631,7 @@ bool FStateTreeExecutionContext::TriggerTransitions()
 	//3. If no transition, Process completion transitions, from bottom to top.
 	//4. If transition occurs, check if there are any frame (sub-tree) that completed.
 
-	CSV_SCOPED_TIMING_STAT_EXCLUSIVE(StateTree_TriggerTransition);
+	CSV_SCOPED_TIMING_STAT(StateTree, TriggerTransition);
 	UE_STATETREE_DEBUG_SCOPED_PHASE(this, EStateTreeUpdatePhase::TriggerTransitions);
 
 	// Set flag for the scope of this function to allow direct transitions without buffering.
@@ -6786,7 +6788,7 @@ bool FStateTreeExecutionContext::SelectStateInternal(
 	const FSelectStateInternalArguments& InternalArgs,
 	const TSharedRef<FSelectStateResult>& OutSelectionResult)
 {
-	CSV_SCOPED_TIMING_STAT_EXCLUSIVE(StateTree_SelectState);
+	CSV_SCOPED_TIMING_STAT(StateTree, SelectState);
 
 	using namespace UE::StateTree;
 	using namespace UE::StateTree::ExecutionContext;
