[StateTree] To save runtime memory, added a flag to not add the node for validation if the option was disabled before the instance started. Remove the node from the list when the node exits.

#Synced-CL 50969155

diff --git a/Engine/Plugins/Runtime/StateTree/Source/StateTreeModule/Private/Debugger/StateTreeRuntimeValidation.cpp b/Engine/Plugins/Runtime/StateTree/Source/StateTreeModule/Private/Debugger/StateTreeRuntimeValidation.cpp
--- a/Engine/Plugins/Runtime/StateTree/Source/StateTreeModule/Private/Debugger/StateTreeRuntimeValidation.cpp
+++ b/Engine/Plugins/Runtime/StateTree/Source/StateTreeModule/Private/Debugger/StateTreeRuntimeValidation.cpp
@@ -62,16 +62,14 @@ FRuntimeValidationInstanceData::~FRuntimeValidationInstanceData()
 	{
 		for (const FNodeStatePair& Pair : NodeStates)
 		{
-			if (EnumHasAllFlags(Pair.State, EState::BetweenEnterExitState))
-			{
-				ensureAlwaysMsgf(false, TEXT("Tree exited. Missing ExitState on %s."), *Private::NodeToString(StateTree.Get(), Pair.NodeID));
-				Private::bRuntimeValidationEnterExitState = false;
-			}
+			ensureAlwaysMsgf(false, TEXT("Tree exited. Missing ExitState on %s."), *Private::NodeToString(StateTree.Get(), Pair.NodeID));
+			Private::bRuntimeValidationEnterExitState = false;
+			break;
 		}
 	}
 }
 
-void FRuntimeValidationInstanceData::SetContext(const UObject* InNewOwner, const UStateTree* InNewStateTree, bool bInInstanceDataWriteAccessAcquired)
+void FRuntimeValidationInstanceData::SetContext(TNotNull<const UObject*> InNewOwner, TNotNull<const UStateTree*> InNewStateTree, const bool bInInstanceDataWriteAccessAcquired)
 {
 	TWeakObjectPtr<const UStateTree> NewStateTree = InNewStateTree;
 	TWeakObjectPtr<const UObject> NewOwner = InNewOwner;
@@ -82,7 +80,7 @@ void FRuntimeValidationInstanceData::SetContext(const UObject* InNewOwner, const
 			ensureAlwaysMsgf(false, TEXT("StateTree runtime check failed: The StateTree '%s' is different from the previously set '%s'.\n"
 				"Make sure you initialize FStateTreeExecutionContext with the same value every time.\n"
 				"Auto deactivate Runtime check StateTree.RuntimeValidation.Context to prevent reporting the same error multiple times.")
-				, InNewStateTree ? *InNewStateTree->GetFullName() : TEXT("StateTree"), *StateTree.Get()->GetFullName());
+				, *InNewStateTree->GetFullName(), *StateTree.Get()->GetFullName());
 			Private::bRuntimeValidationContext = false;
 		}
 		if (Owner.IsValid() && Owner != NewOwner)
@@ -90,7 +88,7 @@ void FRuntimeValidationInstanceData::SetContext(const UObject* InNewOwner, const
 			ensureAlwaysMsgf(false, TEXT("StateTree runtime check failed: The owner '%s' is different from the previously set '%s'.\n"
 				"Make sure you initialize FStateTreeExecutionContext with the same values every time.\n"
 				"Auto deactivate Runtime check StateTree.RuntimeValidation.Context to prevent reporting the same error multiple times.")
-				, InNewOwner ? *InNewOwner->GetFullName() : TEXT("owner"), *Owner.Get()->GetFullName());
+				, *InNewOwner->GetFullName(), *Owner.Get()->GetFullName());
 			Private::bRuntimeValidationContext = false;
 		}
 	}
@@ -101,62 +99,64 @@ void FRuntimeValidationInstanceData::SetContext(const UObject* InNewOwner, const
 	StateTree = NewStateTree;
 	Owner = NewOwner;
 	bInstanceDataWriteAccessAcquired |= bInInstanceDataWriteAccessAcquired;
+	if (!bMaintainNodeStates.IsSet())
+	{
+		bMaintainNodeStates = Private::bRuntimeValidationEnterExitState;
+	}
 }
 
-void FRuntimeValidationInstanceData::NodeEnterState(FGuid NodeID, FActiveFrameID FrameID)
+void FRuntimeValidationInstanceData::NodeEnterState(const FGuid& NodeID, FActiveFrameID FrameID)
 {
-	FNodeStatePair* Found = NodeStates.FindByPredicate([&NodeID, &FrameID](const FNodeStatePair& Other)
+	if (bMaintainNodeStates.Get(false))
 	{
-		return Other.NodeID == NodeID && Other.FrameID == FrameID;
-	});
-	if (Found)
-	{
-		if (Private::bRuntimeValidationEnterExitState && EnumHasAllFlags(Found->State, EState::BetweenEnterExitState))
+		FNodeStatePair* Found = NodeStates.FindByPredicate([&NodeID, &FrameID](const FNodeStatePair& Other)
+			{
+				return Other.NodeID == NodeID && Other.FrameID == FrameID;
+			});
+		if (Found)
 		{
-			ensureAlwaysMsgf(false, TEXT("StateTree runtime check failed: EnterState executed on node %s without an ExitState.\n"
-				"Auto deactivate Runtime check StateTree.RuntimeValidation.EnterExitState to prevent reporting the same error multiple times.")
-				, *Private::NodeToString(Owner.Get(), NodeID));
-			Private::bRuntimeValidationEnterExitState = false;
+			if (Private::bRuntimeValidationEnterExitState)
+			{
+				ensureAlwaysMsgf(false, TEXT("StateTree runtime check failed: EnterState executed on node %s without an ExitState.\n"
+					"Auto deactivate Runtime check StateTree.RuntimeValidation.EnterExitState to prevent reporting the same error multiple times.")
+					, *Private::NodeToString(Owner.Get(), NodeID));
+				Private::bRuntimeValidationEnterExitState = false;
+			}
+		}
+		else
+		{
+			NodeStates.Add(FNodeStatePair{ .NodeID = NodeID, .FrameID = FrameID });
 		}
-		EnumAddFlags(Found->State, EState::BetweenEnterExitState);
-	}
-	else
-	{
-		NodeStates.Add(FNodeStatePair{.NodeID = NodeID, .FrameID = FrameID, .State = EState::BetweenEnterExitState });
 	}
 }
 
-void FRuntimeValidationInstanceData::NodeExitState(FGuid NodeID, FActiveFrameID FrameID)
+void FRuntimeValidationInstanceData::NodeExitState(const FGuid& NodeID, FActiveFrameID FrameID)
 {
-	FNodeStatePair* Found = NodeStates.FindByPredicate([&NodeID, &FrameID](const FNodeStatePair& Other)
+	if (bMaintainNodeStates.Get(false))
 	{
-		return Other.NodeID == NodeID && Other.FrameID == FrameID;
-	});
-	if (Found)
-	{
-		if (Private::bRuntimeValidationEnterExitState && !EnumHasAllFlags(Found->State, EState::BetweenEnterExitState))
+		const int32 FoundIndex = NodeStates.IndexOfByPredicate([&NodeID, &FrameID](const FNodeStatePair& Other)
+			{
+				return Other.NodeID == NodeID && Other.FrameID == FrameID;
+			});
+		if (FoundIndex != INDEX_NONE)
+		{
+			NodeStates.RemoveAtSwap(FoundIndex);
+		}
+		else if (Private::bRuntimeValidationEnterExitState)
 		{
 			ensureAlwaysMsgf(false, TEXT("StateTree runtime check failed: ExitState executed on node %s without an EnterState.\n"
 				"Auto deactivate Runtime check StateTree.RuntimeValidation.EnterExitState to prevent reporting the same error multiple times.")
 				, *Private::NodeToString(Owner.Get(), NodeID));
 			Private::bRuntimeValidationEnterExitState = false;
 		}
-		EnumRemoveFlags(Found->State, EState::BetweenEnterExitState);
-	}
-	else if (Private::bRuntimeValidationEnterExitState)
-	{
-		ensureAlwaysMsgf(false, TEXT("StateTree runtime check failed: ExitState executed on node %s without an EnterState.\n"
-			"Auto deactivate Runtime check StateTree.RuntimeValidation.EnterExitState to prevent reporting the same error multiple times.")
-			, *Private::NodeToString(Owner.Get(), NodeID));
-		Private::bRuntimeValidationEnterExitState = false;
 	}
 }
 
-void FRuntimeValidationInstanceData::ValidateTreeNodes(const UStateTree* InNewStateTree) const
+void FRuntimeValidationInstanceData::ValidateTreeNodes(TNotNull<const UStateTree*> InNewStateTree) const
 {
 	if (Private::bRuntimeValidationDoesNewerVersionExists)
 	{
-		if (InNewStateTree && InNewStateTree->IsReadyToRun())
+		if (InNewStateTree->IsReadyToRun())
 		{
 			auto DoesNewerVersionExists = [](const UObject* InstanceDataType)
 			{
@@ -225,26 +225,23 @@ void FRuntimeValidationInstanceData::ValidateTreeNodes(const UStateTree* InNewSt
 	}
 }
 
-void FRuntimeValidationInstanceData::ValidateInstanceData(const UStateTree* NewStateTree)
+void FRuntimeValidationInstanceData::ValidateInstanceData(TNotNull<const UStateTree*> NewStateTree)
 {
 	if (Private::bRuntimeValidationInstanceData)
 	{
-		if (NewStateTree)
+		if (!NewStateTree->GetDefaultInstanceData().GetStorage().AreAllInstancesValid())
 		{
-			if (!NewStateTree->GetDefaultInstanceData().GetStorage().AreAllInstancesValid())
-			{
-				ensureAlwaysMsgf(false, TEXT("The instance data has invalid data."));
-				Private::bRuntimeValidationInstanceData = false;
-			}
+			ensureAlwaysMsgf(false, TEXT("The instance data has invalid data."));
+			Private::bRuntimeValidationInstanceData = false;
+		}
 
-			TSharedPtr<FStateTreeInstanceData> SharedPtr = NewStateTree->GetSharedInstanceData();
-			if (SharedPtr)
+		TSharedPtr<FStateTreeInstanceData> SharedPtr = NewStateTree->GetSharedInstanceData();
+		if (SharedPtr)
+		{
+			if (!SharedPtr->GetStorage().AreAllInstancesValid())
 			{
-				if (!SharedPtr->GetStorage().AreAllInstancesValid())
-				{
-					ensureAlwaysMsgf(false, TEXT("The shared instance data has invalid data."));
-					Private::bRuntimeValidationInstanceData = false;
-				}
+				ensureAlwaysMsgf(false, TEXT("The shared instance data has invalid data."));
+				Private::bRuntimeValidationInstanceData = false;
 			}
 		}
 	}
@@ -267,7 +264,7 @@ FRuntimeValidationInstanceData* FRuntimeValidation::GetInstanceData() const
 	return RuntimeValidationData;
 }
 
-void FRuntimeValidation::SetContext(const UObject* Owner, const UStateTree* StateTree, bool bInstanceDataWriteAccessAcquired) const
+void FRuntimeValidation::SetContext(TNotNull<const UObject*> Owner, TNotNull<const UStateTree*> StateTree, bool bInstanceDataWriteAccessAcquired) const
 {
 	if (RuntimeValidationData)
 	{
@@ -277,7 +274,7 @@ void FRuntimeValidation::SetContext(const UObject* Owner, const UStateTree* Stat
 
 #else
 
-void FRuntimeValidation::SetContext(const UObject*, const UStateTree*, bool) const
+void FRuntimeValidation::SetContext(TNotNull<const UObject*>, TNotNull<const UStateTree*>, bool) const
 {
 }
 
diff --git a/Engine/Plugins/Runtime/StateTree/Source/StateTreeModule/Private/Debugger/StateTreeRuntimeValidationInstanceData.h b/Engine/Plugins/Runtime/StateTree/Source/StateTreeModule/Private/Debugger/StateTreeRuntimeValidationInstanceData.h
--- a/Engine/Plugins/Runtime/StateTree/Source/StateTreeModule/Private/Debugger/StateTreeRuntimeValidationInstanceData.h
+++ b/Engine/Plugins/Runtime/StateTree/Source/StateTreeModule/Private/Debugger/StateTreeRuntimeValidationInstanceData.h
@@ -20,39 +20,36 @@ public:
 	FRuntimeValidationInstanceData() = default;
 	~FRuntimeValidationInstanceData();
 
-	void SetContext(const UObject* InOwner, const UStateTree* InStateTree, bool bInInstanceDataWriteAccessAcquired);
+	void SetContext(TNotNull<const UObject*> InOwner, TNotNull<const UStateTree*> InStateTree, bool bInInstanceDataWriteAccessAcquired);
 
-	void NodeEnterState(FGuid NodeID, FActiveFrameID FrameID);
-	void NodeExitState(FGuid NodeID, FActiveFrameID FrameID);
+	void NodeEnterState(const FGuid& NodeID, FActiveFrameID FrameID);
+	void NodeExitState(const FGuid& NodeID, FActiveFrameID FrameID);
+	void NodeTick(const FGuid& NodeID, FActiveFrameID FrameID);
 
 private:
-	void ValidateTreeNodes(const UStateTree* NewStateTree) const;
-	void ValidateInstanceData(const UStateTree* NewStateTree);
+	void ValidateTreeNodes(TNotNull<const UStateTree*> NewStateTree) const;
+	void ValidateInstanceData(TNotNull<const UStateTree*> NewStateTree);
 
 private:
-	enum class EState : uint8
-	{
-		None = 0x00,
-		BetweenEnterExitState = 0x01,
-	};
-	FRIEND_ENUM_CLASS_FLAGS(EState);
-
 	struct FNodeStatePair
 	{
 		FGuid NodeID;
 		FActiveFrameID FrameID;
-		EState State = EState::None;
 	};
 	TArray<FNodeStatePair> NodeStates;
 	TWeakObjectPtr<const UStateTree> StateTree;
 	TWeakObjectPtr<const UObject> Owner;
 
-	/** an ExecContext with write access to the Instance Storage data this is debugging on has been created */
+	/** An ExecContext with write access to the Instance Storage data this is debugging on has been created */
 	bool bInstanceDataWriteAccessAcquired = false;
+	/**
+	 * The execution started with bRuntimeValidationEnterExitState set to true.
+	 * It will keep maintaining the NodeStates array even if bRuntimeValidationEnterExitState is set to false afterward.
+	 * The bRuntimeValidationEnterExitState can be reset to true after a fail, and we need to keep maintaining the array.
+	 */
+	TOptional<bool> bMaintainNodeStates;
 };
 
-ENUM_CLASS_FLAGS(FRuntimeValidationInstanceData::EState);
-
 } // UE::StateTree::Debug
 
 #endif // WITH_STATETREE_DEBUG
\ No newline at end of file
diff --git a/Engine/Plugins/Runtime/StateTree/Source/StateTreeModule/Public/Debugger/StateTreeRuntimeValidation.h b/Engine/Plugins/Runtime/StateTree/Source/StateTreeModule/Public/Debugger/StateTreeRuntimeValidation.h
--- a/Engine/Plugins/Runtime/StateTree/Source/StateTreeModule/Public/Debugger/StateTreeRuntimeValidation.h
+++ b/Engine/Plugins/Runtime/StateTree/Source/StateTreeModule/Public/Debugger/StateTreeRuntimeValidation.h
@@ -35,7 +35,7 @@ public:
 #endif
 
 	/** Set the Owner of the data */
-	void SetContext(const UObject* Owner, const UStateTree* StateTree, bool bInstanceDataWriteAccessAcquired) const;
+	void SetContext(TNotNull<const UObject*> Owner, TNotNull<const UStateTree*> StateTree, bool bInstanceDataWriteAccessAcquired) const;
 
 private:
 #if WITH_STATETREE_DEBUG
