[State Tree] no need to check shared pointer validity

#Synced-CL 51004728

diff --git a/Engine/Plugins/Runtime/StateTree/Source/StateTreeEditorModule/Private/Customizations/StateTreeBindingExtension.cpp b/Engine/Plugins/Runtime/StateTree/Source/StateTreeEditorModule/Private/Customizations/StateTreeBindingExtension.cpp
--- a/Engine/Plugins/Runtime/StateTree/Source/StateTreeEditorModule/Private/Customizations/StateTreeBindingExtension.cpp
+++ b/Engine/Plugins/Runtime/StateTree/Source/StateTreeEditorModule/Private/Customizations/StateTreeBindingExtension.cpp
@@ -833,19 +833,16 @@ bool FStateTreeBindingExtension::GetPromotionToParameterOverrideInternal(const T
 		CurrentHandle = ParentHandle;
 	}
 
-	if (CurrentHandle.IsValid())
+	TArray<UObject*> RootObjects;
+	CurrentHandle->GetOuterObjects(RootObjects);
+	if (RootObjects.Num() == 1)
 	{
-		TArray<UObject*> RootObjects;
-		CurrentHandle->GetOuterObjects(RootObjects);
-		if (RootObjects.Num() == 1)
+		if (const UStateTreeState* State = Cast<UStateTreeState>(RootObjects[0]))
 		{
-			if (const UStateTreeState* State = Cast<UStateTreeState>(RootObjects[0]))
+			if (State->Type == EStateTreeStateType::Linked || State->Type == EStateTreeStateType::LinkedAsset)
 			{
-				if (State->Type == EStateTreeStateType::Linked || State->Type == EStateTreeStateType::LinkedAsset)
-				{
-					bOutOverride = false;
-					return true;
-				}
+				bOutOverride = false;
+				return true;
 			}
 		}
 	}
