[StateTree] Add SCOPE_CYCLE_UOBJECT so we can tell which objects are running in profiler captures.

#rnx

#Synced-CL 51047059

diff --git a/Engine/Plugins/Runtime/StateTree/Source/StateTreeModule/Private/StateTreeExecutionContext.cpp b/Engine/Plugins/Runtime/StateTree/Source/StateTreeModule/Private/StateTreeExecutionContext.cpp
--- a/Engine/Plugins/Runtime/StateTree/Source/StateTreeModule/Private/StateTreeExecutionContext.cpp
+++ b/Engine/Plugins/Runtime/StateTree/Source/StateTreeModule/Private/StateTreeExecutionContext.cpp
@@ -1455,6 +1455,7 @@ EStateTreeRunStatus FStateTreeExecutionContext::Start(FStartParameters Parameter
 	CSV_SCOPED_TIMING_STAT(StateTree, Start);
 	TRACE_CPUPROFILER_EVENT_SCOPE(FStateTreeExecutionContext::Start);
 	UE_STATETREE_CRASH_REPORTER_SCOPE(&Owner, &RootStateTree, UE::StateTree::ExecutionContext::Private::Name_Start.Resolve());
+	SCOPE_CYCLE_UOBJECT(StateTree, &Owner);
 
 	using namespace UE::StateTree;
 	using namespace UE::StateTree::ExecutionContext;
@@ -1668,6 +1669,7 @@ EStateTreeRunStatus FStateTreeExecutionContext::Stop(EStateTreeRunStatus Complet
 	CSV_SCOPED_TIMING_STAT(StateTree, Stop);
 	TRACE_CPUPROFILER_EVENT_SCOPE(FStateTreeExecutionContext::Stop);
 	UE_STATETREE_CRASH_REPORTER_SCOPE(&Owner, &RootStateTree, UE::StateTree::ExecutionContext::Private::Name_Stop.Resolve());
+	SCOPE_CYCLE_UOBJECT(StateTree, &Owner);
 
 	if (!IsValid())
 	{
@@ -1804,6 +1806,7 @@ EStateTreeRunStatus FStateTreeExecutionContext::Tick(const float DeltaTime)
 	CSV_SCOPED_TIMING_STAT(StateTree, Tick);
 	TRACE_CPUPROFILER_EVENT_SCOPE(FStateTreeExecutionContext::Tick);
 	UE_STATETREE_CRASH_REPORTER_SCOPE(&Owner, &RootStateTree, UE::StateTree::ExecutionContext::Private::Name_Tick.Resolve());
+	SCOPE_CYCLE_UOBJECT(StateTree, &Owner);
 
 	TGuardValue<bool> ScheduledNextTickScope(bAllowedToScheduleNextTick, false);
 
@@ -1824,6 +1827,7 @@ EStateTreeRunStatus FStateTreeExecutionContext::TickUpdateTasks(const float Delt
 	CSV_SCOPED_TIMING_STAT(StateTree, Tick);
 	TRACE_CPUPROFILER_EVENT_SCOPE(FStateTreeExecutionContext::TickUpdateTasks);
 	UE_STATETREE_CRASH_REPORTER_SCOPE(&Owner, &RootStateTree, UE::StateTree::ExecutionContext::Private::Name_Tick.Resolve());
+	SCOPE_CYCLE_UOBJECT(StateTree, &Owner);
 
 	TGuardValue<bool> ScheduledNextTickScope(bAllowedToScheduleNextTick, false);
 
@@ -1843,6 +1847,7 @@ EStateTreeRunStatus FStateTreeExecutionContext::TickTriggerTransitions()
 	CSV_SCOPED_TIMING_STAT(StateTree, Tick);
 	TRACE_CPUPROFILER_EVENT_SCOPE(FStateTreeExecutionContext::TickTriggerTransitions);
 	UE_STATETREE_CRASH_REPORTER_SCOPE(&Owner, &RootStateTree, UE::StateTree::ExecutionContext::Private::Name_Tick.Resolve());
+	SCOPE_CYCLE_UOBJECT(StateTree, &Owner);
 
 	TGuardValue<bool> ScheduledNextTickScope(bAllowedToScheduleNextTick, false);
 
