[State Tree]
- Disabled Promote to linked state parameter option.
- Fixed an issue where property cannot be promoted to state parameter if state parameter is empty.

#Synced-CL 51004728

diff --git a/Engine/Plugins/Runtime/StateTree/Source/StateTreeEditorModule/Private/StateTreeEditorData.cpp b/Engine/Plugins/Runtime/StateTree/Source/StateTreeEditorModule/Private/StateTreeEditorData.cpp
--- a/Engine/Plugins/Runtime/StateTree/Source/StateTreeEditorModule/Private/StateTreeEditorData.cpp
+++ b/Engine/Plugins/Runtime/StateTree/Source/StateTreeEditorModule/Private/StateTreeEditorData.cpp
@@ -585,8 +585,9 @@ bool UStateTreeEditorData::CanCreateParameter(const FGuid StructID) const
 
 	VisitHierarchy([&StructID, &bFoundStructID](UStateTreeState& State, UStateTreeState* ParentState)->EStateTreeVisitor
 	{
-		if (State.Parameters.ID == StructID)
+		if (State.Type != EStateTreeStateType::Linked && State.Type != EStateTreeStateType::LinkedAsset && State.Parameters.ID == StructID)
 		{
+			// Linked State parameter is fixed. No promotion is allowed
 			bFoundStructID = true;
 			return EStateTreeVisitor::Break;
 		}
@@ -1240,27 +1241,25 @@ EStateTreeVisitor UStateTreeEditorData::VisitStateNodes(const UStateTreeState& S
 	
 	if (bContinue)
 	{
+		check(State.Parameters.ID.IsValid());
 		// Bindable state parameters
-		if (State.Parameters.Parameters.IsValid())
+		if (VisitFuncNodes(State.Parameters.ID, UE::StateTree::Editor::ParametersNodeName) == EStateTreeVisitor::Break)
+		{
+			bContinue = false;
+		}
+		else
 		{
-			if (VisitFuncNodes(State.Parameters.ID, UE::StateTree::Editor::ParametersNodeName) == EStateTreeVisitor::Break)
+			FStateTreeBindableStructDesc Desc;
+			Desc.StatePath = StatePath;
+			Desc.Struct = State.Parameters.Parameters.GetPropertyBagStruct();
+			Desc.Name = FName("Parameters");
+			Desc.ID = State.Parameters.ID;
+			Desc.DataSource = EStateTreeBindableStructSource::StateParameter;
+
+			if (InFunc(&State, Desc, FStateTreeDataView(const_cast<FInstancedPropertyBag&>(State.Parameters.Parameters).GetMutableValue())) == EStateTreeVisitor::Break)
 			{
 				bContinue = false;
 			}
-			else
-			{
-				FStateTreeBindableStructDesc Desc;
-				Desc.StatePath = StatePath;
-				Desc.Struct = State.Parameters.Parameters.GetPropertyBagStruct();
-				Desc.Name = FName("Parameters");
-				Desc.ID = State.Parameters.ID;
-				Desc.DataSource = EStateTreeBindableStructSource::StateParameter;
-
-				if (InFunc(&State, Desc, FStateTreeDataView(const_cast<FInstancedPropertyBag&>(State.Parameters.Parameters).GetMutableValue())) == EStateTreeVisitor::Break)
-				{
-					bContinue = false;
-				}
-			}
 		}
 	}
 
