[StateTree] Fixed ID fixing after duplication. IDs for BP evaluators and tasks were not collected.

#Synced-CL 51092353

diff --git a/Engine/Plugins/Runtime/StateTree/Source/StateTreeEditorModule/Private/StateTreeEditorData.cpp b/Engine/Plugins/Runtime/StateTree/Source/StateTreeEditorModule/Private/StateTreeEditorData.cpp
--- a/Engine/Plugins/Runtime/StateTree/Source/StateTreeEditorModule/Private/StateTreeEditorData.cpp
+++ b/Engine/Plugins/Runtime/StateTree/Source/StateTreeEditorModule/Private/StateTreeEditorData.cpp
@@ -1063,7 +1063,7 @@ void UStateTreeEditorData::DuplicateIDs()
 
 	// Visit and create new ids
 	{
-		auto AddId = [&OldToNewIDs](const FGuid& OldID, bool bTestIfContains = true)
+		auto AddId = [&OldToNewIDs](const FGuid& OldID, bool bTestIfContains)
 		{
 			ensureAlwaysMsgf(bTestIfContains == false || !OldToNewIDs.Contains(OldID), TEXT("The id is duplicated and FixDuplicateIDs failed to fix it."));
 
@@ -1074,57 +1074,21 @@ void UStateTreeEditorData::DuplicateIDs()
 			}
 		};
 
-		auto AddIds = [&AddId](TArrayView<FStateTreeEditorNode> Nodes)
-		{
-			for (FStateTreeEditorNode& Node : Nodes)
+		VisitAllNodes([&AddId](const UStateTreeState* State, const FStateTreeBindableStructDesc& Desc, const FStateTreeDataView Value)
 			{
-				AddId(Node.ID);
-			}
-		};
-
-		// Do not use the VisitGlobalNodes because the schema should not be included in OldToNewIDs 
-		OldToNewIDs.Add(RootParametersGuid, FGuid::NewGuid());
-		AddIds(Evaluators);
-		AddIds(GlobalTasks);
-		for (FStateTreeEditorColor& Color : Colors)
-		{
-			AddId(Color.ColorRef.ID);
-		}
-
-		VisitHierarchy([Self = this, &AddId, &OldToNewIDs](const UStateTreeState& State, UStateTreeState* ParentState)
-		{
-			AddId(State.ID);
-			AddId(State.Parameters.ID);
+				if (State)
+				{
+					AddId(State->ID, false);
+				}
 
-			for (const FStateTreeTransition& Transition : State.Transitions)
-			{
-				AddId(Transition.ID);
-			}
+				// Schema should not be included in OldToNewIDs
+				if(Desc.DataSource != EStateTreeBindableStructSource::Context)
+				{
+					AddId(Desc.ID, true);
+				}
 
-			return Self->VisitStateNodes(State, [&AddId](const UStateTreeState* State, const FStateTreeBindableStructDesc& Desc, const FStateTreeDataView Value)
-			{
-				AddId(Desc.ID, /*bTestIfContains*/false);
 				return EStateTreeVisitor::Continue;
 			});
-		});
-
-		// Confirms that we collected everything.
-		{
-			TMap<FGuid, const FStateTreeDataView> AllStructValues;
-			GetAllStructValues(AllStructValues);
-			// Schema ids are not duplicated
-			if (Schema != nullptr)
-			{
-				for (const FStateTreeExternalDataDesc& ContextDesc : Schema->GetContextDataDescs())
-				{
-					AllStructValues.Remove(ContextDesc.ID);
-				}
-			}
-			for (auto& StructValue : AllStructValues)
-			{
-				ensureMsgf(OldToNewIDs.Contains(StructValue.Key), TEXT("An ID container was not duplicated for asset '%s'."), *GetOutermost()->GetName());
-			}
-		}
 	}
 
 	// Remap ids properties to the new generated ids
@@ -1148,7 +1112,7 @@ void UStateTreeEditorData::DuplicateIDs()
 
 						// Modify the value if needed.
 						// We const_cast because we don't change the layout and the iterator advance is still going to work.
-						FGuid* GuidValue = reinterpret_cast<FGuid*>(const_cast<void*>(It.Value()));
+						FGuid* GuidValue = static_cast<FGuid*>(const_cast<void*>(It.Value()));
 						if (FGuid* NewGuidValue = OldToNewIDs.Find(*GuidValue))
 						{
 							*GuidValue = *NewGuidValue;
