[StateTree] Add missing UStateTreeEditingSubsystem::MarkAsModified to the viewmodel operations.

#rnx
#Synced-CL 51092353

diff --git a/Engine/Plugins/Runtime/StateTree/Source/StateTreeEditorModule/Private/StateTreeState.cpp b/Engine/Plugins/Runtime/StateTree/Source/StateTreeEditorModule/Private/StateTreeState.cpp
--- a/Engine/Plugins/Runtime/StateTree/Source/StateTreeEditorModule/Private/StateTreeState.cpp
+++ b/Engine/Plugins/Runtime/StateTree/Source/StateTreeEditorModule/Private/StateTreeState.cpp
@@ -481,6 +481,11 @@ void UStateTreeState::SetParametersPropertyOverridden(const FGuid PropertyID, co
 			}
 		}
 	}
+
+	if (UStateTree* StateTree = GetTypedOuter<UStateTree>())
+	{
+		UStateTreeEditingSubsystem::MarkAsModified(StateTree);
+	}
 }
 
 const FInstancedPropertyBag* UStateTreeState::GetDefaultParameters() const
@@ -618,6 +623,11 @@ void UStateTreeState::SetLinkedState(FStateTreeStateLink InStateLink)
 	Parameters.bFixedLayout = true;
 	UpdateParametersFromLinkedSubtree();
 	SelectionBehavior = EStateTreeStateSelectionBehavior::TryEnterState;
+
+	if (UStateTree* StateTree = GetTypedOuter<UStateTree>())
+	{
+		UStateTreeEditingSubsystem::MarkAsModified(StateTree);
+	}
 }
 
 void UStateTreeState::SetLinkedStateAsset(UStateTree* InLinkedAsset)
@@ -630,4 +640,65 @@ void UStateTreeState::SetLinkedStateAsset(UStateTree* InLinkedAsset)
 	Parameters.bFixedLayout = true;
 	UpdateParametersFromLinkedSubtree();
 	SelectionBehavior = EStateTreeStateSelectionBehavior::TryEnterState;
-}
\ No newline at end of file
+
+	if (UStateTree* StateTree = GetTypedOuter<UStateTree>())
+	{
+		UStateTreeEditingSubsystem::MarkAsModified(StateTree);
+	}
+}
+
+UStateTreeState& UStateTreeState::AddChildState(const FName ChildName, const EStateTreeStateType StateType)
+{
+	UStateTreeState* ChildState = NewObject<UStateTreeState>(this, FName(), RF_Transactional);
+	check(ChildState);
+	ChildState->Name = ChildName;
+	ChildState->Parent = this;
+	ChildState->Type = StateType;
+	Children.Add(ChildState);
+	if (UStateTree* StateTree = GetTypedOuter<UStateTree>())
+	{
+		UStateTreeEditingSubsystem::MarkAsModified(StateTree);
+	}
+	return *ChildState;
+}
+
+void UStateTreeState::AddNodeImpl(FStateTreeEditorNode& EditorNode)
+{
+	EditorNode.ID = FGuid::NewGuid();
+	const FStateTreeNodeBase& Node = EditorNode.Node.GetMutable<FStateTreeNodeBase>();
+	if (const UScriptStruct* InstanceType = Cast<const UScriptStruct>(Node.GetInstanceDataType()))
+	{
+		EditorNode.Instance.InitializeAs(InstanceType);
+	}
+	if (const UScriptStruct* InstanceType = Cast<const UScriptStruct>(Node.GetExecutionRuntimeDataType()))
+	{
+		EditorNode.ExecutionRuntimeData.InitializeAs(InstanceType);
+	}
+	if (UStateTree* StateTree = GetTypedOuter<UStateTree>())
+	{
+		UStateTreeEditingSubsystem::MarkAsModified(StateTree);
+	}
+}
+
+FStateTreeTransition& UStateTreeState::AddTransition(const EStateTreeTransitionTrigger InTrigger, const EStateTreeTransitionType InType, const UStateTreeState* InState)
+{
+	FStateTreeTransition& Transition = Transitions.Emplace_GetRef(InTrigger, InType, InState);
+	AddTransitionImpl(Transition);
+	return Transition;
+}
+
+FStateTreeTransition& UStateTreeState::AddTransition(const EStateTreeTransitionTrigger InTrigger, const FGameplayTag InEventTag, const EStateTreeTransitionType InType, const UStateTreeState* InState)
+{
+	FStateTreeTransition& Transition = Transitions.Emplace_GetRef(InTrigger, InEventTag, InType, InState);
+	AddTransitionImpl(Transition);
+	return Transition;
+}
+
+void UStateTreeState::AddTransitionImpl(FStateTreeTransition& Transition)
+{
+	Transition.ID = FGuid::NewGuid();
+	if (UStateTree* StateTree = GetTypedOuter<UStateTree>())
+	{
+		UStateTreeEditingSubsystem::MarkAsModified(StateTree);
+	}
+}
diff --git a/Engine/Plugins/Runtime/StateTree/Source/StateTreeEditorModule/Private/StateTreeViewModel.cpp b/Engine/Plugins/Runtime/StateTree/Source/StateTreeEditorModule/Private/StateTreeViewModel.cpp
--- a/Engine/Plugins/Runtime/StateTree/Source/StateTreeEditorModule/Private/StateTreeViewModel.cpp
+++ b/Engine/Plugins/Runtime/StateTree/Source/StateTreeEditorModule/Private/StateTreeViewModel.cpp
@@ -16,6 +16,7 @@
 #include "StateTreeTaskBase.h"
 #include "StateTreeConditionBase.h"
 #include "StateTreeConsiderationBase.h"
+#include "StateTreeEditingSubsystem.h"
 #include "StateTreePropertyHelpers.h"
 #include "StateTreeEditorDataClipboardHelpers.h"
 #include "Styling/SlateBrush.h"
@@ -501,7 +502,7 @@ void FStateTreeViewModel::HandleIdentifierChanged(const UStateTree& StateTree) c
 {
 	if (GetStateTree() == &StateTree)
 	{
-		OnAssetChanged.Broadcast();
+		NotifyAssetChangedExternally();
 	}
 }
 
@@ -839,11 +840,21 @@ void FStateTreeViewModel::RefreshDebuggerBreakpoints()
 void FStateTreeViewModel::NotifyAssetChangedExternally() const
 {
 	OnAssetChanged.Broadcast();
+
+	if (UStateTreeEditorData* TreeData = TreeDataWeak.Get())
+	{
+		UStateTreeEditingSubsystem::MarkAsModified(TreeData->GetTypedOuter<UStateTree>());
+	}
 }
 
 void FStateTreeViewModel::NotifyStatesChangedExternally(const TSet<UStateTreeState*>& ChangedStates, const FPropertyChangedEvent& PropertyChangedEvent) const
 {
 	OnStatesChanged.Broadcast(ChangedStates, PropertyChangedEvent);
+
+	if (UStateTreeEditorData* TreeData = TreeDataWeak.Get())
+	{
+		UStateTreeEditingSubsystem::MarkAsModified(TreeData->GetTypedOuter<UStateTree>());
+	}
 }
 
 TArray<TObjectPtr<UStateTreeState>>* FStateTreeViewModel::GetSubTrees() const
@@ -947,12 +958,12 @@ UStateTreeState* FStateTreeViewModel::GetTransitionToState(TNotNull<const UState
 void FStateTreeViewModel::PostUndo(bool bSuccess)
 {
 	// TODO: see if we can narrow this down.
-	OnAssetChanged.Broadcast();
+	NotifyAssetChangedExternally();
 }
 
 void FStateTreeViewModel::PostRedo(bool bSuccess)
 {
-	OnAssetChanged.Broadcast();
+	NotifyAssetChangedExternally();
 }
 
 void FStateTreeViewModel::ClearSelection()
@@ -1026,6 +1037,7 @@ void FStateTreeViewModel::DeleteNode(TWeakObjectPtr<UStateTreeState> State, cons
 
 
 			OnStateNodesChanged.Broadcast(OwnerState);
+			UStateTreeEditingSubsystem::MarkAsModified(EditorData->GetTypedOuter<UStateTree>());
 		}
 	}
 }
@@ -1070,6 +1082,7 @@ void FStateTreeViewModel::DeleteAllNodes(TWeakObjectPtr<UStateTreeState> State,
 
 
 			OnStateNodesChanged.Broadcast(OwnerState);
+			UStateTreeEditingSubsystem::MarkAsModified(EditorData->GetTypedOuter<UStateTree>());
 		}
 	}
 }
@@ -1233,6 +1246,7 @@ void FStateTreeViewModel::PasteNode(TWeakObjectPtr<UStateTreeState> State, const
 				EPropertyChangeType::ValueSet);
 
 			OnStateNodesChanged.Broadcast(OwnerState);
+			UStateTreeEditingSubsystem::MarkAsModified(EditorData->GetTypedOuter<UStateTree>());
 		}
 	}
 }
@@ -1356,6 +1370,7 @@ void FStateTreeViewModel::MoveNode(TWeakObjectPtr<UStateTreeState> From, TWeakOb
 			
 			OnStateNodesChanged.Broadcast(FromState);
 			OnStateNodesChanged.Broadcast(ToState);
+			UStateTreeEditingSubsystem::MarkAsModified(EditorData->GetTypedOuter<UStateTree>());
 		}
 	}
 }
@@ -1438,6 +1453,7 @@ void FStateTreeViewModel::DuplicateNode(TWeakObjectPtr<UStateTreeState> State, c
 				EPropertyChangeType::Duplicate);
 
 			OnStateNodesChanged.Broadcast(OwnerState);
+			UStateTreeEditingSubsystem::MarkAsModified(EditorData->GetTypedOuter<UStateTree>());
 		}
 	}
 }
@@ -1671,6 +1687,7 @@ void FStateTreeViewModel::AddState(UStateTreeState* AfterState)
 	}
 
 	OnStateAdded.Broadcast(ParentState, NewState);
+	UStateTreeEditingSubsystem::MarkAsModified(TreeData->GetTypedOuter<UStateTree>());
 }
 
 void FStateTreeViewModel::AddChildState(UStateTreeState* ParentState)
@@ -1691,6 +1708,7 @@ void FStateTreeViewModel::AddChildState(UStateTreeState* ParentState)
 	NewState->Parent = ParentState;
 
 	OnStateAdded.Broadcast(ParentState, NewState);
+	UStateTreeEditingSubsystem::MarkAsModified(TreeData->GetTypedOuter<UStateTree>());
 }
 
 void FStateTreeViewModel::RenameState(UStateTreeState* State, FName NewName)
@@ -1709,7 +1727,7 @@ void FStateTreeViewModel::RenameState(UStateTreeState* State, FName NewName)
 
 	FProperty* NameProperty = FindFProperty<FProperty>(UStateTreeState::StaticClass(), GET_MEMBER_NAME_CHECKED(UStateTreeState, Name));
 	FPropertyChangedEvent PropertyChangedEvent(NameProperty, EPropertyChangeType::ValueSet);
-	OnStatesChanged.Broadcast(AffectedStates, PropertyChangedEvent);
+	NotifyStatesChangedExternally(AffectedStates, PropertyChangedEvent);
 }
 
 void FStateTreeViewModel::RemoveSelectedStates()
@@ -1762,6 +1780,7 @@ void FStateTreeViewModel::RemoveSelectedStates()
 		}
 
 		OnStatesRemoved.Broadcast(AffectedParents);
+		UStateTreeEditingSubsystem::MarkAsModified(TreeData->GetTypedOuter<UStateTree>());
 
 		ClearSelection();
 	}
@@ -1903,6 +1922,7 @@ void FStateTreeViewModel::PasteStatesAsChildrenFromText(const FString& TextToImp
 	{
 		OnStateAdded.Broadcast(State->Parent, State);
 	}
+	UStateTreeEditingSubsystem::MarkAsModified(TreeData->GetTypedOuter<UStateTree>());
 }
 
 void FStateTreeViewModel::DuplicateSelectedStates()
@@ -2024,7 +2044,7 @@ void FStateTreeViewModel::SetSelectedStatesEnabled(const bool bEnable)
 			State->bEnabled = bEnable;
 		}
 
-		OnAssetChanged.Broadcast();
+		NotifyAssetChangedExternally();
 	}
 }
 
@@ -2532,6 +2552,7 @@ void FStateTreeViewModel::MoveSelectedStates(UStateTreeState* TargetState, const
 		}
 
 		OnStatesMoved.Broadcast(AffectedParents, AffectedStates);
+		UStateTreeEditingSubsystem::MarkAsModified(TreeData->GetTypedOuter<UStateTree>());
 
 		TArray<TWeakObjectPtr<UStateTreeState>> WeakStates;
 		for (UStateTreeState* State : States)
@@ -2673,6 +2694,7 @@ void FStateTreeViewModel::PasteNodesToState(TNotNull<UStateTreeEditorData*> InEd
 	}
 
 	OnStateNodesChanged.Broadcast(InState);
+	UStateTreeEditingSubsystem::MarkAsModified(InEditorData->GetTypedOuter<UStateTree>());
 }
 
 bool FStateTreeViewModel::IsStateActiveInDebugger(const UStateTreeState& State) const
diff --git a/Engine/Plugins/Runtime/StateTree/Source/StateTreeEditorModule/Public/StateTreeState.h b/Engine/Plugins/Runtime/StateTree/Source/StateTreeEditorModule/Public/StateTreeState.h
--- a/Engine/Plugins/Runtime/StateTree/Source/StateTreeEditorModule/Public/StateTreeState.h
+++ b/Engine/Plugins/Runtime/StateTree/Source/StateTreeEditorModule/Public/StateTreeState.h
@@ -265,16 +265,7 @@ public:
 	UE_API FStateTreeStateLink GetLinkToState() const;
 	
 	/** Adds child state with specified name. */
-	UStateTreeState& AddChildState(const FName ChildName, const EStateTreeStateType StateType = EStateTreeStateType::State)
-	{
-		UStateTreeState* ChildState = NewObject<UStateTreeState>(this, FName(), RF_Transactional);
-		check(ChildState);
-		ChildState->Name = ChildName;
-		ChildState->Parent = this;
-		ChildState->Type = StateType;
-		Children.Add(ChildState);
-		return *ChildState;
-	}
+	UE_API UStateTreeState& AddChildState(const FName ChildName, const EStateTreeStateType StateType = EStateTreeStateType::State);
 
 	/**
 	 * Adds enter condition of specified type.
@@ -284,35 +275,21 @@ public:
 	TStateTreeEditorNode<T>& AddEnterCondition(TArgs&&... InArgs)
 	{
 		FStateTreeEditorNode& CondNode = EnterConditions.AddDefaulted_GetRef();
-		CondNode.ID = FGuid::NewGuid();
 		CondNode.Node.InitializeAs<T>(Forward<TArgs>(InArgs)...);
-		const FStateTreeNodeBase& Node = CondNode.Node.GetMutable<FStateTreeNodeBase>();
-		if (const UScriptStruct* InstanceType = Cast<const UScriptStruct>(Node.GetInstanceDataType()))
-		{
-			CondNode.Instance.InitializeAs(InstanceType);
-		}
-		if (const UScriptStruct* InstanceType = Cast<const UScriptStruct>(Node.GetExecutionRuntimeDataType()))
-		{
-			CondNode.ExecutionRuntimeData.InitializeAs(InstanceType);
-		}
+		AddNodeImpl(CondNode);
 		return static_cast<TStateTreeEditorNode<T>&>(CondNode);
 	}
 
+	/**
+	 * Adds enter utility consideration of specified type.
+	 * @return reference to the new consideration.
+	 */
 	template<typename T, typename... TArgs>
 	TStateTreeEditorNode<T>& AddConsideration(TArgs&&... InArgs)
 	{
 		FStateTreeEditorNode& ConsiderationNode = Considerations.AddDefaulted_GetRef();
-		ConsiderationNode.ID = FGuid::NewGuid();
 		ConsiderationNode.Node.InitializeAs<T>(Forward<TArgs>(InArgs)...);
-		const FStateTreeNodeBase& Node = ConsiderationNode.Node.GetMutable<FStateTreeNodeBase>();
-		if (const UScriptStruct* InstanceType = Cast<const UScriptStruct>(Node.GetInstanceDataType()))
-		{
-			ConsiderationNode.Instance.InitializeAs(InstanceType);
-		}
-		if (const UScriptStruct* InstanceType = Cast<const UScriptStruct>(Node.GetExecutionRuntimeDataType()))
-		{
-			ConsiderationNode.ExecutionRuntimeData.InitializeAs(InstanceType);
-		}
+		AddNodeImpl(ConsiderationNode);
 		return static_cast<TStateTreeEditorNode<T>&>(ConsiderationNode);
 	}
 	
@@ -324,20 +301,15 @@ public:
 	TStateTreeEditorNode<T>& AddTask(TArgs&&... InArgs)
 	{
 		FStateTreeEditorNode& TaskItem = Tasks.AddDefaulted_GetRef();
-		TaskItem.ID = FGuid::NewGuid();
 		TaskItem.Node.InitializeAs<T>(Forward<TArgs>(InArgs)...);
-		const FStateTreeNodeBase& Node = TaskItem.Node.GetMutable<FStateTreeNodeBase>();
-		if (const UScriptStruct* InstanceType = Cast<const UScriptStruct>(Node.GetInstanceDataType()))
-		{
-			TaskItem.Instance.InitializeAs(InstanceType);
-		}
-		if (const UScriptStruct* InstanceType = Cast<const UScriptStruct>(Node.GetExecutionRuntimeDataType()))
-		{
-			TaskItem.ExecutionRuntimeData.InitializeAs(InstanceType);
-		}
+		AddNodeImpl(TaskItem);
 		return static_cast<TStateTreeEditorNode<T>&>(TaskItem);
 	}
 
+private:
+	UE_API void AddNodeImpl(FStateTreeEditorNode& EditorNode);
+
+public:
 	/** Sets linked state and updates parameters to match the linked state. */
 	UE_API void SetLinkedState(FStateTreeStateLink InStateLink);
 
@@ -348,19 +320,17 @@ public:
 	 * Adds Transition.
 	 * @return reference to the new Transition.
 	 */
-	FStateTreeTransition& AddTransition(const EStateTreeTransitionTrigger InTrigger, const EStateTreeTransitionType InType, const UStateTreeState* InState = nullptr)
-	{
-		FStateTreeTransition& Transition = Transitions.Emplace_GetRef(InTrigger, InType, InState);
-		Transition.ID = FGuid::NewGuid();
-		return Transition;
-	}
+	UE_API FStateTreeTransition& AddTransition(const EStateTreeTransitionTrigger InTrigger, const EStateTreeTransitionType InType, const UStateTreeState* InState = nullptr);
+	/**
+	 * Adds Transition.
+	 * @return reference to the new Transition.
+	 */
+	UE_API FStateTreeTransition& AddTransition(const EStateTreeTransitionTrigger InTrigger, const FGameplayTag InEventTag, const EStateTreeTransitionType InType, const UStateTreeState* InState = nullptr);
 
-	FStateTreeTransition& AddTransition(const EStateTreeTransitionTrigger InTrigger, const FGameplayTag InEventTag, const EStateTreeTransitionType InType, const UStateTreeState* InState = nullptr)
-	{
-		FStateTreeTransition& Transition = Transitions.Emplace_GetRef(InTrigger, InEventTag, InType, InState);
-		Transition.ID = FGuid::NewGuid();
-		return Transition;
-	}
+private:
+	void AddTransitionImpl(FStateTreeTransition& Transition);
+
+public:
 
 	const FStateTreeTransition* GetTransitionByID(const FGuid TransitionID) const
 	{
