[State Tree][PropertyBindingUtils]
- Added a cvar to control automatically removing failed-to-resolve source binding path bindings or not for State Tree. Default true(old behavior). Missing target struct binding path bindings or failed-to-resolve invalid target binding path bindings will always be removed automatically. Missing source struct binding path bindings will never be removed automatically.
- Fix up State Tree Editor Bindings when UserDefinedStruct, Blueprint, or UObject gets reinstanced from live coding. Cache dependency to filter out state tree assets.
- Correctly gather dependent structs for compiled property bindings.
- Use struct's authoraritive version when resolving indirections
- Bindings display friendly name for UserDefinedStruct properties

#Synced-CL 51194239

diff --git a/Engine/Plugins/Runtime/PropertyBindingUtils/Source/PropertyBindingUtils/Private/PropertyBindingBindingCollection.cpp b/Engine/Plugins/Runtime/PropertyBindingUtils/Source/PropertyBindingUtils/Private/PropertyBindingBindingCollection.cpp
--- a/Engine/Plugins/Runtime/PropertyBindingUtils/Source/PropertyBindingUtils/Private/PropertyBindingBindingCollection.cpp
+++ b/Engine/Plugins/Runtime/PropertyBindingUtils/Source/PropertyBindingUtils/Private/PropertyBindingBindingCollection.cpp
@@ -11,6 +11,8 @@
 #include "PropertyBindingPath.h"
 #endif // WITH_EDITOR
 
+#include "Algo/Copy.h"
+
 #include UE_INLINE_GENERATED_CPP_BY_NAME(PropertyBindingBindingCollection)
 
 #if WITH_EDITOR
@@ -193,31 +195,22 @@ void FPropertyBindingBindingCollection::GetBindingsFor(const FGuid InStructID, T
 	});
 }
 
-void FPropertyBindingBindingCollection::RemoveInvalidBindings(const TMap<FGuid, const FPropertyBindingDataView>& InValidStructs)
+void FPropertyBindingBindingCollection::RemoveInvalidBindings(const TMap<FGuid, const FPropertyBindingDataView>& InValidStructs, EInvalidBindingCriteria InCriteria)
 {
 	const UObject* LogOwner = GetLogOwner();
 	// first pass, we remove invalid ones
-	RemoveBindingsInternal([LogOwner, InValidStructs](const FPropertyBindingBinding& Binding)
+	RemoveBindingsInternal([LogOwner, InValidStructs, InCriteria](const FPropertyBindingBinding& Binding)
 		{
-			// Remove binding if it's target struct has been removed
-			if (!InValidStructs.Contains(Binding.GetTargetPath().GetStructID()))
-			{
-				// Remove
-				return true;
-			}
-
-			// Target path should always have at least one segment (copy bind directly on a target struct/object is not allowed). 
-			if (Binding.GetTargetPath().IsPathEmpty())
-			{
-				return true;
-			}
-
 			constexpr bool bHandleRedirects = true;
 
-			// Remove binding if path containing instanced indirections (e.g. instance struct or object) cannot be resolved.
 			{
 				const FPropertyBindingDataView* SourceValue = InValidStructs.Find(Binding.GetSourcePath().GetStructID());
-				if (SourceValue && SourceValue->IsValid())
+				if (EnumHasAnyFlags(InCriteria, EInvalidBindingCriteria::RemoveMissingSourceStruct) && (!SourceValue || !SourceValue->IsValid()))
+				{
+					return true;
+				}
+
+				if (EnumHasAnyFlags(InCriteria, EInvalidBindingCriteria::RemoveFailedToResolveSourcePath) && SourceValue && SourceValue->IsValid())
 				{
 					FString Error;
 					TArray<FPropertyBindingPathIndirection> Indirections;
@@ -233,17 +226,31 @@ void FPropertyBindingBindingCollection::RemoveInvalidBindings(const TMap<FGuid,
 			}
 			
 			{
-				const FPropertyBindingDataView TargetValue = InValidStructs.FindChecked(Binding.GetTargetPath().GetStructID());
-				FString Error;
-				TArray<FPropertyBindingPathIndirection> Indirections;
-				if (!Binding.GetTargetPath().ResolveIndirectionsWithValue(TargetValue, Indirections, &Error, bHandleRedirects))
+				const FPropertyBindingDataView* TargetValue = InValidStructs.Find(Binding.GetTargetPath().GetStructID());
+				if (EnumHasAnyFlags(InCriteria, EInvalidBindingCriteria::RemoveMissingTargetStruct) &&(!TargetValue || !TargetValue->IsValid()))
 				{
-					PROPERTY_BINDINGS_LOG(LogOwner, Log, TEXT("Removing binding from '%s' because target path '%s' cannot be resolved: %s"),
-						*GetPathNameSafe(LogOwner), *Binding.GetTargetPath().ToString(), *Error);
-
-					// Remove
 					return true;
 				}
+
+				if (EnumHasAnyFlags(InCriteria, EInvalidBindingCriteria::RemoveFailedToResolveTargetPath) && TargetValue && TargetValue->IsValid())
+				{
+					// Target path should always have at least one segment (copy bind directly on a target struct/object is not allowed). 
+					if (Binding.GetTargetPath().IsPathEmpty())
+					{
+						return true;
+					}
+
+					FString Error;
+					TArray<FPropertyBindingPathIndirection> Indirections;
+					if (!Binding.GetTargetPath().ResolveIndirectionsWithValue(*TargetValue, Indirections, &Error, bHandleRedirects))
+					{
+						PROPERTY_BINDINGS_LOG(LogOwner, Log, TEXT("Removing binding from '%s' because target path '%s' cannot be resolved: %s"),
+							*GetPathNameSafe(LogOwner), *Binding.GetTargetPath().ToString(), *Error);
+
+						// Remove
+						return true;
+					}
+				}
 			}
 
 			return false;
@@ -1400,34 +1407,59 @@ bool FPropertyBindingBindingCollection::ContainsAnyStruct(const TSet<const UStru
 		}
 	}
 
-	auto PathContainsStruct = [&InStructs](const FPropertyBindingPath& PropertyPath)
+	// @todo: at runtime, we won't be able to check any Container Struct.
+	// Ex: we are able to check Container Structs for FStateTreeEditorPropertyBindings but not FStateTreePropertyBindings 
+	bool bPathContainsStruct = false;
+
+	TArray<const UStruct*> StructsArray;
+	StructsArray.Reserve(InStructs.Num());
+	Algo::Copy(InStructs, StructsArray, Algo::NoRef);
+
+#if WITH_EDITOR
+	if (const IPropertyBindingBindingCollectionOwner* BindingsOwnerPtr = GetBindingsOwner())
 	{
-		for (const FPropertyBindingPathSegment& Segment : PropertyPath.GetSegments())
-		{
-			if (InStructs.Contains(Segment.GetInstanceStruct()))
+		VisitBindings([&bPathContainsStruct, &BindingsOwnerPtr, &StructsArray](const FPropertyBindingBinding& Binding)
 			{
-				return true;
-			}
-		}
-		return false;
-	};
+				FPropertyBindingDataView SourceDataView;
+				BindingsOwnerPtr->GetBindingDataViewByID(Binding.GetSourcePath().GetStructID(), SourceDataView);
+				if (Binding.GetSourcePath().ContainsAnyStructWithValue(StructsArray, SourceDataView))
+				{
+					bPathContainsStruct = true;
+					return EVisitResult::Break;
+				}
 
-	bool bPathContainsStruct = false;
-	VisitBindings([PathContainsStruct, &bPathContainsStruct](const FPropertyBindingBinding& Binding)
-		{
-			if (PathContainsStruct(Binding.GetSourcePath()))
-			{
-				bPathContainsStruct = true;
-				return EVisitResult::Break;
-			}
+				FPropertyBindingDataView TargetDataView;
+				BindingsOwnerPtr->GetBindingDataViewByID(Binding.GetTargetPath().GetStructID(), TargetDataView);
+				if (Binding.GetTargetPath().ContainsAnyStructWithValue(StructsArray, TargetDataView))
+				{
+					bPathContainsStruct = true;
+					return EVisitResult::Break;
+				}
 
-			if (PathContainsStruct(Binding.GetTargetPath()))
+				return EVisitResult::Continue;
+			});
+	}
+	else
+#endif //WITH_EDITOR
+	{
+		VisitBindings([&bPathContainsStruct, &StructsArray](const FPropertyBindingBinding& Binding)
 			{
-				bPathContainsStruct = true;
-				return EVisitResult::Break;
-			}
-			return EVisitResult::Continue;
-		});
+				if (Binding.GetSourcePath().ContainsAnyStruct(StructsArray))
+				{
+					bPathContainsStruct = true;
+					return EVisitResult::Break;
+				}
+
+				if (Binding.GetTargetPath().ContainsAnyStruct(StructsArray))
+				{
+					bPathContainsStruct = true;
+					return EVisitResult::Break;
+				}
+
+				return EVisitResult::Continue;
+			});
+	}
+
 	return bPathContainsStruct;
 }
 
diff --git a/Engine/Plugins/Runtime/PropertyBindingUtils/Source/PropertyBindingUtils/Private/PropertyBindingPath.cpp b/Engine/Plugins/Runtime/PropertyBindingUtils/Source/PropertyBindingUtils/Private/PropertyBindingPath.cpp
--- a/Engine/Plugins/Runtime/PropertyBindingUtils/Source/PropertyBindingUtils/Private/PropertyBindingPath.cpp
+++ b/Engine/Plugins/Runtime/PropertyBindingUtils/Source/PropertyBindingUtils/Private/PropertyBindingPath.cpp
@@ -94,12 +94,18 @@ bool FPropertyBindingPath::UpdateSegments(const UStruct* BaseStruct, FString* Ou
 
 bool FPropertyBindingPath::UpdateSegmentsFromValue(const FPropertyBindingDataView BaseValueView, FString* OutError)
 {
-	TArray<FPropertyBindingPathIndirection> Indirections;
-	if (!ResolveIndirectionsWithValue(BaseValueView, Indirections, OutError, /*bHandleRedirects*/true))
+	TArray<FPropertyBindingPathIndirection, TInlineAllocator<16>> Indirections;
+	if (ResolveIndirectionsWithValue(BaseValueView, Indirections, OutError, /*bHandleRedirects*/true))
 	{
-		return false;
+		UpdateSegmentsFromIndirections(Indirections);
+		return true;
 	}
+	
+	return false;
+}
 
+void FPropertyBindingPath::UpdateSegmentsFromIndirections(const TConstArrayView<FPropertyBindingPathIndirection>& Indirections)
+{
 	for (const FPropertyBindingPathIndirection& Indirection : Indirections)
 	{
 		FPropertyBindingPathSegment& Segment = Segments[Indirection.PathSegmentIndex];
@@ -123,10 +129,56 @@ bool FPropertyBindingPath::UpdateSegmentsFromValue(const FPropertyBindingDataVie
 			Segment.SetName(Indirection.GetRedirectedName());
 		}
 		Segment.SetPropertyGuid(Indirection.GetPropertyGuid());
+		if (!Indirection.GetDisplayName().IsNone())
+		{
+			Segment.SetDisplayName(Indirection.GetDisplayName());
+		}
 #endif
 	}
+}
 
-	return true;
+bool FPropertyBindingPath::ContainsAnyStruct(TConstArrayView<const UStruct*> InStructs) const
+{
+	for (const FPropertyBindingPathSegment& Segment : Segments)
+	{
+		if (InStructs.Contains(Segment.GetInstanceStruct()))
+		{
+			return true;
+		}
+	}
+
+	return false;
+}
+
+bool FPropertyBindingPath::ContainsAnyStruct(TConstArrayView<const UStruct*> InStructs, const UStruct* BaseStruct) const
+{
+	return ContainsAnyStructWithValue(InStructs, FPropertyBindingDataView(BaseStruct, nullptr));
+}
+
+bool FPropertyBindingPath::ContainsAnyStructWithValue(TConstArrayView<const UStruct*> InStructs, const FPropertyBindingDataView BaseValueView) const
+{
+	constexpr bool bHandleRedirects = true;
+	constexpr FString* Error = nullptr;
+	// Check for any of the structs found so far even if it failed to resolve
+	constexpr bool bResetIndirectionsIfFailed = false;
+
+	TArray<FPropertyBindingPathIndirection, TInlineAllocator<16>> Indirections;
+	ResolveIndirectionsWithValue(BaseValueView, Indirections, Error, bHandleRedirects, bResetIndirectionsIfFailed);
+
+	return IndirectionsContainAnyStruct(InStructs, Indirections);
+}
+
+bool FPropertyBindingPath::IndirectionsContainAnyStruct(TConstArrayView<const UStruct*> InStructs, const TConstArrayView<FPropertyBindingPathIndirection> Indirections)
+{
+	for (const FPropertyBindingPathIndirection& Indirection : Indirections)
+	{
+		if (InStructs.Contains(Indirection.GetInstanceStruct()) || InStructs.Contains(Indirection.GetContainerStruct()))
+		{
+			return true;
+		}
+	}
+
+	return false;
 }
 
 FString FPropertyBindingPath::ToString(const int32 HighlightedSegment, const TCHAR* HighlightPrefix, const TCHAR* HighlightPostfix, const bool bOutputInstances, const int32 FirstSegment) const
@@ -151,19 +203,9 @@ FString FPropertyBindingPath::ToString(const int32 HighlightedSegment, const TCH
 		}
 
 #if WITH_EDITORONLY_DATA
-		const UStruct* ParentInstanceStruct = Segment.GetIndex() > 0 ? Segments[Segment.GetIndex() - 1].GetInstanceStruct() : nullptr;
-		if (const UUserDefinedStruct* ParentUserDefinedStruct = Cast<const UUserDefinedStruct>(ParentInstanceStruct))
+		if (!Segment->GetDisplayName().IsNone())
 		{
-			// Find friendly names for UDS properties (the property name itself has hash in it). 
-			const FString FriendlyName = FStructureEditorUtils::GetVariableFriendlyName(ParentUserDefinedStruct, Segment->GetPropertyGuid());
-			if (!FriendlyName.IsEmpty())
-			{
-				Result << FriendlyName;
-			}
-			else
-			{
-				Result << Segment->GetName();
-			}
+			Result << Segment->GetDisplayName();
 		}
 		else
 #endif // WITH_EDITORONLY_DATA
@@ -241,6 +283,16 @@ bool FPropertyBindingPath::ResolveIndirections(const UStruct* BaseStruct,
 template<typename Allocator>
 bool FPropertyBindingPath::ResolveIndirectionsWithValue(const FPropertyBindingDataView BaseValueView, TArray<FPropertyBindingPathIndirection, Allocator>& OutIndirections, FString* OutError, bool bHandleRedirects, bool bResetIndirectionsIfFailed) const
 {
+	auto GetAuthoritativeStruct = [](const UStruct* InStruct)->const UStruct*
+		{
+			if (const UClass* Class = Cast<UClass>(InStruct))
+			{
+				return Class->GetAuthoritativeClass();
+			}
+
+			return InStruct;
+		};
+
 	OutIndirections.Reset();
 	if (OutError)
 	{
@@ -279,6 +331,7 @@ bool FPropertyBindingPath::ResolveIndirectionsWithValue(const FPropertyBindingDa
 		const bool bWithValue = CurrentAddress != nullptr;
 
 #if WITH_EDITORONLY_DATA
+		FName DisplayName = Segment->GetDisplayName();
 		FName RedirectedName;
 		FGuid PropertyGuid = Segment->GetPropertyGuid();
 
@@ -324,6 +377,18 @@ bool FPropertyBindingPath::ResolveIndirectionsWithValue(const FPropertyBindingDa
 		if (!Property)
 		{
 			Property = CurrentStruct->CustomFindProperty(PropertyName);
+
+#if WITH_EDITORONLY_DATA
+			if(!Property && !DisplayName.IsNone())
+			{
+				Property = CurrentStruct->CustomFindProperty(DisplayName);
+			}
+#endif // WITH_EDITORONLY_DATA
+
+			if (Property)
+			{
+				PropertyName = Property->GetFName();
+			}
 		}
 
 		if (!Property)
@@ -349,6 +414,7 @@ bool FPropertyBindingPath::ResolveIndirectionsWithValue(const FPropertyBindingDa
 		{
 			PropertyGuid = CurrentStruct->FindPropertyGuidFromName(PropertyName);
 		}
+		DisplayName = FName(Property->GetDisplayNameText().ToString());
 #endif // WITH_EDITORONLY_DATA
 
 		const FArrayProperty* ArrayProperty = CastField<FArrayProperty>(Property);
@@ -371,6 +437,7 @@ bool FPropertyBindingPath::ResolveIndirectionsWithValue(const FPropertyBindingDa
 #if WITH_EDITORONLY_DATA
 			Indirection.RedirectedName = RedirectedName;
 			Indirection.PropertyGuid = PropertyGuid;
+			Indirection.DisplayName = DisplayName;
 #endif
 
 			ArrayIndex = INDEX_NONE;
@@ -439,6 +506,7 @@ bool FPropertyBindingPath::ResolveIndirectionsWithValue(const FPropertyBindingDa
 #if WITH_EDITORONLY_DATA
 		Indirection.RedirectedName = RedirectedName;
 		Indirection.PropertyGuid = PropertyGuid;
+		Indirection.DisplayName = DisplayName;
 #endif
 		const bool bLastSegment = Segment.GetIndex() == (Segments.Num() - 1);
 
@@ -452,7 +520,7 @@ bool FPropertyBindingPath::ResolveIndirectionsWithValue(const FPropertyBindingDa
 					// The property path is pointing into the instanced struct, it must be present.
 					// @TODO:	We could potentially check the BaseStruct metadata in editor (for similar behavior as objects)
 					//			Omitting for now to have matching functionality in editor and runtime.
-					auto SetStructInstanceIndirection = [this, bResetIndirectionsIfFailed, &CurrentAddress, &CurrentStruct, &Indirection, &OutError, &OutIndirections]
+					auto SetStructInstanceIndirection = [&GetAuthoritativeStruct, this, bResetIndirectionsIfFailed, &CurrentAddress, &CurrentStruct, &Indirection, &OutError, &OutIndirections]
 						(const UScriptStruct* ScriptStruct, const uint8* StructMemory, const int32 SegmentIndex, const UStruct* SegmentScriptStruct)
 						{
 							if (ScriptStruct == nullptr || StructMemory == nullptr)
@@ -472,7 +540,7 @@ bool FPropertyBindingPath::ResolveIndirectionsWithValue(const FPropertyBindingDa
 							}
 
 							CurrentAddress = StructMemory;
-							CurrentStruct = ScriptStruct;
+							CurrentStruct = GetAuthoritativeStruct(ScriptStruct);
 							Indirection.InstanceStruct = CurrentStruct;
 							Indirection.ArrayIndex = INDEX_NONE;
 							return true;
@@ -520,13 +588,13 @@ bool FPropertyBindingPath::ResolveIndirectionsWithValue(const FPropertyBindingDa
 				{
 					if (Segment->GetInstanceStruct())
 					{
-						CurrentStruct = Segment->GetInstanceStruct();
+						CurrentStruct = GetAuthoritativeStruct(Segment->GetInstanceStruct());
 						Indirection.InstanceStruct = CurrentStruct;
 						Indirection.AccessType = Segment->GetInstancedStructAccessType();
 					}
 					else
 					{
-						CurrentStruct = StructProperty->Struct;
+						CurrentStruct = GetAuthoritativeStruct(StructProperty->Struct);
 						Indirection.AccessType = EPropertyBindingPropertyAccessType::Offset;
 					}
 				}
@@ -541,13 +609,13 @@ bool FPropertyBindingPath::ResolveIndirectionsWithValue(const FPropertyBindingDa
 					// The property path is pointing into the object, if the object is present use it's specific type, otherwise use the type of the pointer.
 					if (Object)
 					{
-						CurrentStruct = Object->GetClass();
+						CurrentStruct = GetAuthoritativeStruct(Object->GetClass());
 						Indirection.InstanceStruct = CurrentStruct;
 						Indirection.AccessType = EPropertyBindingPropertyAccessType::ObjectInstance;
 					}
 					else
 					{
-						CurrentStruct = ObjectProperty->PropertyClass;
+						CurrentStruct = GetAuthoritativeStruct(ObjectProperty->PropertyClass);
 						Indirection.AccessType = EPropertyBindingPropertyAccessType::Object;
 					}
 				}
@@ -555,13 +623,13 @@ bool FPropertyBindingPath::ResolveIndirectionsWithValue(const FPropertyBindingDa
 				{
 					if (Segment->GetInstanceStruct())
 					{
-						CurrentStruct = Segment->GetInstanceStruct();
+						CurrentStruct = GetAuthoritativeStruct(Segment->GetInstanceStruct());
 						Indirection.InstanceStruct = CurrentStruct;
 						Indirection.AccessType = EPropertyBindingPropertyAccessType::ObjectInstance;
 					}
 					else
 					{
-						CurrentStruct = ObjectProperty->PropertyClass;
+						CurrentStruct = GetAuthoritativeStruct(ObjectProperty->PropertyClass);
 						Indirection.AccessType = EPropertyBindingPropertyAccessType::Object;
 					}
 				}
@@ -577,13 +645,13 @@ bool FPropertyBindingPath::ResolveIndirectionsWithValue(const FPropertyBindingDa
 
 					if (Object)
 					{
-						CurrentStruct = Object->GetClass();
+						CurrentStruct = GetAuthoritativeStruct(Object->GetClass());
 						Indirection.InstanceStruct = CurrentStruct;
 					}
 				}
 				else
 				{
-					CurrentStruct = WeakObjectProperty->PropertyClass;
+					CurrentStruct = GetAuthoritativeStruct(WeakObjectProperty->PropertyClass);
 				}
 
 				Indirection.AccessType = EPropertyBindingPropertyAccessType::WeakObject;
@@ -599,13 +667,13 @@ bool FPropertyBindingPath::ResolveIndirectionsWithValue(const FPropertyBindingDa
 
 					if (Object)
 					{
-						CurrentStruct = Object->GetClass();
+						CurrentStruct = GetAuthoritativeStruct(Object->GetClass());
 						Indirection.InstanceStruct = CurrentStruct;
 					}
 				}
 				else
 				{
-					CurrentStruct = SoftObjectProperty->PropertyClass;
+					CurrentStruct = GetAuthoritativeStruct(SoftObjectProperty->PropertyClass);
 				}
 
 				Indirection.AccessType = EPropertyBindingPropertyAccessType::SoftObject;
diff --git a/Engine/Plugins/Runtime/PropertyBindingUtils/Source/PropertyBindingUtils/Public/PropertyBindingBindingCollection.h b/Engine/Plugins/Runtime/PropertyBindingUtils/Source/PropertyBindingUtils/Public/PropertyBindingBindingCollection.h
--- a/Engine/Plugins/Runtime/PropertyBindingUtils/Source/PropertyBindingUtils/Public/PropertyBindingBindingCollection.h
+++ b/Engine/Plugins/Runtime/PropertyBindingUtils/Source/PropertyBindingUtils/Public/PropertyBindingBindingCollection.h
@@ -84,12 +84,28 @@ struct FPropertyBindingBindingCollection
 	 * @param InPredicate Predicate to validate the binding
 	 */
 	PROPERTYBINDINGUTILS_API void RemoveBindings(TFunctionRef<bool(FPropertyBindingBinding&)> InPredicate);
+	
+	enum class EInvalidBindingCriteria : uint8
+	{
+		// Binding is removed if Binding Source Struct is invalid
+		RemoveMissingSourceStruct = 1 << 0,
+		// Binding is removed if Binding Target Struct is invalid
+		RemoveMissingTargetStruct = 1 << 1,
+		// Binding is removed if Binding Source Path failed to resolve
+		RemoveFailedToResolveSourcePath = 1 << 2,
+		// Binding is removed if Binding Target Path failed to resolve
+		RemoveFailedToResolveTargetPath = 1 << 3,
+
+		// Missing Source Struct is not included to be consistent with the old behavior 
+		Default = RemoveMissingTargetStruct | RemoveFailedToResolveSourcePath | RemoveFailedToResolveTargetPath
+	};
 
 	/**
 	 * Removes bindings which do not point to valid structs IDs.
 	 * @param InValidStructs Set of struct IDs that are currently valid.
+	 * @param InCriteria Rules for what bindings are regarded as invalid
 	 */
-	PROPERTYBINDINGUTILS_API void RemoveInvalidBindings(const TMap<FGuid, const FPropertyBindingDataView>& InValidStructs);
+	PROPERTYBINDINGUTILS_API void RemoveInvalidBindings(const TMap<FGuid, const FPropertyBindingDataView>& InValidStructs, EInvalidBindingCriteria InCriteria = EInvalidBindingCriteria::Default);
 
 	/**
 	 * Has any binding to the target path.
@@ -393,3 +409,7 @@ struct TStructOpsTypeTraits<FPropertyBindingBindingCollection> : public TStructO
 		WithPureVirtual = true,
 	};
 };
+
+#if WITH_EDITOR
+ENUM_CLASS_FLAGS(FPropertyBindingBindingCollection::EInvalidBindingCriteria)
+#endif //WITH_EDITOR
\ No newline at end of file
diff --git a/Engine/Plugins/Runtime/PropertyBindingUtils/Source/PropertyBindingUtils/Public/PropertyBindingPath.h b/Engine/Plugins/Runtime/PropertyBindingUtils/Source/PropertyBindingUtils/Public/PropertyBindingPath.h
--- a/Engine/Plugins/Runtime/PropertyBindingUtils/Source/PropertyBindingUtils/Public/PropertyBindingPath.h
+++ b/Engine/Plugins/Runtime/PropertyBindingUtils/Source/PropertyBindingUtils/Public/PropertyBindingPath.h
@@ -16,6 +16,8 @@
 struct FPropertyBindingDataView;
 struct FPropertyBindingPath;
 
+#define UE_API PROPERTYBINDINGUTILS_API
+
 /**
  * Struct describing an indirection at specific segment at path.
  * Returned by FPropertyBindingPath::ResolveIndirections() and FPropertyBindingPath::ResolveIndirectionsWithValue().
@@ -80,6 +82,11 @@ struct FPropertyBindingPathIndirection
 	}
 
 #if WITH_EDITORONLY_DATA
+	FName GetDisplayName() const
+	{
+		return DisplayName;
+	}
+
 	FName GetRedirectedName() const
 	{
 		return RedirectedName;
@@ -117,6 +124,8 @@ private:
 	EPropertyBindingPropertyAccessType AccessType = EPropertyBindingPropertyAccessType::Offset;
 
 #if WITH_EDITORONLY_DATA
+	FName DisplayName;
+
 	/** Redirected name, if the give property name was not found but was reconciled using core redirect or property Guid. Requires ResolveIndirections/WithValue() to be called with bHandleRedirects = true. */
 	FName RedirectedName;
 
@@ -196,6 +205,16 @@ struct FPropertyBindingPathSegment
 	}
 
 #if WITH_EDITORONLY_DATA
+	FName GetDisplayName() const
+	{
+		return DisplayName;
+	}
+
+	void SetDisplayName(const FName InName)
+	{
+		DisplayName = InName;
+	}
+
 	FGuid GetPropertyGuid() const
 	{
 		return PropertyGuid;
@@ -224,6 +243,9 @@ private:
 	EPropertyBindingPropertyAccessType InstancedStructAccessType = EPropertyBindingPropertyAccessType::Unset;
 
 #if WITH_EDITORONLY_DATA
+	UPROPERTY(VisibleAnywhere, Category = "Bindings")
+	FName DisplayName;
+
 	/** Guid of the property for Blueprint classes, User Defined Structs, or Property Bags. */
 	UPROPERTY(transient, VisibleAnywhere, Category = "Bindings", AdvancedDisplay)
 	FGuid PropertyGuid;
@@ -271,7 +293,7 @@ struct FPropertyBindingPath
 	 * @param InPath Path string to parse
 	 * @return true if path was parsed successfully.
 	 */
-	PROPERTYBINDINGUTILS_API bool FromString(const FStringView InPath);
+	UE_API bool FromString(const FStringView InPath);
 
 	/**
 	 * Returns the property path as a one string. Highlight allows to decorate a specific segment.
@@ -281,7 +303,7 @@ struct FPropertyBindingPath
 	 * @param bOutputInstances if true, the instance struct types will be output. 
 	 * @param FirstSegment Index of the first path segment to be stringified.
 	 */
-	PROPERTYBINDINGUTILS_API FString ToString(const int32 HighlightedSegment = INDEX_NONE, const TCHAR* HighlightPrefix = nullptr, const TCHAR* HighlightPostfix = nullptr, const bool bOutputInstances = false, const int32 FirstSegment = 0) const;
+	UE_API FString ToString(const int32 HighlightedSegment = INDEX_NONE, const TCHAR* HighlightPrefix = nullptr, const TCHAR* HighlightPostfix = nullptr, const bool bOutputInstances = false, const int32 FirstSegment = 0) const;
 
 	/**
 	 * Resolves the property path against base struct type. The path is assumed to be relative to the BaseStruct.
@@ -292,8 +314,8 @@ struct FPropertyBindingPath
 	 * @param bResetIndirectionsIfFailed If true, indirections resolved so far will be cleared if the path failed to resolve in the end. 
 	 * @return true of path could be resolved against the base value, and instance types were updated.
 	 */
-	PROPERTYBINDINGUTILS_API bool ResolveIndirections(const UStruct* BaseStruct, TArray<FPropertyBindingPathIndirection>& OutIndirections, FString* OutError = nullptr, bool bHandleRedirects = false, bool bResetIndirectionsIfFailed = true) const;
-	PROPERTYBINDINGUTILS_API bool ResolveIndirections(const UStruct* BaseStruct, TArray<FPropertyBindingPathIndirection, TInlineAllocator<16>>& OutIndirections, FString* OutError = nullptr, bool bHandleRedirects = false, bool bResetIndirectionsIfFailed = true) const;
+	UE_API bool ResolveIndirections(const UStruct* BaseStruct, TArray<FPropertyBindingPathIndirection>& OutIndirections, FString* OutError = nullptr, bool bHandleRedirects = false, bool bResetIndirectionsIfFailed = true) const;
+	UE_API bool ResolveIndirections(const UStruct* BaseStruct, TArray<FPropertyBindingPathIndirection, TInlineAllocator<16>>& OutIndirections, FString* OutError = nullptr, bool bHandleRedirects = false, bool bResetIndirectionsIfFailed = true) const;
 
 	/**
 	 * Resolves the property path against base value. The path is assumed to be relative to the BaseValueView.
@@ -304,17 +326,17 @@ struct FPropertyBindingPath
 	 * @param bResetIndirectionsIfFailed If true, indirections resolved so far will be cleared if the path failed to resolve in the end.
 	 * @return true of path could be resolved against the base value, and instance types were updated.
 	 */
-	PROPERTYBINDINGUTILS_API bool ResolveIndirectionsWithValue(const FPropertyBindingDataView BaseValueView, TArray<FPropertyBindingPathIndirection>& OutIndirections, FString* OutError = nullptr, bool bHandleRedirects = false, bool bResetIndirectionsIfFailed = true) const;
-	PROPERTYBINDINGUTILS_API bool ResolveIndirectionsWithValue(const FPropertyBindingDataView BaseValueView, TArray<FPropertyBindingPathIndirection, TInlineAllocator<16>>& OutIndirections, FString* OutError = nullptr, bool bHandleRedirects = false, bool bResetIndirectionsIfFailed = true) const;
+	UE_API bool ResolveIndirectionsWithValue(const FPropertyBindingDataView BaseValueView, TArray<FPropertyBindingPathIndirection>& OutIndirections, FString* OutError = nullptr, bool bHandleRedirects = false, bool bResetIndirectionsIfFailed = true) const;
+	UE_API bool ResolveIndirectionsWithValue(const FPropertyBindingDataView BaseValueView, TArray<FPropertyBindingPathIndirection, TInlineAllocator<16>>& OutIndirections, FString* OutError = nullptr, bool bHandleRedirects = false, bool bResetIndirectionsIfFailed = true) const;
 
 	/**
 	 * Updates property segments from base struct type. The path is expected to be relative to the BaseStruct.
 	 * The method handles renamed properties (core redirect, Blueprint, User Defined Structs and Property Bags by ID).
-	 * @param BaseStruct Base value the path is relative to.
+	 * @param BaseStruct Base struct/class type the path is relative to
 	 * @param OutError Optional, pointer to string where error will be logged if update fails.
 	 * @return true of path could be resolved against the base value, and instance types were updated.
 	 */
-	PROPERTYBINDINGUTILS_API bool UpdateSegments(const UStruct* BaseStruct, FString* OutError = nullptr);
+	UE_API bool UpdateSegments(const UStruct* BaseStruct, FString* OutError = nullptr);
 
 	/**
 	 * Updates property segments from base value. The path is expected to be relative to the base value.
@@ -324,7 +346,46 @@ struct FPropertyBindingPath
 	 * @param OutError Optional, pointer to string where error will be logged if update fails.
 	 * @return true of path could be resolved against the base value, and instance types were updated.
 	 */
-	PROPERTYBINDINGUTILS_API bool UpdateSegmentsFromValue(const FPropertyBindingDataView BaseValueView, FString* OutError = nullptr);
+	UE_API bool UpdateSegmentsFromValue(const FPropertyBindingDataView BaseValueView, FString* OutError = nullptr);
+
+	/**
+	 * Updates property segments from resolved indirections. The path is expected to be relative to the base value.
+	 * The method updates instance types, and handles renamed properties (core redirect, Blueprint, User Defined Structs and Property Bags by ID).
+	 * By storing the instance types on the path, we can resolve the path without the base value later.
+	 * @param Indirections completely resolved indirections
+	 */
+	UE_API void UpdateSegmentsFromIndirections(const TConstArrayView<FPropertyBindingPathIndirection>& Indirections);
+
+	/**
+	 * Check if the binding path contains any structs from cached instance struct type.
+	 * @param InStructs any structs the binding path may contain
+	 * @return true if path contains any struct
+	 */
+	UE_API bool ContainsAnyStruct(TConstArrayView<const UStruct*> InStructs) const;
+
+	/**
+	 * Check if the binding path contains any structs from base struct type.
+	 * @param InStructs any structs the binding path may contain
+	 * @param BaseStruct Base struct/class type the path is relative to
+	 * @return true if path contains any struct
+	 */
+	UE_API bool ContainsAnyStruct(TConstArrayView<const UStruct*> InStructs, const UStruct* BaseStruct) const;
+
+	/**
+	 * Check if the binding path contains any structs from base value.
+	 * @param InStructs any structs the binding path may contain
+	 * @param BaseValueView Base value the path is relative to
+	 * @return true if path contains any struct
+	 */
+	UE_API bool ContainsAnyStructWithValue(TConstArrayView<const UStruct*> InStructs, const FPropertyBindingDataView BaseValueView) const;
+
+	/**
+	 * Check if indirections contain any structs. 
+	 * @param InStructs any structs the binding path may contain
+	 * @param Indirections Indirections resolved so far to check against
+	 * @return true if indirections contains any struct
+	 */
+	UE_API static bool IndirectionsContainAnyStruct(TConstArrayView<const UStruct*> InStructs, const TConstArrayView<FPropertyBindingPathIndirection> Indirections);
 
 	/** @return true if the path is empty. In that case the path points to the struct. */
 	bool IsPathEmpty() const
@@ -398,13 +459,13 @@ struct FPropertyBindingPath
 	}
 
 	/** Test if paths are equal. */
-	PROPERTYBINDINGUTILS_API bool operator==(const FPropertyBindingPath& RHS) const;
+	UE_API bool operator==(const FPropertyBindingPath& RHS) const;
 
 	/**
 	 * Test if this paths includes the provided path.
 	 * A path includes another one when they are == but this path can be longer.
 	 */
-	PROPERTYBINDINGUTILS_API bool Includes(const FPropertyBindingPath& Other) const;
+	UE_API bool Includes(const FPropertyBindingPath& Other) const;
 
 private:
 #if WITH_EDITORONLY_DATA
@@ -603,3 +664,5 @@ struct FPropertyBindingCopyInfoBatch
 	UPROPERTY()
 	FPropertyBindingIndex16 PropertyFunctionsEnd;
 };
+
+#undef UE_API
diff --git a/Engine/Plugins/Runtime/StateTree/Source/StateTreeEditorModule/Private/Customizations/StateTreeEditorNodeDetails.cpp b/Engine/Plugins/Runtime/StateTree/Source/StateTreeEditorModule/Private/Customizations/StateTreeEditorNodeDetails.cpp
--- a/Engine/Plugins/Runtime/StateTree/Source/StateTreeEditorModule/Private/Customizations/StateTreeEditorNodeDetails.cpp
+++ b/Engine/Plugins/Runtime/StateTree/Source/StateTreeEditorModule/Private/Customizations/StateTreeEditorNodeDetails.cpp
@@ -2271,7 +2271,7 @@ void FStateTreeEditorNodeDetails::OnDeleteNode() const
 
 			ParentArrayProperty->DeleteItem(Index);
 
-			UE::StateTreeEditor::RemoveInvalidBindings(EditorDataPtr);
+			EditorDataPtr->RemoveInvalidBindings();
 		}
 	}
 }
@@ -2290,7 +2290,7 @@ void FStateTreeEditorNodeDetails::OnDeleteAllNodes() const
 
 				ArrayHandle->EmptyArray();
 
-				UE::StateTreeEditor::RemoveInvalidBindings(EditorDataPtr);
+				EditorDataPtr->RemoveInvalidBindings();
 			}
 		}
 	}
diff --git a/Engine/Plugins/Runtime/StateTree/Source/StateTreeEditorModule/Private/Customizations/StateTreeTransitionDetails.cpp b/Engine/Plugins/Runtime/StateTree/Source/StateTreeEditorModule/Private/Customizations/StateTreeTransitionDetails.cpp
--- a/Engine/Plugins/Runtime/StateTree/Source/StateTreeEditorModule/Private/Customizations/StateTreeTransitionDetails.cpp
+++ b/Engine/Plugins/Runtime/StateTree/Source/StateTreeEditorModule/Private/Customizations/StateTreeTransitionDetails.cpp
@@ -537,7 +537,7 @@ void FStateTreeTransitionDetails::OnDeleteTransition() const
 
 				ArrayHandle->DeleteItem(Index);
 
-				UE::StateTreeEditor::RemoveInvalidBindings(EditorData);
+				EditorData->RemoveInvalidBindings();
 			}
 		}
 	}
@@ -556,7 +556,7 @@ void FStateTreeTransitionDetails::OnDeleteAllTransitions() const
 
 				ArrayHandle->EmptyArray();
 
-				UE::StateTreeEditor::RemoveInvalidBindings(EditorData);
+				EditorData->RemoveInvalidBindings();
 			}
 		}
 	}
diff --git a/Engine/Plugins/Runtime/StateTree/Source/StateTreeEditorModule/Private/StateTreeCompilerManager.cpp b/Engine/Plugins/Runtime/StateTree/Source/StateTreeEditorModule/Private/StateTreeCompilerManager.cpp
--- a/Engine/Plugins/Runtime/StateTree/Source/StateTreeEditorModule/Private/StateTreeCompilerManager.cpp
+++ b/Engine/Plugins/Runtime/StateTree/Source/StateTreeEditorModule/Private/StateTreeCompilerManager.cpp
@@ -11,6 +11,7 @@
 #include "StateTreeEditorPropertyBindings.h"
 #include "StateTreeSchema.h"
 
+#include "Algo/Copy.h"
 #include "Containers/Ticker.h"
 #include "Editor.h"
 #include "StructUtilsDelegates.h"
@@ -77,14 +78,21 @@ struct FStateTreeDependencies
 };
 ENUM_CLASS_FLAGS(FStateTreeDependencies::EDependencyType);
 
+bool IsStructCoreOrStateTreeInternal(TNotNull<const UStruct*> InStruct)
+{
+	const TNotNull<const UPackage*> StateTreeModulePackage = UStateTree::StaticClass()->GetOutermost();
+	const TNotNull<const UPackage*> StateTreeEditorModulePackage = UStateTreeEditorData::StaticClass()->GetOutermost();
+	const TNotNull<const UPackage*> CoreUObjectModulePackage = UObject::StaticClass()->GetOutermost();
+
+	return InStruct->IsInPackage(StateTreeModulePackage)
+		|| InStruct->IsInPackage(StateTreeEditorModulePackage)
+		|| InStruct->IsInPackage(CoreUObjectModulePackage);
+}
 /** Find the references that are needed by the asset. */
 class FArchiveReferencingPropertiesBase : public FArchiveUObject
 {
 public:
 	FArchiveReferencingPropertiesBase()
-		: StateTreeModulePackage(UStateTree::StaticClass()->GetOutermost())
-		, StateTreeEditorModulePackage(UStateTreeEditorData::StaticClass()->GetOutermost())
-		, CoreUObjectModulePackage(UObject::StaticClass()->GetOutermost())
 	{
 		ArIsObjectReferenceCollector = true;
 		ArIgnoreOuterRef = true;
@@ -95,18 +103,11 @@ public:
 		SetShouldSkipCompilingAssets(false);
 	}
 
-	bool IsSupportedObject(TNotNull<const UStruct*> Struct)
+	static bool IsSupportedObject(TNotNull<const UStruct*> Struct)
 	{
 		/** As an optimization, do not include basic structures like FVector and state tree internal types. */
-		return !Struct->IsInPackage(StateTreeModulePackage)
-			&& !Struct->IsInPackage(StateTreeEditorModulePackage)
-			&& !Struct->IsInPackage(CoreUObjectModulePackage);
+		return !IsStructCoreOrStateTreeInternal(Struct);
 	}
-
-private:
-	TNotNull<const UPackage*> StateTreeModulePackage;
-	TNotNull<const UPackage*> StateTreeEditorModulePackage;
-	TNotNull<const UPackage*> CoreUObjectModulePackage;
 };
 
 /** Find the references that are needed by the asset. */
@@ -218,9 +219,18 @@ public:
 	bool CompileInternalSynchronously(TNotNull<UStateTree*> InStateTree, FStateTreeCompilerLog& InOutLog);
 	TOptional<bool> CompileIfNeededSynchronously(TNotNull<UStateTree*> StateTree);
 
+	void CacheEditorBindingExternalDependencies(TNotNull<UStateTreeEditorData*> InEditorData);
+
 private:
 	void AddToQueue_AnyThread(TNotNull<UStateTree*> InStateTree);
-	void UpdateBindingsInstanceStructsIfNeeded(TSet<const UStruct*>& Structs, TNotNull<UStateTreeEditorData*> EditorData);
+
+	/**
+	 * Fix up editor bindings when external Blueprint or UserDefinedStruct that an EditorData is dependent on got changed and reinstanced
+	 * For editor bindings fixup, we don't need to recompile. Editor bindings are in sync with compiled bindings(otherwise the compilation would have already been dirty). 
+	 * When a UDS changes or an UObject reinstanced, we perform fix up on editor bindings, and relink the compiled data, during which it will also fix up the compiled bindings. So compiled bindings are still in sync with editor bindings.
+	 * @param InChangedStructs changed User Defined Structs or reinstanced UObjects
+	 */
+	void UpdateEditorBindingsIfNeeded(TSet<const UStruct*>& InChangedStructs);
 
 	void GatherDependencies(TNotNull<UStateTree*> StateTree);
 	TSet<UStateTree*> GatherStateTreeDependencies_AnyThread(TNotNull<UStateTree*> StateTree) const;
@@ -234,7 +244,6 @@ private:
 	void ConditionalResetLastCompilationHash(TNotNull<UStateTree*> StateTree) const;
 	bool AllowQueuedCompilation(TNotNull<UStateTree*> StateTree) const;
 
-	void HandleStateTreeLoaded(TNotNull<UStateTree*> StateTree);
 	void HandleStateTreeCooking(TNotNull<UStateTree*> StateTree);
 	using FReplacementObjectMap = TMap<UObject*, UObject*>;
 	void HandleObjectsReinstanced(const FReplacementObjectMap& ObjectMap);
@@ -242,12 +251,30 @@ private:
 	void HandleEndPIE(bool bIsSimulating);
 	void HandleUserDefinedStructReinstanced(const UUserDefinedStruct& UserDefinedStruct);
 
+	void PurgeEditorBindingExternalDependencies();
+
 private:
 	FDelegateHandle ObjectsReinstancedHandle;
 	FDelegateHandle UserDefinedStructReinstancedHandle;
 	FDelegateHandle PreBeginPIEHandle;
 	FDelegateHandle EndPIEHandle;
 
+	struct FEditorBindingDependencies
+	{
+		using FTreeToDependencyMap = TMap<TObjectKey<const UStateTree>, TSet<TObjectKey<const UStruct>>>;
+		using FDependencyToTreeMap = TMap<TObjectKey<const UStruct>, TSet<TObjectKey<const UStateTree>>>;
+		FTreeToDependencyMap TreeToDependencyMap;
+		FDependencyToTreeMap DependencyToTreeMap;
+
+		static bool IsExternalDependency(TNotNull<const UStruct*> InStruct)
+		{
+			return IsStructCoreOrStateTreeInternal(InStruct) == false;
+		}
+	};
+
+	/** Dependencies that are state tree external, and not in core */
+	FEditorBindingDependencies EditorBindingExternalDependencies;
+
 	/** Referencee to dependencies. Generated with GatherDependencies and used when a dependencies is "recreated" (blueprint compilation). */
 	TMap<TObjectKey<UStateTree>, TSharedPtr<FStateTreeDependencies>> StateTreeToDependencies;
 	TMap<FObjectKey, TArray<TObjectKey<UStateTree>>> DependenciesToStateTree;
@@ -346,6 +373,47 @@ TOptional<bool> FCompilerManagerImpl::CompileIfNeededSynchronously(TNotNull<USta
 	return {};
 }
 
+void FCompilerManagerImpl::CacheEditorBindingExternalDependencies(TNotNull<UStateTreeEditorData*> InEditorData)
+{
+	// For now, state tree asset only gets loaded on game thread and we don't use StrongObjPtr here.
+	check(IsInGameThread());
+
+	PurgeEditorBindingExternalDependencies();
+
+	const UStateTree* StateTree = InEditorData->GetTypedOuter<UStateTree>();
+	check(StateTree);
+
+	TSet<TObjectKey<const UStruct>>& CurrentExternalDependencies = EditorBindingExternalDependencies.TreeToDependencyMap.FindOrAdd(StateTree);
+
+	// Clean up registered dependencies
+	for (TObjectKey<const UStruct> DependentStruct : CurrentExternalDependencies)
+	{
+		TSet<TObjectKey<const UStateTree>>* Trees = EditorBindingExternalDependencies.DependencyToTreeMap.Find(DependentStruct);
+		check(Trees);
+
+		const int32 RemovedCnt = Trees->Remove(StateTree);
+		check(RemovedCnt == 1);
+		if (Trees->IsEmpty())
+		{
+			EditorBindingExternalDependencies.DependencyToTreeMap.Remove(DependentStruct);
+		}
+	}
+
+	CurrentExternalDependencies.Reset();
+	
+	const FStateTreeEditorPropertyBindings* EditorBindings = InEditorData->GetPropertyEditorBindings();
+	check(EditorBindings);
+
+	const TArray<TNotNull<const UStruct*>> NewExternalDependencies = EditorBindings->GatherDependenciesByPredicate(FEditorBindingDependencies::IsExternalDependency);
+	Algo::Copy(NewExternalDependencies, CurrentExternalDependencies, Algo::NoRef);
+
+	for (TObjectKey<const UStruct> DependentStruct : CurrentExternalDependencies)
+	{
+		TSet<TObjectKey<const UStateTree>>& Trees = EditorBindingExternalDependencies.DependencyToTreeMap.FindOrAdd(DependentStruct);
+		Trees.Emplace(StateTree);
+	}
+}
+
 void FCompilerManagerImpl::QueueForCompilation(TNotNull<UStateTree*> StateTree)
 {
 	check(IsInGameThread() || IsInAsyncLoadingThread());
@@ -565,27 +633,74 @@ void FCompilerManagerImpl::CompileStateTree_AnyThread(TNotNull<UStateTree*> Stat
 	}
 }
 
-void FCompilerManagerImpl::UpdateBindingsInstanceStructsIfNeeded(TSet<const UStruct*>& Structs, TNotNull<UStateTreeEditorData*> EditorData)
+void FCompilerManagerImpl::UpdateEditorBindingsIfNeeded(TSet<const UStruct*>& InChangedStructs)
 {
-	bool bShouldUpdate = false;
-	EditorData->VisitAllNodes([&Structs, &bShouldUpdate](const UStateTreeState* State, const FStateTreeBindableStructDesc& Desc, const FStateTreeDataView Value)
+	for (const UStruct* ChangedStruct : InChangedStructs)
+	{
+		if (const TSet<TObjectKey<const UStateTree>>* Trees = EditorBindingExternalDependencies.DependencyToTreeMap.Find(ChangedStruct))
 		{
-			if (Structs.Contains(Value.GetStruct()))
+			for (TObjectKey<const UStateTree> Tree : *Trees)
 			{
-				bShouldUpdate = true;
-				return EStateTreeVisitor::Break;
-			}
-			return EStateTreeVisitor::Continue;
-		});
+				if (const UStateTree* TreePtr = Tree.ResolveObjectPtr())
+				{
+					UStateTreeEditorData* EditorDataPtr = Cast<UStateTreeEditorData>(TreePtr->EditorData);
+					check(EditorDataPtr);
 
-	if (!bShouldUpdate)
-	{
-		bShouldUpdate = EditorData->GetEditorPropertyBindings()->ContainsAnyStruct(Structs);
-	}
+					FStateTreeEditorPropertyBindings* EditorBindings = EditorDataPtr->GetPropertyEditorBindings();
+					check(EditorBindings);
 
-	if (bShouldUpdate)
-	{
-		EditorData->UpdateBindingsInstanceStructs();
+					TMap<FGuid, const FStateTreeDataView> AllValues;
+					EditorDataPtr->GetAllStructValues(AllValues);
+
+					bool bAnyBindingChanged = false;
+
+					TArray<FPropertyBindingPathIndirection, TInlineAllocator<16>> Indirections;
+					auto UpdateBindingPathIfNeeded = [&AllValues, &bAnyBindingChanged, &Indirections, ChangedStruct](FPropertyBindingPath& InBindingPath)
+						{
+							if (const FStateTreeDataView* StructValue = AllValues.Find(InBindingPath.GetStructID()))
+							{
+								constexpr FString* Error = nullptr;
+								constexpr bool bHandleRedirects = true;
+
+								// Will reset the last Indirections
+								if (InBindingPath.ResolveIndirectionsWithValue(*StructValue, Indirections, Error, bHandleRedirects))
+								{
+									if (FPropertyBindingPath::IndirectionsContainAnyStruct({ ChangedStruct }, Indirections))
+									{
+										const FPropertyBindingPath BindingPathBeforeUpdate = InBindingPath;
+										InBindingPath.UpdateSegmentsFromIndirections(Indirections);
+
+										bAnyBindingChanged |= (InBindingPath != BindingPathBeforeUpdate);
+									}
+								}
+							}
+						};
+
+					EditorBindings->ForEachMutableBinding([&UpdateBindingPathIfNeeded](FPropertyBindingBinding& Binding)
+						{
+							UpdateBindingPathIfNeeded(Binding.GetMutableSourcePath());
+						});
+
+					const UPackage* Package = EditorDataPtr->GetPackage();
+					check(Package);
+
+					// During object re-instancing, blueprint intentionally un-dirtied packages that were not dirty before the process but got dirtied during the process. This is to prevent dirtying assets on load.
+					// During load, bindings either will fail to resolve indirections, or are already up to date. So no package dirtying will happen during load.
+					// Outside loading phase, we need to defer dirtying the package to next frame to work around the blueprint mechanism.
+					// Object re-instancing is not transactable. We call this function when a UDS change is undone/redone. It is expected that a dirty package remains dirty after a UDS change is undone.
+					if (!Package->IsDirty() && bAnyBindingChanged)
+					{
+						ExecuteOnGameThread(UE_SOURCE_LOCATION, [WeakEditorData = TWeakObjectPtr<UStateTreeEditorData>(EditorDataPtr)]()
+						{
+							if (UStateTreeEditorData* EditorDataPtr = WeakEditorData.Get())
+							{
+								EditorDataPtr->MarkPackageDirty();
+							}
+						});
+					}
+				}
+			}
+		}
 	}
 }
 
@@ -654,6 +769,8 @@ void FCompilerManagerImpl::HandleObjectsReinstanced(const FReplacementObjectMap&
 		}
 	}
 
+	UpdateEditorBindingsIfNeeded(StructsToBeReplaced);
+
 	if (bUseDependenciesToTriggerCompilation)
 	{
 		TArray<TNotNull<UStateTree*>, TInlineAllocator<8>> StateTreeToLink;
@@ -709,12 +826,6 @@ void FCompilerManagerImpl::HandleObjectsReinstanced(const FReplacementObjectMap&
 	}
 	else
 	{
-		for (TObjectIterator<UStateTreeEditorData> It; It; ++It)
-		{
-			UStateTreeEditorData* StateTreeEditorData = *It;
-			UpdateBindingsInstanceStructsIfNeeded(StructsToBeReplaced, StateTreeEditorData);
-		}
-
 		for (TObjectIterator<UStateTree> It; It; ++It)
 		{
 			UStateTree* StateTree = *It;
@@ -776,7 +887,12 @@ void FCompilerManagerImpl::HandleObjectsReinstanced(const FReplacementObjectMap&
 
 void FCompilerManagerImpl::HandleUserDefinedStructReinstanced(const UUserDefinedStruct& UserDefinedStruct)
 {
-	if (UE::StateTree::Compiler::Private::bUseDependenciesToTriggerCompilation)
+	TSet<const UStruct*> Structs;
+	Structs.Add(&UserDefinedStruct);
+
+	UpdateEditorBindingsIfNeeded(Structs);
+
+	if (bUseDependenciesToTriggerCompilation)
 	{
 		TSet<TNotNull<UStateTree*>> StateTreeToLink;
 		const FObjectKey StructToReplacedKey = &UserDefinedStruct;
@@ -803,15 +919,6 @@ void FCompilerManagerImpl::HandleUserDefinedStructReinstanced(const UUserDefined
 	}
 	else
 	{
-		TSet<const UStruct*> Structs;
-		Structs.Add(&UserDefinedStruct);
-
-		for (TObjectIterator<UStateTreeEditorData> It; It; ++It)
-		{
-			UStateTreeEditorData* StateTreeEditorData = *It;
-			UpdateBindingsInstanceStructsIfNeeded(Structs, StateTreeEditorData);
-		}
-
 		for (TObjectIterator<UStateTree> It; It; ++It)
 		{
 			UStateTree* StateTree = *It;
@@ -827,6 +934,34 @@ void FCompilerManagerImpl::HandleUserDefinedStructReinstanced(const UUserDefined
 	}
 }
 
+void FCompilerManagerImpl::PurgeEditorBindingExternalDependencies()
+{
+	auto PurgeImpl = []<typename MapType>(MapType& InMap)
+		requires std::is_same_v<MapType, FEditorBindingDependencies::FDependencyToTreeMap> || std::is_same_v<MapType, FEditorBindingDependencies::FTreeToDependencyMap>
+	{
+		for (auto MapIt = InMap.CreateIterator(); MapIt; ++MapIt)
+		{
+			if (!MapIt->Key.ResolveObjectPtr())
+			{
+				MapIt.RemoveCurrent();
+			}
+			else
+			{
+				for (auto SetIt = MapIt->Value.CreateIterator(); SetIt; ++SetIt)
+				{
+					if (!SetIt->ResolveObjectPtr())
+					{
+						SetIt.RemoveCurrent();
+					}
+				}
+			}
+		}
+	};
+
+	PurgeImpl(EditorBindingExternalDependencies.TreeToDependencyMap);
+	PurgeImpl(EditorBindingExternalDependencies.DependencyToTreeMap);
+}
+
 void FCompilerManagerImpl::GatherDependencies(TNotNull<UStateTree*> StateTree)
 {
 	// Find the tree in the StateTreeToDependencies
@@ -900,7 +1035,7 @@ void FCompilerManagerImpl::GatherDependencies(TNotNull<UStateTree*> StateTree)
 		{
 			TMap<FGuid, const FPropertyBindingDataView> AllStructValues;
 			EditorData->GetAllStructValues(AllStructValues);
-			auto AddBindingPathDependencies = [&AddDependency , &AllStructValues](const FPropertyBindingPath& PropertyPath)
+			auto AddBindingPathDependencies = [&AddDependency, &AllStructValues](const FPropertyBindingPath& PropertyPath)
 				{
 					const FPropertyBindingDataView* FoundStruct = AllStructValues.Find(PropertyPath.GetStructID());
 					if (FoundStruct)
@@ -915,7 +1050,8 @@ void FCompilerManagerImpl::GatherDependencies(TNotNull<UStateTree*> StateTree)
 								{
 									AddDependency(Indirection.GetInstanceStruct(), FStateTreeDependencies::EDependencyType::DT_Link);
 								}
-								else if (Indirection.GetContainerStruct())
+
+								if (Indirection.GetContainerStruct())
 								{
 									AddDependency(Indirection.GetContainerStruct(), FStateTreeDependencies::EDependencyType::DT_Link);
 								}
@@ -1072,5 +1208,13 @@ TOptional<bool> FCompilerManager::CompileIfNeededSynchronously(TNotNull<UStateTr
 	return {};
 }
 
+void FCompilerManager::CacheEditorBindingExternalDependencies(TNotNull<UStateTreeEditorData*> EditorData)
+{
+	TRACE_CPUPROFILER_EVENT_SCOPE(FCompilerManager::CacheEditorBindingExternalDependencies)
+	if (ensureMsgf(Private::CompilerManagerImpl.IsValid(), TEXT("Can't compile the asset when the module is not available.")))
+	{
+		return Private::CompilerManagerImpl->CacheEditorBindingExternalDependencies(EditorData);
+	}
+}
 } // UE::StateTree::Compiler
 
diff --git a/Engine/Plugins/Runtime/StateTree/Source/StateTreeEditorModule/Private/StateTreeEditingSubsystem.cpp b/Engine/Plugins/Runtime/StateTree/Source/StateTreeEditorModule/Private/StateTreeEditingSubsystem.cpp
--- a/Engine/Plugins/Runtime/StateTree/Source/StateTreeEditorModule/Private/StateTreeEditingSubsystem.cpp
+++ b/Engine/Plugins/Runtime/StateTree/Source/StateTreeEditorModule/Private/StateTreeEditingSubsystem.cpp
@@ -263,9 +263,7 @@ void UStateTreeEditingSubsystem::ValidateStateTree(TNotNull<UStateTree*> InState
 
 	auto RemoveUnusedBindings = [](TNotNull<UStateTreeEditorData*> EditorData)
 		{
-			TMap<FGuid, const FPropertyBindingDataView> AllStructValues;
-			EditorData->GetAllStructValues(AllStructValues);
-			EditorData->GetPropertyEditorBindings()->RemoveInvalidBindings(AllStructValues);
+			EditorData->RemoveInvalidBindings();
 		};
 
 	auto UpdateLinkedStateParameters = [](TNotNull<UStateTreeEditorData*> TreeData)
@@ -342,7 +340,7 @@ void UStateTreeEditingSubsystem::ValidateStateTree(TNotNull<UStateTree*> InState
 		EditorData->FixDuplicateIDs();
 
 		RemoveUnusedBindings(EditorData);
-		EditorData->UpdateBindingsInstanceStructs();
+		EditorData->UpdateBindings();
 		ValidateLinkedStates(EditorData);
 		UpdateLinkedStateParameters(EditorData);
 		UpdateTransactionalFlags(EditorData);
diff --git a/Engine/Plugins/Runtime/StateTree/Source/StateTreeEditorModule/Private/StateTreeEditorData.cpp b/Engine/Plugins/Runtime/StateTree/Source/StateTreeEditorModule/Private/StateTreeEditorData.cpp
--- a/Engine/Plugins/Runtime/StateTree/Source/StateTreeEditorModule/Private/StateTreeEditorData.cpp
+++ b/Engine/Plugins/Runtime/StateTree/Source/StateTreeEditorModule/Private/StateTreeEditorData.cpp
@@ -21,10 +21,11 @@
 #include "UObject/UE5SpecialProjectStreamObjectVersion.h"
 
 #if WITH_EDITOR
-#include "StructUtilsDelegates.h"
 #include "StructUtils/UserDefinedStruct.h"
 #endif
 
+#include "StateTreeDelegatesInternal.h"
+
 #include UE_INLINE_GENERATED_CPP_BY_NAME(StateTreeEditorData)
 
 #define LOCTEXT_NAMESPACE "StateTreeEditor"
@@ -39,6 +40,17 @@ namespace UE::StateTree::Editor
 		TEXT("A parent state can enter a child state during state selection (before the transition completes).")
 	);
 
+	FAutoConsoleVariable CVarAutoRemoveSourcePathInvalidPropertyBindings(
+		TEXT("StateTree.Editor.AutoRemoveSourcePathInvalidEditorBindings"),
+		true,
+		TEXT("If set to true, EditorBindings whose source path does not resolve anymore,")
+		TEXT("will be automatically removed when the asset is loaded, when the users authors in the editor,")
+		TEXT("when state tree gets compiled.\n")
+		TEXT("If set to false, the compilation will fail and the user will be notified and be responsible for which binding is invalid and fixing them.\n")
+		TEXT("Note that missing target struct or failed to resolve target binding path bindings will always be removed automatically.")
+		TEXT("Note that missing source struct bindings will never be removed automatically.")
+	);
+
 	const FString GlobalStateName(TEXT("Global"));
 	const FString PropertyFunctionStateName(TEXT("Property Functions"));
 	const FName ParametersNodeName(TEXT("Parameters"));
@@ -132,7 +144,7 @@ void UStateTreeEditorData::OnParametersChanged(const UStateTree& StateTree)
 	{
 		if (OwnerStateTree == &StateTree)
 		{
-			UpdateBindingsInstanceStructs();
+			UpdateBindings();
 		}
 	}
 }
@@ -143,7 +155,10 @@ void UStateTreeEditorData::OnStateParametersChanged(const UStateTree& StateTree,
 	{
 		if (OwnerStateTree == &StateTree)
 		{
-			UpdateBindingsInstanceStructs();
+			if (UStateTreeState* State = GetMutableStateByID(StateID))
+			{
+				UpdateBindings();
+			}
 		}
 	}
 }
@@ -188,7 +203,10 @@ void UStateTreeEditorData::PostLoad()
 	ReparentStates();
 	FixObjectNodes();
 	FixDuplicateIDs();
-	UpdateBindingsInstanceStructs();
+
+	RemoveInvalidBindings();
+
+	UpdateBindings();
 }
 
 void UStateTreeEditorData::PostEditChangeChainProperty(FPropertyChangedChainEvent& PropertyChangedEvent)
@@ -277,10 +295,7 @@ void UStateTreeEditorData::PostEditChangeChainProperty(FPropertyChangedChainEven
 			if (MemberProperty->GetFName() == GET_MEMBER_NAME_CHECKED(UStateTreeEditorData, Evaluators)
 				|| MemberProperty->GetFName() == GET_MEMBER_NAME_CHECKED(UStateTreeEditorData, GlobalTasks))
 			{
-				TMap<FGuid, const FPropertyBindingDataView> AllStructValues;
-				GetAllStructValues(AllStructValues);
-				Modify();
-				EditorBindings.RemoveInvalidBindings(AllStructValues);
+				RemoveInvalidBindings();
 			}
 		}
 
@@ -1050,11 +1065,6 @@ void UStateTreeEditorData::FixDuplicateIDs()
 		
 		return EStateTreeVisitor::Continue;
 	});
-
-	// It is possible that the user has changed the node type so some of the bindings might not make sense anymore, clean them up.
-	TMap<FGuid, const FPropertyBindingDataView> AllValues;
-	GetAllStructValues(AllValues);
-	EditorBindings.RemoveInvalidBindings(AllValues);
 }
 
 void UStateTreeEditorData::DuplicateIDs()
@@ -1149,22 +1159,72 @@ void UStateTreeEditorData::DuplicateIDs()
 	}
 }
 
-void UStateTreeEditorData::UpdateBindingsInstanceStructs()
+void UStateTreeEditorData::UpdateBindings()
 {
 	TMap<FGuid, const FStateTreeDataView> AllValues;
 	GetAllStructValues(AllValues);
+
+	TArray<TNotNull<const UStruct*>, TInlineAllocator<32>> BindingsDependentStructs;
+
+	TArray<FPropertyBindingPathIndirection, TInlineAllocator<16>> Indirections;
+
+	auto UpdateBindingPathAndAddDependencies = [&AllValues, &Indirections, &BindingsDependentStructs](FPropertyBindingPath& InBindingPath)
+		{
+			constexpr bool bHandleRedirects = true;
+			constexpr FString* Error = nullptr;
+
+			// For failed to resolve bindings, we need to log dependency on structs found so far
+			constexpr bool bResetIndirectionsIfFailed = false;
+
+			if (AllValues.Contains(InBindingPath.GetStructID()))
+			{
+				// Will reset Indirection 
+				if (InBindingPath.ResolveIndirectionsWithValue(AllValues[InBindingPath.GetStructID()], Indirections, Error, bHandleRedirects, bResetIndirectionsIfFailed))
+				{
+					InBindingPath.UpdateSegmentsFromIndirections(Indirections);
+				}
+
+				for (const FPropertyBindingPathIndirection& Indirection : Indirections)
+				{
+					if (const UStruct* InstanceStruct = Indirection.GetInstanceStruct())
+					{
+						BindingsDependentStructs.Emplace(InstanceStruct);
+					}
+
+					if (const UStruct* ContainerStruct = Indirection.GetContainerStruct())
+					{
+						BindingsDependentStructs.Emplace(ContainerStruct);
+					}
+				}
+			}
+		};
+
 	for (FPropertyBindingBinding& Binding : EditorBindings.GetMutableBindings())
 	{
-		if (AllValues.Contains(Binding.GetSourcePath().GetStructID()))
-		{
-			Binding.GetMutableSourcePath().UpdateSegmentsFromValue(AllValues[Binding.GetSourcePath().GetStructID()]);
-		}
+		UpdateBindingPathAndAddDependencies(Binding.GetMutableSourcePath());
+		UpdateBindingPathAndAddDependencies(Binding.GetMutableTargetPath());
+	}
 
-		if (AllValues.Contains(Binding.GetTargetPath().GetStructID()))
-		{
-			Binding.GetMutableTargetPath().UpdateSegmentsFromValue(AllValues[Binding.GetTargetPath().GetStructID()]);
-		}
+	EditorBindings.CacheDependency(BindingsDependentStructs);
+
+	UE::StateTree::Delegates::Private::OnStateTreeEditorBindingUpdated.ExecuteIfBound(this);
+}
+
+void UStateTreeEditorData::RemoveInvalidBindings()
+{
+	TMap<FGuid, const FPropertyBindingDataView> AllStructValues;
+	GetAllStructValues(AllStructValues);
+
+	FPropertyBindingBindingCollection::EInvalidBindingCriteria Criteria = 
+		FPropertyBindingBindingCollection::EInvalidBindingCriteria::RemoveMissingTargetStruct
+		| FPropertyBindingBindingCollection::EInvalidBindingCriteria::RemoveFailedToResolveTargetPath;
+
+	if (UE::StateTree::Editor::CVarAutoRemoveSourcePathInvalidPropertyBindings->GetBool())
+	{
+		Criteria |= FPropertyBindingBindingCollection::EInvalidBindingCriteria::RemoveFailedToResolveSourcePath;
 	}
+
+	EditorBindings.RemoveInvalidBindings(AllStructValues, Criteria);
 }
 
 void UStateTreeEditorData::CallPostLoadOnNodes()
diff --git a/Engine/Plugins/Runtime/StateTree/Source/StateTreeEditorModule/Private/StateTreeEditorDataClipboardHelpers.cpp b/Engine/Plugins/Runtime/StateTree/Source/StateTreeEditorModule/Private/StateTreeEditorDataClipboardHelpers.cpp
--- a/Engine/Plugins/Runtime/StateTree/Source/StateTreeEditorModule/Private/StateTreeEditorDataClipboardHelpers.cpp
+++ b/Engine/Plugins/Runtime/StateTree/Source/StateTreeEditorModule/Private/StateTreeEditorDataClipboardHelpers.cpp
@@ -68,16 +68,6 @@ bool ImportTextAsClipboardEditorData(const UScriptStruct* InTargetType, TNotNull
 	return !bProcessBuffer || OutClipboardEditorData.IsValid();
 }
 
-void RemoveInvalidBindings(TNotNull<UStateTreeEditorData*> InEditorData)
-{
-	if (FStateTreeEditorPropertyBindings* Bindings = InEditorData->GetPropertyEditorBindings())
-	{
-		TMap<FGuid, const FPropertyBindingDataView> AllStructValues;
-		InEditorData->GetAllStructValues(AllStructValues);
-		Bindings->RemoveInvalidBindings(AllStructValues);
-	}
-}
-
 void AddErrorNotification(const FText& InText, float InExpiredDuration)
 {
 	FNotificationInfo NotificationInfo(FText::GetEmpty());
diff --git a/Engine/Plugins/Runtime/StateTree/Source/StateTreeEditorModule/Private/StateTreeEditorDataClipboardHelpers.h b/Engine/Plugins/Runtime/StateTree/Source/StateTreeEditorModule/Private/StateTreeEditorDataClipboardHelpers.h
--- a/Engine/Plugins/Runtime/StateTree/Source/StateTreeEditorModule/Private/StateTreeEditorDataClipboardHelpers.h
+++ b/Engine/Plugins/Runtime/StateTree/Source/StateTreeEditorModule/Private/StateTreeEditorDataClipboardHelpers.h
@@ -23,8 +23,6 @@ void ExportTextAsClipboardEditorData(const FClipboardEditorData& InClipboardEdit
 bool ImportTextAsClipboardEditorData(const UScriptStruct* InTargetType, TNotNull<UStateTreeEditorData*> InTargetTree, TNotNull<UObject*> InOwner,
                                      FClipboardEditorData& OutClipboardEditorData, bool bProcessBuffer = true);
 
-void RemoveInvalidBindings(TNotNull<UStateTreeEditorData*> InEditorData);
-
 void AddErrorNotification(const FText& InText, float InExpiredDuration = 5.f);
 
 USTRUCT()
diff --git a/Engine/Plugins/Runtime/StateTree/Source/StateTreeEditorModule/Private/StateTreeEditorModule.cpp b/Engine/Plugins/Runtime/StateTree/Source/StateTreeEditorModule/Private/StateTreeEditorModule.cpp
--- a/Engine/Plugins/Runtime/StateTree/Source/StateTreeEditorModule/Private/StateTreeEditorModule.cpp
+++ b/Engine/Plugins/Runtime/StateTree/Source/StateTreeEditorModule/Private/StateTreeEditorModule.cpp
@@ -187,6 +187,7 @@ void FStateTreeEditorModule::StartupModule()
 	UE::StateTree::Delegates::Private::OnCompileIfChanged.BindRaw(this, &FStateTreeEditorModule::HandleCompileIfChanged);
 	UE::StateTree::Delegates::Private::OnRequestAssetRegistryTags.BindRaw(this, &FStateTreeEditorModule::HandleRequestAssetRegistryTags);
 	UE::StateTree::Delegates::Private::OnStateTreeAssetLoaded.BindRaw(this, &FStateTreeEditorModule::HandleStateTreeAssetLoaded);
+	UE::StateTree::Delegates::Private::OnStateTreeEditorBindingUpdated.BindRaw(this, &FStateTreeEditorModule::HandleStateTreeEditorBindingUpdated);
 	UE::StateTree::Delegates::Private::OnAppendToClassSchema.BindRaw(this, &FStateTreeEditorModule::HandleAppendToStateTreeClassSchema);
 
 	// class cache isn't needed in non-editor runs
@@ -286,6 +287,7 @@ void FStateTreeEditorModule::ShutdownModule()
 	UE::StateTree::Delegates::Private::OnRequestAssetRegistryTags.Unbind();
 	UE::StateTree::Delegates::Private::OnCompileIfChanged.Unbind();
 	UE::StateTree::Delegates::OnRequestEditorHash.Unbind();
+	UE::StateTree::Delegates::Private::OnStateTreeEditorBindingUpdated.Unbind();
 }
 
 TSharedRef<IStateTreeEditor> FStateTreeEditorModule::CreateStateTreeEditor(const EToolkitMode::Type Mode, const TSharedPtr< class IToolkitHost >& InitToolkitHost, UStateTree* StateTree)
@@ -360,6 +362,11 @@ void FStateTreeEditorModule::HandleStateTreeAssetLoaded(TNotNull<UStateTree*> In
 	}
 }
 
+void FStateTreeEditorModule::HandleStateTreeEditorBindingUpdated(TNotNull<UStateTreeEditorData*> InStateTreeEditorData)
+{
+	UE::StateTree::Compiler::FCompilerManager::CacheEditorBindingExternalDependencies(InStateTreeEditorData);
+}
+
 void FStateTreeEditorModule::HandleAppendToStateTreeClassSchema(FAppendToClassSchemaContext& Context)
 {
 	FStateTreeCompiler::AppendToStateTreeClassSchema(Context);
diff --git a/Engine/Plugins/Runtime/StateTree/Source/StateTreeEditorModule/Private/StateTreeEditorPropertyBindings.cpp b/Engine/Plugins/Runtime/StateTree/Source/StateTreeEditorModule/Private/StateTreeEditorPropertyBindings.cpp
--- a/Engine/Plugins/Runtime/StateTree/Source/StateTreeEditorModule/Private/StateTreeEditorPropertyBindings.cpp
+++ b/Engine/Plugins/Runtime/StateTree/Source/StateTreeEditorModule/Private/StateTreeEditorPropertyBindings.cpp
@@ -1,6 +1,8 @@
 // Copyright Epic Games, Inc. All Rights Reserved.
 
 #include "StateTreeEditorPropertyBindings.h"
+
+#include "Algo/Copy.h"
 #include "StateTreePropertyBindingCompiler.h"
 #include "StateTreeNodeBase.h"
 #include "StateTreePropertyFunctionBase.h"
@@ -60,6 +62,43 @@ const FStateTreePropertyPathBinding* FStateTreeEditorPropertyBindings::AddOutput
 	return nullptr;
 }
 
+bool FStateTreeEditorPropertyBindings::IsDependentOn(TConstArrayView<UStruct*> InStructsDependedOn) const
+{
+	for (const UStruct* Struct : InStructsDependedOn)
+	{
+		if (CachedDependencies.Contains(Struct))
+		{
+			return true;
+		}
+	}
+
+	return false;
+}
+
+TArray<TNotNull<const UStruct*>> FStateTreeEditorPropertyBindings::GatherDependencies() const
+{
+	TArray<TNotNull<const UStruct*>> DependentStructs;
+	DependentStructs.Reserve(CachedDependencies.Num());
+
+	for (const TObjectKey<const UStruct> StructObjectKey : CachedDependencies)
+	{
+		if (const UStruct* Struct = StructObjectKey.ResolveObjectPtr())
+		{
+			DependentStructs.Emplace(Struct);
+		}
+	}
+
+	return DependentStructs;
+}
+
+void FStateTreeEditorPropertyBindings::CacheDependency(const TConstArrayView<TNotNull<const UStruct*>> InDependentStructs)
+{
+	CachedDependencies.Reset();
+	CachedDependencies.Reserve(InDependentStructs.Num());
+
+	Algo::Copy(InDependentStructs, CachedDependencies, Algo::NoRef);
+}
+
 FPropertyBindingBinding* FStateTreeEditorPropertyBindings::AddBindingInternal(const FPropertyBindingPath& InSourcePath
 	, const FPropertyBindingPath& InTargetPath)
 {
diff --git a/Engine/Plugins/Runtime/StateTree/Source/StateTreeEditorModule/Private/StateTreeViewModel.cpp b/Engine/Plugins/Runtime/StateTree/Source/StateTreeEditorModule/Private/StateTreeViewModel.cpp
--- a/Engine/Plugins/Runtime/StateTree/Source/StateTreeEditorModule/Private/StateTreeViewModel.cpp
+++ b/Engine/Plugins/Runtime/StateTree/Source/StateTreeEditorModule/Private/StateTreeViewModel.cpp
@@ -1023,7 +1023,7 @@ void FStateTreeViewModel::DeleteNode(TWeakObjectPtr<UStateTreeState> State, cons
 						StateNodeHandle.GetEditorNodeArray().RemoveAt(StateNodeHandle.GetNodeIndex());
 					}
 
-					UE::StateTreeEditor::RemoveInvalidBindings(InEditorData);
+					InEditorData->RemoveInvalidBindings();
 				};
 
 			UE::StateTree::PropertyHelpers::ModifyStateInPreAndPostEdit(
@@ -1068,7 +1068,7 @@ void FStateTreeViewModel::DeleteAllNodes(TWeakObjectPtr<UStateTreeState> State,
 						StateNodeHandle.GetEditorNodeArray().Empty();
 					}
 
-					UE::StateTreeEditor::RemoveInvalidBindings(InEditorData);
+					InEditorData->RemoveInvalidBindings();
 				};
 
 			UE::StateTree::PropertyHelpers::ModifyStateInPreAndPostEdit(
diff --git a/Engine/Plugins/Runtime/StateTree/Source/StateTreeEditorModule/Public/StateTreeCompilerManager.h b/Engine/Plugins/Runtime/StateTree/Source/StateTreeEditorModule/Public/StateTreeCompilerManager.h
--- a/Engine/Plugins/Runtime/StateTree/Source/StateTreeEditorModule/Public/StateTreeCompilerManager.h
+++ b/Engine/Plugins/Runtime/StateTree/Source/StateTreeEditorModule/Public/StateTreeCompilerManager.h
@@ -51,6 +51,13 @@ public:
 	 */
 	static UE_API TOptional<bool> CompileIfNeededSynchronously(TNotNull<UStateTree*> StateTree);
 	
+	/*
+	 * Cache dependent external objects for tree's editor bindings. 
+	 * This includes any objects outside StateTree and CoreModule, that editor bindings are dependent on.
+	 * This function gets called when EditorData updates its bindings.
+	 */
+	static UE_API void CacheEditorBindingExternalDependencies(TNotNull<UStateTreeEditorData*> EditorData);
+
 private:
 	FCompilerManager() = delete;
 	FCompilerManager(const FCompilerManager&) = delete;
diff --git a/Engine/Plugins/Runtime/StateTree/Source/StateTreeEditorModule/Public/StateTreeEditorData.h b/Engine/Plugins/Runtime/StateTree/Source/StateTreeEditorModule/Public/StateTreeEditorData.h
--- a/Engine/Plugins/Runtime/StateTree/Source/StateTreeEditorModule/Public/StateTreeEditorData.h
+++ b/Engine/Plugins/Runtime/StateTree/Source/StateTreeEditorModule/Public/StateTreeEditorData.h
@@ -239,9 +239,19 @@ public:
 	UE_API void ReparentStates();
 	/** Fix any nodes or bindings with duplicated ID. */
 	UE_API void FixDuplicateIDs();
-	/** Update the binding segments. */
-	UE_API void UpdateBindingsInstanceStructs();
-	
+
+	UE_DEPRECATED(5.8, "Use UpdateEditorBindings instead")
+	void UpdateBindingsInstanceStructs()
+	{
+		UpdateBindings();
+	}
+
+	/** Update the editor binding segments and cache editor bindings dependency. */
+	UE_API void UpdateBindings();
+
+	/** Remove invalid editor property bindings. */
+	UE_API void RemoveInvalidBindings();
+
 	// StateTree Builder API
 
 	/**
diff --git a/Engine/Plugins/Runtime/StateTree/Source/StateTreeEditorModule/Public/StateTreeEditorModule.h b/Engine/Plugins/Runtime/StateTree/Source/StateTreeEditorModule/Public/StateTreeEditorModule.h
--- a/Engine/Plugins/Runtime/StateTree/Source/StateTreeEditorModule/Public/StateTreeEditorModule.h
+++ b/Engine/Plugins/Runtime/StateTree/Source/StateTreeEditorModule/Public/StateTreeEditorModule.h
@@ -136,6 +136,7 @@ private:
 	void HandleCompileIfChanged(TNotNull<UStateTree*> InStateTree);
 	void HandleRequestAssetRegistryTags(TNotNull<const UStateTree*> InStateTree, FAssetRegistryTagsContext& Context);
 	void HandleStateTreeAssetLoaded(TNotNull<UStateTree*> InStateTree);
+	void HandleStateTreeEditorBindingUpdated(TNotNull<UStateTreeEditorData*> InStateTreeEditorData);
 	void HandleAppendToStateTreeClassSchema(FAppendToClassSchemaContext& Context);
 
 protected:
diff --git a/Engine/Plugins/Runtime/StateTree/Source/StateTreeEditorModule/Public/StateTreeEditorPropertyBindings.h b/Engine/Plugins/Runtime/StateTree/Source/StateTreeEditorModule/Public/StateTreeEditorPropertyBindings.h
--- a/Engine/Plugins/Runtime/StateTree/Source/StateTreeEditorModule/Public/StateTreeEditorPropertyBindings.h
+++ b/Engine/Plugins/Runtime/StateTree/Source/StateTreeEditorModule/Public/StateTreeEditorPropertyBindings.h
@@ -3,6 +3,7 @@
 #pragma once
 
 #include "UObject/Interface.h"
+#include "UObject/ObjectKey.h"
 #include "PropertyBindingExtension.h"
 #include "PropertyBindingBindingCollectionOwner.h"
 #include "StateTreePropertyBindings.h"
@@ -72,6 +73,43 @@ struct FStateTreeEditorPropertyBindings : public FPropertyBindingBindingCollecti
 	 */
 	UE_API const FStateTreePropertyPathBinding* AddOutputBinding(const FPropertyBindingPath& InSourcePath, const FPropertyBindingPath& InTargetPath);
 
+         /**
+	 * Check if any editor binding is dependent on any of the structs
+	 * @param InStructsDependedOn Structs to check if any editor binding is dependent on
+	 * @return true if any editor binding is dependent on any of the passed in structs
+	 */
+	UE_API bool IsDependentOn(TConstArrayView<UStruct*> InStructsDependedOn) const;
+
+	/**
+	 * Collect all the structs that this editor bindings are dependent on
+	 */
+	UE_API TArray<TNotNull<const UStruct*>> GatherDependencies() const;
+
+	template<typename Predicate>
+	TArray<TNotNull<const UStruct*>> GatherDependenciesByPredicate(Predicate Pred) const
+	{
+		TArray<TNotNull<const UStruct*>> DependentStructs;
+		DependentStructs.Reserve(CachedDependencies.Num());
+
+		for (const TObjectKey<const UStruct> StructObjectKey : CachedDependencies)
+		{
+			if (const UStruct* Struct = StructObjectKey.ResolveObjectPtr())
+			{
+				if (::Invoke(Pred, Struct))
+				{
+					DependentStructs.Emplace(Struct);
+				}
+			}
+		}
+
+		return DependentStructs;
+	}
+
+	/**
+	 * Set Dependency for Editor Bindings
+	 * @param InDependentStructs Structs editor bindings are dependent on
+	 */
+	UE_API void CacheDependency(const TConstArrayView<TNotNull<const UStruct*>> InDependentStructs);
 protected:
 	//~ Begin FPropertyBindingBindingCollection overrides
 	virtual FPropertyBindingBinding* AddBindingInternal(const FPropertyBindingPath& InSourcePath, const FPropertyBindingPath& InTargetPath) override;
@@ -84,6 +122,8 @@ protected:
 private:
 	UPROPERTY()
 	TArray<FStateTreePropertyPathBinding> PropertyBindings;
+
+	TSet<TObjectKey<const UStruct>> CachedDependencies;
 };
 
 
diff --git a/Engine/Plugins/Runtime/StateTree/Source/StateTreeModule/Internal/StateTreeDelegatesInternal.h b/Engine/Plugins/Runtime/StateTree/Source/StateTreeModule/Internal/StateTreeDelegatesInternal.h
--- a/Engine/Plugins/Runtime/StateTree/Source/StateTreeModule/Internal/StateTreeDelegatesInternal.h
+++ b/Engine/Plugins/Runtime/StateTree/Source/StateTreeModule/Internal/StateTreeDelegatesInternal.h
@@ -6,6 +6,7 @@
 
 #define UE_API STATETREEMODULE_API
 
+class UStateTreeEditorData;
 class UStateTree;
 #if WITH_EDITOR
 struct FAppendToClassSchemaContext;
@@ -18,6 +19,9 @@ namespace UE::StateTree::Delegates::Private
 	DECLARE_DELEGATE_OneParam(FOnStateTreeAssetLoaded, TNotNull<UStateTree*> /*StateTree*/);
 	extern UE_API FOnStateTreeAssetLoaded OnStateTreeAssetLoaded;
 
+	DECLARE_DELEGATE_OneParam(FOnStateTreeEditorBindingUpdated, TNotNull<UStateTreeEditorData*> /*StateTree*/);
+	extern UE_API FOnStateTreeEditorBindingUpdated OnStateTreeEditorBindingUpdated;
+
 #if WITH_EDITOR
 	/** A callback for UStateTree::CompileIfChanged. */
 	DECLARE_DELEGATE_OneParam(FOnCompileIfChanged, TNotNull<UStateTree*> /*StateTree*/);
diff --git a/Engine/Plugins/Runtime/StateTree/Source/StateTreeModule/Private/StateTreeDelegates.cpp b/Engine/Plugins/Runtime/StateTree/Source/StateTreeModule/Private/StateTreeDelegates.cpp
--- a/Engine/Plugins/Runtime/StateTree/Source/StateTreeModule/Private/StateTreeDelegates.cpp
+++ b/Engine/Plugins/Runtime/StateTree/Source/StateTreeModule/Private/StateTreeDelegates.cpp
@@ -35,6 +35,7 @@ PRAGMA_ENABLE_DEPRECATION_WARNINGS
 
 namespace UE::StateTree::Delegates::Private
 {
+	FOnStateTreeEditorBindingUpdated OnStateTreeEditorBindingUpdated;
 	FOnStateTreeAssetLoaded OnStateTreeAssetLoaded;
 #if WITH_EDITOR
 	FOnCompileIfChanged OnCompileIfChanged;
