[State Tree] Prevent Linked Asset override at runtime when the Linked State Parameter is binding source or binding target. This includes the parameter is bound to something else, a node is bound to the parameter or a node's appended property function is bound to the parameter.

- Use an additional flag for override. Cache the flag during Compile for runtime check.

#Synced-CL 51198676

diff --git a/Engine/Plugins/Runtime/StateTree/Source/StateTreeEditorModule/Private/Customizations/StateTreeStateDetails.cpp b/Engine/Plugins/Runtime/StateTree/Source/StateTreeEditorModule/Private/Customizations/StateTreeStateDetails.cpp
--- a/Engine/Plugins/Runtime/StateTree/Source/StateTreeEditorModule/Private/Customizations/StateTreeStateDetails.cpp
+++ b/Engine/Plugins/Runtime/StateTree/Source/StateTreeEditorModule/Private/Customizations/StateTreeStateDetails.cpp
@@ -261,7 +261,7 @@ void FStateTreeStateDetails::CustomizeDetails(IDetailLayoutBuilder& DetailBuilde
 
 		const FButtonStyle* const CloseButtonStyle = &FAppStyle::Get().GetWidgetStyle<FDockTabStyle>("Docking.Tab").CloseButtonStyle;
 		StateCategory.AddCustomRow(LOCTEXT("WarningBindingWithLinkedAsset", "WarningBindingWithLinkedAsset"))
-			.Visibility(MakeAttributeSP(this, &FStateTreeStateDetails::HandleLinkedAssetVisibility))
+			.Visibility(MakeAttributeSP(this, &FStateTreeStateDetails::HandleDisableLinkedAssetOverrideVisibility))
 			.WholeRowContent()
 			[
 				SNew(SWarningOrErrorBox)
@@ -484,7 +484,7 @@ void FStateTreeStateDetails::CustomizeDetails(IDetailLayoutBuilder& DetailBuilde
 	}
 }
 
-EVisibility FStateTreeStateDetails::HandleLinkedAssetVisibility() const
+EVisibility FStateTreeStateDetails::HandleDisableLinkedAssetOverrideVisibility() const
 {
 	const UStateTreeEditorData* EditorData = WeakEditorData.Get();
 	const UStateTreeState* State = WeakState.Get();
@@ -494,11 +494,19 @@ EVisibility FStateTreeStateDetails::HandleLinkedAssetVisibility() const
 		&& State->Type == EStateTreeStateType::LinkedAsset
 		&& State->Tag.IsValid())
 	{
-		TArray<const FPropertyBindingBinding*> Bindings;
-		EditorData->GetEditorPropertyBindings()->GetBindingsFor(State->Parameters.ID, Bindings);
-		const bool bHasBinding = Bindings.Num() > 0;
-		return bHasBinding ? EVisibility::Visible : EVisibility::Collapsed;
+		const EStateTreeVisitor VisitResult = EditorData->VisitStateBindings(State, [](TNotNull<const UStateTreeState*> State, const FStateTreePropertyPathBinding& Binding, const FStateTreeBindableStructDesc& BindingTargetNodeDesc, const FStateTreeDataView BindingTargetNodeValue)
+			{
+				if (Binding.GetTargetPath().GetStructID() == State->Parameters.ID || Binding.GetSourcePath().GetStructID() == State->Parameters.ID)
+				{
+					return EStateTreeVisitor::Break;
+				}
+
+				return EStateTreeVisitor::Continue;
+			});
+
+		return VisitResult == EStateTreeVisitor::Break ? EVisibility::Visible : EVisibility::Collapsed;
 	}
+
 	return EVisibility::Collapsed;
 }
 
diff --git a/Engine/Plugins/Runtime/StateTree/Source/StateTreeEditorModule/Private/Customizations/StateTreeStateDetails.h b/Engine/Plugins/Runtime/StateTree/Source/StateTreeEditorModule/Private/Customizations/StateTreeStateDetails.h
--- a/Engine/Plugins/Runtime/StateTree/Source/StateTreeEditorModule/Private/Customizations/StateTreeStateDetails.h
+++ b/Engine/Plugins/Runtime/StateTree/Source/StateTreeEditorModule/Private/Customizations/StateTreeStateDetails.h
@@ -36,7 +36,7 @@ protected:
 	virtual void PostRedo(bool bSuccess) override;
 
 private:
-	EVisibility HandleLinkedAssetVisibility() const;
+	EVisibility HandleDisableLinkedAssetOverrideVisibility() const;
 
 protected:
 	TWeakPtr<IPropertyUtilities> WeakPropertyUtilities;
diff --git a/Engine/Plugins/Runtime/StateTree/Source/StateTreeEditorModule/Private/StateTreeCompiler.cpp b/Engine/Plugins/Runtime/StateTree/Source/StateTreeEditorModule/Private/StateTreeCompiler.cpp
--- a/Engine/Plugins/Runtime/StateTree/Source/StateTreeEditorModule/Private/StateTreeCompiler.cpp
+++ b/Engine/Plugins/Runtime/StateTree/Source/StateTreeEditorModule/Private/StateTreeCompiler.cpp
@@ -903,6 +903,24 @@ bool FStateTreeCompiler::CreateStateRecursive(UStateTreeState& State, const FSta
 		return false;
 	}
 
+	if (State.Type == EStateTreeStateType::LinkedAsset && State.Tag.IsValid())
+	{
+		const EStateTreeVisitor VisitResult = EditorData->VisitStateBindings(&State, [](TNotNull<const UStateTreeState*> State, const FStateTreePropertyPathBinding& Binding, const FStateTreeBindableStructDesc& BindingTargetNodeDesc, const FStateTreeDataView BindingTargetNodeValue)
+			{
+				if (Binding.GetSourcePath().GetStructID() == State->Parameters.ID || Binding.GetTargetPath().GetStructID() == State->Parameters.ID)
+				{
+					return EStateTreeVisitor::Break;
+				}
+
+				return EStateTreeVisitor::Continue;
+			});
+
+		if (VisitResult == EStateTreeVisitor::Break)
+		{
+			CompactState.bCanOverrideLinkedAssetAtRuntime = false;
+		}
+	}
+
 	SourceStates.Add(&State);
 	IDToState.Add(State.ID, StateIdx);
 
diff --git a/Engine/Plugins/Runtime/StateTree/Source/StateTreeEditorModule/Private/StateTreeEditorData.cpp b/Engine/Plugins/Runtime/StateTree/Source/StateTreeEditorModule/Private/StateTreeEditorData.cpp
--- a/Engine/Plugins/Runtime/StateTree/Source/StateTreeEditorModule/Private/StateTreeEditorData.cpp
+++ b/Engine/Plugins/Runtime/StateTree/Source/StateTreeEditorModule/Private/StateTreeEditorData.cpp
@@ -1593,6 +1593,34 @@ EStateTreeVisitor UStateTreeEditorData::VisitStructBoundPropertyFunctions(FGuid
 	return EStateTreeVisitor::Continue;
 }
 
+EStateTreeVisitor UStateTreeEditorData::VisitStateBindings(TNotNull<const UStateTreeState*> State,
+	TFunctionRef<EStateTreeVisitor(TNotNull<const UStateTreeState*> State, const FStateTreePropertyPathBinding& Binding, const
+	FStateTreeBindableStructDesc& BindingTargetNodeDesc, const FStateTreeDataView BindingTargetNodeValue)> InFunc) const
+{
+	return VisitStateNodes(*State, [&InFunc, Self = this](const UStateTreeState* State, const FStateTreeBindableStructDesc& Desc, const FStateTreeDataView Value)
+		{
+			check(State);
+			EStateTreeVisitor VisitResult = EStateTreeVisitor::Continue;
+
+			Self->EditorBindings.VisitBindings([&InFunc, &Desc, &Value, State, &VisitResult](const FPropertyBindingBinding& Binding)
+				{
+					const FStateTreePropertyPathBinding& StateTreeBinding = static_cast<const FStateTreePropertyPathBinding&>(Binding);
+					if (Desc.ID == StateTreeBinding.GetTargetPath().GetStructID())
+					{
+						if (InFunc(State, static_cast<const FStateTreePropertyPathBinding&>(Binding), Desc, Value) == EStateTreeVisitor::Break)
+						{
+							VisitResult = EStateTreeVisitor::Break;
+							return FPropertyBindingBindingCollection::EVisitResult::Break;
+						}
+					}
+
+					return FPropertyBindingBindingCollection::EVisitResult::Continue;
+				});
+
+			return VisitResult;
+		});
+}
+
 EStateTreeVisitor UStateTreeEditorData::VisitHierarchy(TFunctionRef<EStateTreeVisitor(UStateTreeState& State, UStateTreeState* ParentState)> InFunc) const
 {
 	using FStatePair = TTuple<UStateTreeState*, UStateTreeState*>; 
diff --git a/Engine/Plugins/Runtime/StateTree/Source/StateTreeEditorModule/Public/StateTreeEditorData.h b/Engine/Plugins/Runtime/StateTree/Source/StateTreeEditorModule/Public/StateTreeEditorData.h
--- a/Engine/Plugins/Runtime/StateTree/Source/StateTreeEditorModule/Public/StateTreeEditorData.h
+++ b/Engine/Plugins/Runtime/StateTree/Source/StateTreeEditorModule/Public/StateTreeEditorData.h
@@ -186,6 +186,12 @@ public:
 	 */
 	UE_API EStateTreeVisitor VisitStructBoundPropertyFunctions(FGuid StructID, const FString& StatePath, TFunctionRef<EStateTreeVisitor(const FStateTreeEditorNode& EditorNode, const FStateTreeBindableStructDesc& Desc, const FStateTreeDataView Value)> InFunc) const;
 
+	/**
+	 * Iterates over all bindings whose target path is in a given state.
+	 * @param InFunc function called at each binding, should return Continue if visiting is continued or Break to stop.
+	 */
+	UE_API EStateTreeVisitor VisitStateBindings(TNotNull<const UStateTreeState*> State, TFunctionRef<EStateTreeVisitor(TNotNull<const UStateTreeState*> State, const FStateTreePropertyPathBinding& Binding, const FStateTreeBindableStructDesc& BindingTargetNodeDesc, const FStateTreeDataView BindingTargetNodeValue)> InFunc) const;
+
 	/**
 	 * Returns array of nodes along the execution path, up to the TargetStruct.
 	 * @param Path The states to visit during the check
@@ -337,6 +343,15 @@ public:
 		EditorBindings.AddFunctionBinding(PropertyFunctionNodeStruct, SourcePathSegments, TargetPath);
 	}
 
+	/**
+	 * Removes property binding of target path.
+	 */
+	void RemovePropertyBinding(const FPropertyBindingPath& TargetPath)
+	{
+		constexpr FPropertyBindingBindingCollection::ESearchMode SearchMode = FPropertyBindingBindingCollection::ESearchMode::Includes;
+		EditorBindings.RemoveBindings(TargetPath, SearchMode);
+	}
+
 	/**
 	 * Adds property binding between two structs.
 	 */
diff --git a/Engine/Plugins/Runtime/StateTree/Source/StateTreeModule/Private/StateTreeExecutionContext.cpp b/Engine/Plugins/Runtime/StateTree/Source/StateTreeModule/Private/StateTreeExecutionContext.cpp
--- a/Engine/Plugins/Runtime/StateTree/Source/StateTreeModule/Private/StateTreeExecutionContext.cpp
+++ b/Engine/Plugins/Runtime/StateTree/Source/StateTreeModule/Private/StateTreeExecutionContext.cpp
@@ -7027,8 +7027,7 @@ bool FStateTreeExecutionContext::SelectStateInternal(
 		}
 
 		// LinkedAsset that has bindings cannot be overridden by another asset. The bindings will write invalid data.
-		const bool bInvalidOverride = NextState.ParameterBindingsBatch.IsValid();
-		if (bCreateNewState && !bInvalidOverride)
+		if (bCreateNewState && NextState.bCanOverrideLinkedAssetAtRuntime)
 		{
 			if (const FStateTreeReference* Override = GetLinkedStateTreeOverrideForTag(NextState.Tag))
 			{
diff --git a/Engine/Plugins/Runtime/StateTree/Source/StateTreeModule/Public/StateTreeTypes.h b/Engine/Plugins/Runtime/StateTree/Source/StateTreeModule/Public/StateTreeTypes.h
--- a/Engine/Plugins/Runtime/StateTree/Source/StateTreeModule/Public/StateTreeTypes.h
+++ b/Engine/Plugins/Runtime/StateTree/Source/StateTreeModule/Public/StateTreeTypes.h
@@ -959,6 +959,10 @@ struct FCompactStateTreeState
 	UPROPERTY()
 	EStateTreeStateSelectionBehavior SelectionBehavior = EStateTreeStateSelectionBehavior::TrySelectChildrenInOrder;
 
+	/** Whether the linked state supports Linked Asset override. */
+	UPROPERTY()
+	uint8 bCanOverrideLinkedAssetAtRuntime : 1 = true;
+
 	/** True if the state contains tasks that should be called during transition handling. */
 	UPROPERTY()
 	uint8 bHasTransitionTasks : 1 = false;
diff --git a/Engine/Plugins/Runtime/StateTree/Source/StateTreeTestSuite/Private/StateTreeLinkedStateTest.cpp b/Engine/Plugins/Runtime/StateTree/Source/StateTreeTestSuite/Private/StateTreeLinkedStateTest.cpp
--- a/Engine/Plugins/Runtime/StateTree/Source/StateTreeTestSuite/Private/StateTreeLinkedStateTest.cpp
+++ b/Engine/Plugins/Runtime/StateTree/Source/StateTreeTestSuite/Private/StateTreeLinkedStateTest.cpp
@@ -1808,8 +1808,8 @@ struct FStateTreeTest_OverrideLinkedAsset : FStateTreeTestBase
 			UStateTreeEditorData& EditorData = *Cast<UStateTreeEditorData>(StateTree3.EditorData);
 			GetRootPropertyBag(EditorData).AddProperty("B", EPropertyBagPropertyType::Int32);
 
-			UStateTreeState& Root = EditorData.AddSubTree(FName(TEXT("Tree3Root")));
-			TStateTreeEditorNode<FTestTask_PrintValue>& RootTask1 = Root.AddTask<FTestTask_PrintValue>(FName(TEXT("Tree3RootTask1")));
+			UStateTreeState& Root = EditorData.AddSubTree(FName("Tree3Root"));
+			TStateTreeEditorNode<FTestTask_PrintValue>& RootTask1 = Root.AddTask<FTestTask_PrintValue>(FName("Tree3RootTask1"));
 			RootTask1.GetInstanceData().Value = -3;
 
 			EditorData.AddPropertyBinding(FPropertyBindingPath(EditorData.GetRootParametersGuid(), "B"), FPropertyBindingPath(RootTask1.ID, TEXT("Value")));
@@ -1827,19 +1827,19 @@ struct FStateTreeTest_OverrideLinkedAsset : FStateTreeTestBase
 			UStateTreeEditorData& EditorData = *Cast<UStateTreeEditorData>(StateTree2.EditorData);
 			GetRootPropertyBag(EditorData).AddProperty("A", EPropertyBagPropertyType::Int32);
 
-			UStateTreeState& Root2 = EditorData.AddSubTree(FName(TEXT("Tree2Root")));
+			UStateTreeState& Root2 = EditorData.AddSubTree(FName("Tree2Root"));
 
 			UStateTreeState& State1 = Root2.AddChildState("Tree2State1");
-			TStateTreeEditorNode<FTestTask_Stand>& State1Task1 = State1.AddTask<FTestTask_Stand>(FName(TEXT("Tree2State1Task1")));
-			TStateTreeEditorNode<FTestTask_PrintValue>& State1Task2 = State1.AddTask<FTestTask_PrintValue>(FName(TEXT("Tree2State1Task2")));
+			TStateTreeEditorNode<FTestTask_Stand>& State1Task1 = State1.AddTask<FTestTask_Stand>(FName("Tree2State1Task1"));
+			TStateTreeEditorNode<FTestTask_PrintValue>& State1Task2 = State1.AddTask<FTestTask_PrintValue>(FName("Tree2State1Task2"));
 			State1.AddTransition(EStateTreeTransitionTrigger::OnStateCompleted, EStateTreeTransitionType::NextState);
 
-			EditorData.AddPropertyBinding(FPropertyBindingPath(EditorData.GetRootParametersGuid(), "A"), FPropertyBindingPath(State1Task2.ID, TEXT("Value")));
+			EditorData.AddPropertyBinding(FPropertyBindingPath(EditorData.GetRootParametersGuid(), "A"), FPropertyBindingPath(State1Task2.ID, "Value"));
 
 			UStateTreeState& State2 = Root2.AddChildState("Tree2State2");
-			TStateTreeEditorNode<FTestTask_PrintValue>& State2Task1 = State2.AddTask<FTestTask_PrintValue>(FName(TEXT("Tree2State2Task1")));
+			TStateTreeEditorNode<FTestTask_PrintValue>& State2Task1 = State2.AddTask<FTestTask_PrintValue>(FName("Tree2State2Task1"));
 
-			EditorData.AddPropertyBinding(FPropertyBindingPath(EditorData.GetRootParametersGuid(), "A"), FPropertyBindingPath(State2Task1.ID, TEXT("Value")));
+			EditorData.AddPropertyBinding(FPropertyBindingPath(EditorData.GetRootParametersGuid(), "A"), FPropertyBindingPath(State2Task1.ID, "Value"));
 
 			{
 				FStateTreeCompiler Compiler(Log);
@@ -1849,33 +1849,42 @@ struct FStateTreeTest_OverrideLinkedAsset : FStateTreeTestBase
 		}
 
 		// Main asset
-		UStateTree& StateTree = NewStateTree();
+		UStateTree& StateTree1 = NewStateTree();
+		UStateTreeEditorData& EditorData1 = *Cast<UStateTreeEditorData>(StateTree1.EditorData);
+		FPropertyBindingPath Tree1LinkedStateParameterBindingPath;
+		FPropertyBindingPath Tree1LinkedStateCondition1BindingPath;
+		FPropertyBindingPath Tree1LinkedStateCondition2BindingPath;
 		{
-			UStateTreeEditorData& EditorData = *Cast<UStateTreeEditorData>(StateTree.EditorData);
-
-			GetRootPropertyBag(EditorData).AddProperty("Int1Param", EPropertyBagPropertyType::Int32);
+			GetRootPropertyBag(EditorData1).AddProperty("Int1Param", EPropertyBagPropertyType::Int32);
 
-			UStateTreeState& Root = EditorData.AddSubTree(FName(TEXT("Tree1Root1")));
-			UStateTreeState& State1 = Root.AddChildState(FName(TEXT("Tree1State1")), EStateTreeStateType::LinkedAsset);
+			UStateTreeState& Root = EditorData1.AddSubTree(FName("Tree1Root1"));
+			UStateTreeState& State1 = Root.AddChildState(FName("Tree1State1"), EStateTreeStateType::LinkedAsset);
 			State1.Tag = GetTestTag1();
 			State1.SetLinkedStateAsset(&StateTree2);
 
-			EditorData.AddPropertyBinding(FPropertyBindingPath(EditorData.GetRootParametersGuid(), "Int1Param"), FPropertyBindingPath(State1.Parameters.ID, "A"));
+			FStateTreeEditorNode& Cond = State1.AddEnterCondition<FTestCondition_PrintValue>(FName("Tree1State2Condition1"));
+			Tree1LinkedStateCondition1BindingPath = FPropertyBindingPath(Cond.ID, "Value");
+			
+			FStateTreeEditorNode& PropertyRefCond = State1.AddEnterCondition<FTestCondition_PropertyRefOnNodeAndInstance>(FName("Tree1State2Condition2"));
+			Tree1LinkedStateCondition2BindingPath = FPropertyBindingPath(PropertyRefCond.GetNodeID(), "RefOnNode");
 
+			Tree1LinkedStateParameterBindingPath = FPropertyBindingPath(State1.Parameters.ID, "A");
+			EditorData1.AddPropertyBinding(FPropertyBindingPath(EditorData1.GetRootParametersGuid(), "Int1Param"), Tree1LinkedStateParameterBindingPath);
+			
 			{
 				FStateTreeCompiler Compiler(Log);
-				const bool bResult = Compiler.Compile(StateTree);
+				const bool bResult = Compiler.Compile(StateTree1);
 				AITEST_TRUE(TEXT("StateTree should get compiled"), bResult);
 			}
 		}
 
 		{
 			FStateTreeInstanceData InstanceData;
-			FTestStateTreeExecutionContext Exec(StateTree, StateTree, InstanceData);
+			FTestStateTreeExecutionContext Exec(StateTree1, StateTree1, InstanceData);
 			const bool bInitSucceeded = Exec.IsValid();
 			AITEST_TRUE(TEXT("StateTree should init"), bInitSucceeded);
 
-			FInstancedPropertyBag GlobalParameters = StateTree.GetDefaultParameters();
+			FInstancedPropertyBag GlobalParameters = StateTree1.GetDefaultParameters();
 			GlobalParameters.SetValueInt32("Int1Param", 11);
 			{
 				EStateTreeRunStatus Status = Exec.Start(GlobalParameters.GetValue());
@@ -1915,6 +1924,86 @@ struct FStateTreeTest_OverrideLinkedAsset : FStateTreeTestBase
 			Exec.Stop();
 		}
 
+		auto RecompileAndStartTree = [this, &StateTree1, &Log, &StateTree2, &StateTree3](TNotNull<const UStateTree*> TargetLinkedAsset)
+			{
+				{
+					FStateTreeCompiler Compiler(Log);
+					const bool bResult = Compiler.Compile(StateTree1);
+					AITEST_TRUE(TEXT("StateTree should get compiled"), bResult);
+				}
+
+				FStateTreeInstanceData InstanceData;
+				FTestStateTreeExecutionContext Exec(StateTree1, StateTree1, InstanceData);
+				const bool bInitSucceeded = Exec.IsValid();
+				AITEST_TRUE(TEXT("StateTree should init"), bInitSucceeded);
+
+				// Set the overrides
+				{
+					FStateTreeReference StateTreeReference;
+					StateTreeReference.SetStateTree(&StateTree3);
+					FStateTreeReferenceOverrides Overrides;
+					Overrides.AddOverride(GetTestTag1(), StateTreeReference);
+					Exec.SetLinkedStateTreeOverrides(Overrides);
+				}
+
+				{
+					EStateTreeRunStatus Status = Exec.Start();
+					AITEST_EQUAL("Start should complete with Running", Status, EStateTreeRunStatus::Running);
+
+					if (TargetLinkedAsset == &StateTree2)
+					{
+						AITEST_TRUE("StateTree should be in Tree2's active frame because binding from condition to linked state parameter prevents override",
+							Exec.ExpectInActiveStates("Tree1Root1", "Tree1State1", "Tree2Root", "Tree2State1"));
+					}
+					else
+					{
+						AITEST_TRUE("StateTree should be in Tree3's active frame because of linked tree override",
+							Exec.ExpectInActiveStates("Tree1Root1", "Tree1State1", "Tree3Root"));
+					}
+					
+				}
+
+				{
+					Exec.Stop();
+					Exec.LogClear();
+				}
+
+				return true;
+			};
+
+
+		// Reset binding from global parameter to linked state parameter and add binding from state parameter to enter condition
+		// Linked asset override should be disabled
+		{
+			EditorData1.RemovePropertyBinding(Tree1LinkedStateParameterBindingPath);
+			EditorData1.AddPropertyBinding(Tree1LinkedStateParameterBindingPath, Tree1LinkedStateCondition1BindingPath);
+
+			if (!RecompileAndStartTree(&StateTree2))
+			{
+				return false;
+			}
+		}
+
+		// Reset binding from Linked state parameter to enter condition
+		// Linked asset override should be enabled
+		{
+			EditorData1.RemovePropertyBinding(Tree1LinkedStateCondition1BindingPath);
+			if (!RecompileAndStartTree(&StateTree3))
+			{
+				return false;
+			}
+		}
+
+		// Add binding from enter condition property ref to linked state parameter
+		// Linked asset override should be disabled
+		{
+			EditorData1.AddPropertyBinding(Tree1LinkedStateParameterBindingPath, Tree1LinkedStateCondition2BindingPath);
+			
+			if (!RecompileAndStartTree(&StateTree2))
+			{
+				return false;
+			}
+		}
 		return true;
 	}
 };
