[StateTree] StateTree PropertyRef should not allow binding to child classes of a UScriptStruct. Similar to UClass. PropertyRef is 2 ways, read-write. While it is fine to read to a child, it is not okay to write to a child.

#Synced-CL 51208877

diff --git a/Engine/Plugins/Runtime/StateTree/Source/StateTreeModule/Private/StateTreePropertyRefHelpers.cpp b/Engine/Plugins/Runtime/StateTree/Source/StateTreeModule/Private/StateTreePropertyRefHelpers.cpp
--- a/Engine/Plugins/Runtime/StateTree/Source/StateTreeModule/Private/StateTreePropertyRefHelpers.cpp
+++ b/Engine/Plugins/Runtime/StateTree/Source/StateTreeModule/Private/StateTreePropertyRefHelpers.cpp
@@ -167,7 +167,8 @@ namespace UE::StateTree::PropertyRefHelpers
 
 				if (SourceStructProperty)
 				{
-					if (SourceStructProperty->Struct->IsChildOf(Cast<UStruct>(TargetRefField)))
+					// Only referencing struct of the same exact type should be allowed. Otherwise one could e.g assign FVector to FMyVector:FVector property through reference missing some properties.
+					if (SourceStructProperty->Struct == TargetRefField)
 					{
 						return true;
 					}
@@ -176,7 +177,7 @@ namespace UE::StateTree::PropertyRefHelpers
 				if (const FObjectProperty* ObjectProperty = CastField<FObjectProperty>(TestProperty))
 				{
 					// Only referencing object of the same exact class should be allowed. Otherwise one could e.g assign UObject to AActor property through reference to UObject.
-					if(ObjectProperty->PropertyClass == Cast<UStruct>(TargetRefField))
+					if(ObjectProperty->PropertyClass == TargetRefField)
 					{
 						return true;
 
