[StateTree] Use UStateTreeEditingSubsystem::MarkAsPubliclyModified to mark the asset dirty instead of using the hash value. That gives a big speed improvement. Use the hash as a backup while using the editor.
Use a UStateTree::EDirtyStatus instead of the hash value to decide if the asset needs to compile. The CompileStatus can be valid and the editor data dirty when the last compilation succeed and a new modification was made. The CompileStatus can be invalid, and the editor data can be valid when the last compilation failed, and no new modifications were made.

#Synced-CL 51260989

diff --git a/Engine/Plugins/Experimental/UAF/UAFStateTree/Source/UAFStateTreeEditor/Private/StateTreeAssetCompilationHandler.cpp b/Engine/Plugins/Experimental/UAF/UAFStateTree/Source/UAFStateTreeEditor/Private/StateTreeAssetCompilationHandler.cpp
--- a/Engine/Plugins/Experimental/UAF/UAFStateTree/Source/UAFStateTreeEditor/Private/StateTreeAssetCompilationHandler.cpp
+++ b/Engine/Plugins/Experimental/UAF/UAFStateTree/Source/UAFStateTreeEditor/Private/StateTreeAssetCompilationHandler.cpp
@@ -1,4 +1,4 @@
-﻿// Copyright Epic Games, Inc. All Rights Reserved.
+// Copyright Epic Games, Inc. All Rights Reserved.
 
 #include "StateTreeAssetCompilationHandler.h"
 
@@ -18,10 +18,11 @@ namespace UE::UAF::StateTree
 FStateTreeAssetCompilationHandler::FStateTreeAssetCompilationHandler(UObject* InAsset)
 	: FAssetCompilationHandler(InAsset)
 {
-	UAnimNextStateTree* AnimNextStateTree = Cast<UAnimNextStateTree>(InAsset);
+	UAnimNextStateTree* AnimNextStateTree = CastChecked<UAnimNextStateTree>(InAsset);
+	UStateTree* StateTree = AnimNextStateTree->StateTree;
+	check(AnimNextStateTree->StateTree);
 	CachedStateTree = AnimNextStateTree->StateTree;
 	UStateTreeEditingSubsystem::ValidateStateTree(AnimNextStateTree->StateTree);
-	EditorDataHash = AnimNextStateTree->StateTree->LastCompiledEditorDataHash;
 }
 
 void FStateTreeAssetCompilationHandler::Initialize()
@@ -59,8 +60,7 @@ void FStateTreeAssetCompilationHandler::Compile(TSharedRef<Workspace::IWorkspace
 	{
 		// First compile the state tree
 		FStateTreeCompilerLog Log;
-		bLastCompileSucceeded = UStateTreeEditingSubsystem::CompileStateTree(AnimNextStateTree->StateTree, Log);
-		EditorDataHash = AnimNextStateTree->StateTree->LastCompiledEditorDataHash;
+		UStateTreeEditingSubsystem::CompileStateTree(AnimNextStateTree->StateTree, Log);
 
 		// Add to messages
 		for(const TSharedRef<FTokenizedMessage>& Message : Log.ToTokenizedMessages())
@@ -89,14 +89,13 @@ UE::UAF::Editor::ECompileStatus FStateTreeAssetCompilationHandler::GetCompileSta
 	UE::UAF::Editor::ECompileStatus StateTreeAssetStatus = UE::UAF::Editor::ECompileStatus::Unknown;
 	if (AnimNextStateTree->StateTree)
 	{
-		const bool bCompiledDataResetDuringLoad = AnimNextStateTree->StateTree->LastCompiledEditorDataHash == EditorDataHash && !AnimNextStateTree->StateTree->IsReadyToRun();
-		if (!bLastCompileSucceeded || bCompiledDataResetDuringLoad)
+		if (AnimNextStateTree->StateTree->IsEditorDataDirty())
 		{
-			StateTreeAssetStatus = UE::UAF::Editor::ECompileStatus::Error;
+			StateTreeAssetStatus = UE::UAF::Editor::ECompileStatus::Dirty;
 		}
-		else if (AnimNextStateTree->StateTree->LastCompiledEditorDataHash != EditorDataHash)
+		else if (!AnimNextStateTree->StateTree->IsReadyToRun())
 		{
-			StateTreeAssetStatus = UE::UAF::Editor::ECompileStatus::Dirty;
+			StateTreeAssetStatus = UE::UAF::Editor::ECompileStatus::Error;
 		}
 		else
 		{
@@ -116,7 +115,6 @@ void FStateTreeAssetCompilationHandler::UpdateCachedInfo()
 	}
 
 	UStateTreeEditingSubsystem::ValidateStateTree(StateTree);
-	EditorDataHash = UStateTreeEditingSubsystem::CalculateStateTreeHash(StateTree);
 	OnCompileStatusChanged().ExecuteIfBound();
 }
 
diff --git a/Engine/Plugins/Experimental/UAF/UAFStateTree/Source/UAFStateTreeEditor/Private/StateTreeAssetCompilationHandler.h b/Engine/Plugins/Experimental/UAF/UAFStateTree/Source/UAFStateTreeEditor/Private/StateTreeAssetCompilationHandler.h
--- a/Engine/Plugins/Experimental/UAF/UAFStateTree/Source/UAFStateTreeEditor/Private/StateTreeAssetCompilationHandler.h
+++ b/Engine/Plugins/Experimental/UAF/UAFStateTree/Source/UAFStateTreeEditor/Private/StateTreeAssetCompilationHandler.h
@@ -1,4 +1,4 @@
-﻿// Copyright Epic Games, Inc. All Rights Reserved.
+// Copyright Epic Games, Inc. All Rights Reserved.
 
 #pragma once
 
@@ -34,8 +34,6 @@ private:
 	void HandleStatesMoved(const TSet<UStateTreeState*>& /*AffectedParents*/, const TSet<UStateTreeState*>& /*MovedStates*/);
 
 	TWeakObjectPtr<UStateTree> CachedStateTree;
-	uint32 EditorDataHash = 0;
-	bool bLastCompileSucceeded = true;
 };
 
 }
diff --git a/Engine/Plugins/Runtime/StateTree/Source/StateTreeEditorModule/Private/StateTreeCompiler.cpp b/Engine/Plugins/Runtime/StateTree/Source/StateTreeEditorModule/Private/StateTreeCompiler.cpp
--- a/Engine/Plugins/Runtime/StateTree/Source/StateTreeEditorModule/Private/StateTreeCompiler.cpp
+++ b/Engine/Plugins/Runtime/StateTree/Source/StateTreeEditorModule/Private/StateTreeCompiler.cpp
@@ -456,6 +456,11 @@ bool FStateTreeCompiler::FailCompilation(UStateTree::ECompileStatus Status) cons
 {
 	StateTree->ResetCompiled();
 	StateTree->CompileStatus = Status;
+	// The editor data doesn't match the state tree data, but another compilation would create the same result.
+	//Mark it as not dirty and use the CompileStatus to mark the runtime data as invalid.
+	//Keeping it dirty would trigger multiple compilations that are not needed.
+	StateTree->EditorDataDirtyStatus = UStateTree::EDirtyStatus::None;
+	StateTree->LastCompiledEditorDataHash = 0;
 	return false;
 }
 
@@ -499,16 +504,17 @@ bool FStateTreeCompiler::CompilePublic(TNotNull<UStateTree*> InStateTree)
 	TRACE_CPUPROFILER_EVENT_SCOPE(FStateTreeCompiler::CompilePublic)
 	if (!PreCompile(InStateTree))
 	{
-		return FailCompilation(UStateTree::ECompileStatus::PublicDirty);
+		return FailCompilation(UStateTree::ECompileStatus::Public);
 	}
 
 	if (!CreateParameters())
 	{
-		return FailCompilation(UStateTree::ECompileStatus::PublicDirty);
+		return FailCompilation(UStateTree::ECompileStatus::Public);
 	}
 
 	// Public succeeded, the next step is dirty
-	StateTree->CompileStatus = UStateTree::ECompileStatus::InternalDirty;
+	StateTree->CompileStatus = UStateTree::ECompileStatus::Internal;
+	StateTree->EditorDataDirtyStatus = UStateTree::EDirtyStatus::Internal;
 
 	return true;
 }
@@ -545,7 +551,7 @@ bool FStateTreeCompiler::CompileInternalImpl()
 				if (const auto Validation = UE::StateTree::Compiler::IsValidIndex16(ContextDataIndex); Validation.DidFail())
 				{
 					Validation.Log(Log, TEXT("ExternalStructIndex"), ExtDataDesc);
-					return FailCompilation(UStateTree::ECompileStatus::InternalDirty);
+					return FailCompilation(UStateTree::ECompileStatus::Internal);
 				}
 				Desc.Handle.DataHandle = ExtDataDesc.DataHandle;
 			}
@@ -555,13 +561,13 @@ bool FStateTreeCompiler::CompileInternalImpl()
 	if (const UE::StateTree::Compiler::FValidationResult Validation = UE::StateTree::Compiler::IsValidIndex16(ContextDataIndex); Validation.DidFail())
 	{
 		Validation.Log(Log, TEXT("NumContextData"));
-		return FailCompilation(UStateTree::ECompileStatus::InternalDirty);
+		return FailCompilation(UStateTree::ECompileStatus::Internal);
 	}
 	StateTree->NumContextData = static_cast<uint16>(ContextDataIndex);
 
 	if (!CreateStates())
 	{
-		return FailCompilation(UStateTree::ECompileStatus::InternalDirty);
+		return FailCompilation(UStateTree::ECompileStatus::Internal);
 	}
 
 	// Eval and Global task methods use InstanceStructs.Num() as ID generator.
@@ -569,35 +575,35 @@ bool FStateTreeCompiler::CompileInternalImpl()
 	
 	if (!CreateEvaluators())
 	{
-		return FailCompilation(UStateTree::ECompileStatus::InternalDirty);
+		return FailCompilation(UStateTree::ECompileStatus::Internal);
 	}
 
 	if (!CreateGlobalTasks())
 	{
-		return FailCompilation(UStateTree::ECompileStatus::InternalDirty);
+		return FailCompilation(UStateTree::ECompileStatus::Internal);
 	}
 
 	const int32 NumGlobalInstanceData = InstanceStructs.Num();
 	if (const auto Validation = UE::StateTree::Compiler::IsValidIndex16(NumGlobalInstanceData); Validation.DidFail())
 	{
 		Validation.Log(Log, TEXT("NumGlobalInstanceData"));
-		return FailCompilation(UStateTree::ECompileStatus::InternalDirty);
+		return FailCompilation(UStateTree::ECompileStatus::Internal);
 	}
 	StateTree->NumGlobalInstanceData = uint16(NumGlobalInstanceData);
 
 	if (!CreateStateTasksAndParameters())
 	{
-		return FailCompilation(UStateTree::ECompileStatus::InternalDirty);
+		return FailCompilation(UStateTree::ECompileStatus::Internal);
 	}
 
 	if (!CreateStateTransitions())
 	{
-		return FailCompilation(UStateTree::ECompileStatus::InternalDirty);
+		return FailCompilation(UStateTree::ECompileStatus::Internal);
 	}
 
 	if (!CreateStateConsiderations())
 	{
-		return FailCompilation(UStateTree::ECompileStatus::InternalDirty);
+		return FailCompilation(UStateTree::ECompileStatus::Internal);
 	}
 
 	StateTree->Nodes = Nodes;
@@ -612,12 +618,13 @@ bool FStateTreeCompiler::CompileInternalImpl()
 	BindingsCompiler.Finalize();
 
 	// Internal succeeded, the next step is dirty
-	StateTree->CompileStatus = UStateTree::ECompileStatus::LinkDirty;
+	StateTree->CompileStatus = UStateTree::ECompileStatus::Link;
+	StateTree->EditorDataDirtyStatus = UStateTree::EDirtyStatus::Link;
 
 	if (!StateTree->Link())
 	{
 		Log.Reportf(EMessageSeverity::Error, TEXT("Unexpected failure to link the StateTree asset. See log for more info."));
-		return FailCompilation(UStateTree::ECompileStatus::LinkDirty);
+		return FailCompilation(UStateTree::ECompileStatus::Link);
 	}
 
 	// Store mapping between node unique ID and their compiled index. Used for debugging purposes.
@@ -640,7 +647,7 @@ bool FStateTreeCompiler::CompileInternalImpl()
 
 	if (!NotifyInternalPost())
 	{
-		return FailCompilation(UStateTree::ECompileStatus::InternalDirty);
+		return FailCompilation(UStateTree::ECompileStatus::Internal);
 	}
 
 	{
@@ -659,10 +666,10 @@ bool FStateTreeCompiler::CompileInternalImpl()
 
 bool FStateTreeCompiler::CompileInternal(TNotNull<UStateTree*> InStateTree)
 {
-	// For now, don't take any chances and also compile the public.
+	// Compile both step to have the public up to date.
 	if (!CompilePublic(InStateTree))
 	{
-		return FailCompilation(UStateTree::ECompileStatus::PublicDirty);
+		return FailCompilation(UStateTree::ECompileStatus::Public);
 	}
 	return CompileInternalImpl();
 }
diff --git a/Engine/Plugins/Runtime/StateTree/Source/StateTreeEditorModule/Private/StateTreeCompilerManager.cpp b/Engine/Plugins/Runtime/StateTree/Source/StateTreeEditorModule/Private/StateTreeCompilerManager.cpp
--- a/Engine/Plugins/Runtime/StateTree/Source/StateTreeEditorModule/Private/StateTreeCompilerManager.cpp
+++ b/Engine/Plugins/Runtime/StateTree/Source/StateTreeEditorModule/Private/StateTreeCompilerManager.cpp
@@ -241,7 +241,6 @@ private:
 	int32 CompileWithQueue_AnyThread(TNotNull<UStateTree*> StateTree, FStateTreeCompilerLog& Log);
 	int32 CompileWithQueueImpl_AnyThread(TNotNull<UStateTree*> StateTree, FStateTreeCompilerLog& Log, TSet<TNotNull<const UStateTree*>>& ProcessedAssets);
 	void CompileStateTree_AnyThread(TNotNull<UStateTree*> StateTree, FStateTreeCompilerLog& Log);
-	void ConditionalResetLastCompilationHash(TNotNull<UStateTree*> StateTree) const;
 	bool AllowQueuedCompilation(TNotNull<UStateTree*> StateTree) const;
 
 	void HandleStateTreeCooking(TNotNull<UStateTree*> StateTree);
@@ -316,20 +315,11 @@ bool FCompilerManagerImpl::CompilePublicSynchronously(TNotNull<UStateTree*> Stat
 
 	UStateTreeEditingSubsystem::ValidateStateTree(StateTree);
 
-	StateTree->CompileStatus = UStateTree::ECompileStatus::PublicDirty;
+	StateTree->CompileStatus = UStateTree::ECompileStatus::Public;
 
 	FStateTreeCompilerLog Log;
 	FStateTreeCompiler Compiler(Log);
-	const bool bResult = Compiler.CompilePublic(StateTree);
-	if (bResult)
-	{
-		check(StateTree->CompileStatus == UStateTree::ECompileStatus::InternalDirty);
-
-		// Queue the internal steps
-		UE::TScopeLock _(QueueLock);
-		AddToQueue_AnyThread(StateTree);
-	}
-	return bResult;
+	return Compiler.CompilePublic(StateTree);
 }
 
 bool FCompilerManagerImpl::CompileInternalSynchronously(TNotNull<UStateTree*> StateTree, FStateTreeCompilerLog& Log)
@@ -355,17 +345,8 @@ TOptional<bool> FCompilerManagerImpl::CompileIfNeededSynchronously(TNotNull<USta
 {
 	check(IsInGameThread() || IsInAsyncLoadingThread());
 
-	if (StateTree->CompileStatus == UStateTree::ECompileStatus::PublicDirty
-		|| StateTree->CompileStatus == UStateTree::ECompileStatus::InternalDirty)
+	if (StateTree->IsEditorDataDirty())
 	{
-		// The last compilation failed and we are trying to compile again with the same data (nothing changed). Do not compile.
-		//The user should trigger the compilation manually from the editor.
-		const uint32 CurrentHash = UStateTreeEditingSubsystem::CalculateStateTreeHash(StateTree);
-		if (CurrentHash == StateTree->LastCompilationEditorDataHash)
-		{
-			return {};
-		}
-
 		FStateTreeCompilerLog Log;
 		const bool bResult = CompileInternalSynchronously(StateTree, Log);
 		return TOptional<bool>(bResult);
@@ -422,7 +403,15 @@ void FCompilerManagerImpl::QueueForCompilation(TNotNull<UStateTree*> StateTree)
 	//They might be compiled via CompileIfNeededSynchronously on a worker thread.
 	if (bCanQueueCompilation && AllowQueuedCompilation(StateTree))
 	{
-		CompilePublicSynchronously(StateTree);
+		const bool bCompiled = CompilePublicSynchronously(StateTree);
+		if (bCompiled)
+		{
+			check(StateTree->CompileStatus == UStateTree::ECompileStatus::Internal);
+
+			// Queue the internal steps
+			UE::TScopeLock _(QueueLock);
+			AddToQueue_AnyThread(StateTree);
+		}
 	}
 	else
 	{
@@ -451,18 +440,6 @@ void FCompilerManagerImpl::FlushCompilationQueue()
 	bPushedBatchOnGameThread_AnyThread = !DirtyAssets_AnyThread.IsEmpty();
 }
 
-void FCompilerManagerImpl::ConditionalResetLastCompilationHash(TNotNull<UStateTree*> StateTree) const
-{
-	if (StateTree->CompileStatus == UStateTree::ECompileStatus::PublicDirty
-		|| StateTree->CompileStatus == UStateTree::ECompileStatus::InternalDirty)
-	{
-		// The last compilation failed.
-		//In CompileIfNeededSynchronously, it will trigger a compilation only if the hash has also changed.
-		//Reset the LastCompilationEditorDataHash, to trigger a new compilation.
-		StateTree->LastCompilationEditorDataHash = 0;
-	}
-}
-
 bool FCompilerManagerImpl::AllowQueuedCompilation(TNotNull<UStateTree*> StateTree) const
 {
 	if (UStateTreeEditorData* EditorData = Cast<UStateTreeEditorData>(StateTree->EditorData))
@@ -565,20 +542,23 @@ int32 FCompilerManagerImpl::CompileWithQueueImpl_AnyThread(TNotNull<UStateTree*>
 		const TSet<UStateTree*> Dependencies = GatherStateTreeDependencies_AnyThread(StateTree);
 		for (UStateTree* Dependency : Dependencies)
 		{
-			if (ProcessedAssets.Contains(Dependency))
-			{
-				UE_LOG(LogStateTree, Warning, TEXT("Circular dependency detected while processing %s and %s."), *StateTree->GetFullName(), *Dependency->GetFullName());
-			}
-			else
+			if (Dependency->IsEditorDataDirty())
 			{
-				const int32 DirtyAssetsIndex = DirtyAssets_AnyThread.IndexOfByKey(TObjectKey<UStateTree>(Dependency));
-				if (DirtyAssetsIndex != INDEX_NONE)
+				if (ProcessedAssets.Contains(Dependency))
 				{
-					DirtyAssets_AnyThread.RemoveAtSwap(DirtyAssetsIndex, EAllowShrinking::No);
+					UE_LOG(LogStateTree, Warning, TEXT("Circular dependency detected while processing %s and %s."), *StateTree->GetFullName(), *Dependency->GetFullName());
+				}
+				else
+				{
+					const int32 DirtyAssetsIndex = DirtyAssets_AnyThread.IndexOfByKey(TObjectKey<UStateTree>(Dependency));
+					if (DirtyAssetsIndex != INDEX_NONE)
+					{
+						DirtyAssets_AnyThread.RemoveAtSwap(DirtyAssetsIndex, EAllowShrinking::No);
 
-					FStateTreeCompilerLog NewLog;
-					ProcessedAssets.Add(Dependency);
-					CompiledStateTreeCount += CompileWithQueueImpl_AnyThread(Dependency, NewLog, ProcessedAssets);
+						FStateTreeCompilerLog NewLog;
+						ProcessedAssets.Add(Dependency);
+						CompiledStateTreeCount += CompileWithQueueImpl_AnyThread(Dependency, NewLog, ProcessedAssets);
+					}
 				}
 			}
 		}
@@ -595,7 +575,6 @@ void FCompilerManagerImpl::CompileStateTree_AnyThread(TNotNull<UStateTree*> Stat
 	const bool bCompilationResult = Compiler.Compile(StateTree);
 
 	const uint32 NewHash = UStateTreeEditingSubsystem::CalculateStateTreeHash(StateTree);
-	StateTree->LastCompilationEditorDataHash = NewHash;
 	checkf(StateTree->LastCompiledEditorDataHash == 0, TEXT("The compilation has to reset the statetree data. Including the hash value."));
 
 	if (bCompilationResult)
@@ -706,18 +685,38 @@ void FCompilerManagerImpl::UpdateEditorBindingsIfNeeded(TSet<const UStruct*>& In
 
 void FCompilerManagerImpl::HandleStateTreeCooking(TNotNull<UStateTree*> StateTree)
 {
-	switch (StateTree->CompileStatus)
+	if (StateTree->IsEditorDataDirty())
 	{
-		case UStateTree::ECompileStatus::PublicDirty:
-		case UStateTree::ECompileStatus::InternalDirty:
+		switch (StateTree->CompileStatus)
+		{
+		case UStateTree::ECompileStatus::Public:
+		case UStateTree::ECompileStatus::Internal:
+		case UStateTree::ECompileStatus::Executable:
 		{
 			FStateTreeCompilerLog Log;
 			CompileInternalSynchronously(StateTree, Log);
 			break;
 		}
-		case UStateTree::ECompileStatus::LinkDirty:
+		case UStateTree::ECompileStatus::Link:
 			std::ignore = StateTree->Link();
 			break;
+		}
+	}
+	else
+	{
+		// Data that is not dirty and requests a compilation should not need a new compilation.
+		//The last Link failed.
+		//Since we are cooking, compile again and report the errors for the build machine to report the errors.
+		if (StateTree->CompileStatus == UStateTree::ECompileStatus::Link)
+		{
+
+			std::ignore = StateTree->Link();
+		}
+		else if (StateTree->CompileStatus != UStateTree::ECompileStatus::Executable)
+		{
+			FStateTreeCompilerLog Log;
+			CompileInternalSynchronously(StateTree, Log);
+		}
 	}
 
 	if (StateTree->CompileStatus != UStateTree::ECompileStatus::Executable)
@@ -797,29 +796,8 @@ void FCompilerManagerImpl::HandleObjectsReinstanced(const FReplacementObjectMap&
 
 		for (UStateTree* StateTree : StateTreeToLink)
 		{
-			ConditionalResetLastCompilationHash(StateTree);
 			if (!StateTree->Link())
 			{
-				//If the task modification make it not linkable, then we need to recompile the asset.
-				bool bModifyPublic = true;
-				const TObjectKey<UStateTree> StateTreeKey = StateTree;
-				if (TSharedPtr<FStateTreeDependencies>* FoundDependencies = StateTreeToDependencies.Find(StateTreeKey))
-				{
-					bModifyPublic = false;
-					for (const FStateTreeDependencies::FItem& Dependency : (*FoundDependencies)->Dependencies)
-					{
-						if (StructsToBeReplacedAsKey.Contains(Dependency.Key))
-						{
-							if (EnumHasAnyFlags(Dependency.Type, FStateTreeDependencies::EDependencyType::DT_Public))
-							{
-								bModifyPublic = true;
-								break;
-							}
-						}
-					}
-				}
-
-				StateTree->MarkAsModified(bModifyPublic);
 				UE_LOG(LogStateTree, Error, TEXT("%s failed to link after Object reinstantiation. Take a look at the asset for any errors. Asset will not be usable at runtime."), *StateTree->GetPathName());
 			}
 		}
@@ -830,6 +808,14 @@ void FCompilerManagerImpl::HandleObjectsReinstanced(const FReplacementObjectMap&
 		{
 			UStateTree* StateTree = *It;
 			check(!StateTree->HasAnyFlags(RF_ClassDefaultObject));
+
+			// If the asset is not linked yet (or has failed), no need to link.
+			if (StateTree->CompileStatus != UStateTree::ECompileStatus::Executable
+				&& StateTree->CompileStatus != UStateTree::ECompileStatus::Link)
+			{
+				continue;
+			}
+
 			bool bShouldRelink = false;
 
 			// Relink if one of the out of date objects got reinstanced.
@@ -848,12 +834,6 @@ void FCompilerManagerImpl::HandleObjectsReinstanced(const FReplacementObjectMap&
 				}
 			}
 
-			// If the asset is not linked yet (or has failed), no need to link.
-			if (!bShouldRelink && StateTree->CompileStatus != UStateTree::ECompileStatus::Executable)
-			{
-				continue;
-			}
-
 			// Relink only if the reinstantiated object belongs to this asset,
 			// or anything from the property binding refers to the classes of the reinstantiated object.
 			if (!bShouldRelink)
@@ -875,7 +855,6 @@ void FCompilerManagerImpl::HandleObjectsReinstanced(const FReplacementObjectMap&
 
 			if (bShouldRelink)
 			{
-				ConditionalResetLastCompilationHash(StateTree);
 				if (!StateTree->Link())
 				{
 					UE_LOG(LogStateTree, Error, TEXT("%s failed to link after Object reinstantiation. Take a look at the asset for any errors. Asset will not be usable at runtime."), *StateTree->GetPathName());
@@ -910,7 +889,6 @@ void FCompilerManagerImpl::HandleUserDefinedStructReinstanced(const UUserDefined
 
 		for (UStateTree* StateTree : StateTreeToLink)
 		{
-			ConditionalResetLastCompilationHash(StateTree);
 			if (!StateTree->Link())
 			{
 				UE_LOG(LogStateTree, Error, TEXT("%s failed to link after Object reinstantiation. Take a look at the asset for any errors. Asset will not be usable at runtime."), *StateTree->GetPathName());
@@ -924,7 +902,6 @@ void FCompilerManagerImpl::HandleUserDefinedStructReinstanced(const UUserDefined
 			UStateTree* StateTree = *It;
 			if (StateTree->PropertyBindings.ContainsAnyStruct(Structs))
 			{
-				ConditionalResetLastCompilationHash(StateTree);
 				if (!StateTree->Link())
 				{
 					UE_LOG(LogStateTree, Error, TEXT("%s failed to link after Struct reinstantiation. Take a look at the asset for any errors. Asset will not be usable at runtime."), *StateTree->GetPathName());
diff --git a/Engine/Plugins/Runtime/StateTree/Source/StateTreeEditorModule/Private/StateTreeEditingSubsystem.cpp b/Engine/Plugins/Runtime/StateTree/Source/StateTreeEditorModule/Private/StateTreeEditingSubsystem.cpp
--- a/Engine/Plugins/Runtime/StateTree/Source/StateTreeEditorModule/Private/StateTreeEditingSubsystem.cpp
+++ b/Engine/Plugins/Runtime/StateTree/Source/StateTreeEditorModule/Private/StateTreeEditingSubsystem.cpp
@@ -54,22 +54,22 @@ bool UStateTreeEditingSubsystem::NeedsRecompile(TNotNull<const UStateTree*> InSt
 
 void UStateTreeEditingSubsystem::MarkAsPubliclyModified(TNotNull<UStateTree*> InStateTree)
 {
+	constexpr bool bPubliclyModified = true;
+	InStateTree->MarkAsModified(bPubliclyModified);
+
 	if (UE::StateTree::Editor::bEnableQueuedCompilationWhenAssetDirty)
 	{
-		//@todo fix mark as modified to only trigger a runtime compilation
-		constexpr bool bPubliclyModified = true;
-		InStateTree->MarkAsModified(bPubliclyModified);
 		UE::StateTree::Compiler::FCompilerManager::QueueForCompilation(InStateTree);
 	}
 }
 
 void UStateTreeEditingSubsystem::MarkAsModified(TNotNull<UStateTree*> InStateTree)
 {
+	constexpr bool bPubliclyModified = false;
+	InStateTree->MarkAsModified(bPubliclyModified);
+
 	if (UE::StateTree::Editor::bEnableQueuedCompilationWhenAssetDirty)
 	{
-		//@todo fix mark as modified to only trigger a runtime compilation
-		constexpr bool bPubliclyModified = false;
-		InStateTree->MarkAsModified(bPubliclyModified);
 		UE::StateTree::Compiler::FCompilerManager::QueueForCompilation(InStateTree);
 	}
 }
diff --git a/Engine/Plugins/Runtime/StateTree/Source/StateTreeEditorModule/Private/StateTreeEditorMode.cpp b/Engine/Plugins/Runtime/StateTree/Source/StateTreeEditorModule/Private/StateTreeEditorMode.cpp
--- a/Engine/Plugins/Runtime/StateTree/Source/StateTreeEditorModule/Private/StateTreeEditorMode.cpp
+++ b/Engine/Plugins/Runtime/StateTree/Source/StateTreeEditorModule/Private/StateTreeEditorMode.cpp
@@ -434,13 +434,13 @@ void UStateTreeEditorMode::Compile()
 	}
 	
 	FStateTreeCompilerLog Log;
-	bLastCompileSucceeded = UStateTreeEditingSubsystem::CompileStateTree(StateTree, Log);
+	const bool bCompileSucceeded = UStateTreeEditingSubsystem::CompileStateTree(StateTree, Log);
 
 	if (TSharedPtr<IMessageLogListing> Listing = GetMessageLogListing())
-	{					
+	{
 		Log.AppendToLog(Listing.Get());
 
-		if (!bLastCompileSucceeded)
+		if (!bCompileSucceeded)
 		{
 			// Show log
 			ShowCompilerTab();
@@ -450,7 +450,7 @@ void UStateTreeEditorMode::Compile()
 
 	const UStateTreeEditorSettings* Settings = GetMutableDefault<UStateTreeEditorSettings>();
 	const bool bShouldSaveOnCompile = ((Settings->SaveOnCompile == EStateTreeSaveOnCompile::Always)
-									|| ((Settings->SaveOnCompile == EStateTreeSaveOnCompile::SuccessOnly) && bLastCompileSucceeded));
+									|| ((Settings->SaveOnCompile == EStateTreeSaveOnCompile::SuccessOnly) && bCompileSucceeded));
 
 	if (bShouldSaveOnCompile)
 	{
@@ -964,7 +964,11 @@ void UStateTreeEditorMode::UpdateAsset()
 	}
 
 	UStateTreeEditingSubsystem::ValidateStateTree(StateTree);
-	EditorDataHash = UStateTreeEditingSubsystem::CalculateStateTreeHash(StateTree);
+	const uint32 CurrentEditorDataHash = UStateTreeEditingSubsystem::CalculateStateTreeHash(StateTree);
+	if (CurrentEditorDataHash != StateTree->LastCompiledEditorDataHash)
+	{
+		UStateTreeEditingSubsystem::MarkAsModified(StateTree);
+	}
 }
 
 TSharedPtr<IDetailsView> UStateTreeEditorMode::GetDetailsView() const
@@ -1077,6 +1081,7 @@ void UStateTreeEditorMode::OnSelectionFinishedChangingProperties(const FProperty
 			{
 				StateTreeEditingSubsystem->FindOrAddViewModel(GetStateTree())->NotifyStatesChangedExternally(ChangedStates, PropertyChangedEvent);
 				UpdateAsset();
+				UStateTreeEditingSubsystem::MarkAsModified(GetStateTree());
 			}
 		}
 	}
diff --git a/Engine/Plugins/Runtime/StateTree/Source/StateTreeEditorModule/Private/StateTreeEditorModeToolkit.cpp b/Engine/Plugins/Runtime/StateTree/Source/StateTreeEditorModule/Private/StateTreeEditorModeToolkit.cpp
--- a/Engine/Plugins/Runtime/StateTree/Source/StateTreeEditorModule/Private/StateTreeEditorModeToolkit.cpp
+++ b/Engine/Plugins/Runtime/StateTree/Source/StateTreeEditorModule/Private/StateTreeEditorModeToolkit.cpp
@@ -219,16 +219,13 @@ FSlateIcon FStateTreeEditorModeToolkit::GetCompileStatusImage() const
 		UStateTree* StateTree = EditorMode->GetStateTree();
 		if (StateTree)
 		{
-			const bool bCompiledDataResetDuringLoad = StateTree->LastCompiledEditorDataHash == EditorMode->EditorDataHash && !StateTree->IsReadyToRun();
-
-			if (!EditorMode->bLastCompileSucceeded || bCompiledDataResetDuringLoad)
+			if (StateTree->IsEditorDataDirty())
 			{
-				return FSlateIcon(FAppStyle::GetAppStyleSetName(), CompileStatusBackground, NAME_None, CompileStatusError);
+				return FSlateIcon(FAppStyle::GetAppStyleSetName(), CompileStatusBackground, NAME_None, CompileStatusUnknown);
 			}
-
-			if (StateTree->LastCompiledEditorDataHash != EditorMode->EditorDataHash)
+			else if (!StateTree->IsReadyToRun())
 			{
-				return FSlateIcon(FAppStyle::GetAppStyleSetName(), CompileStatusBackground, NAME_None, CompileStatusUnknown);
+				return FSlateIcon(FAppStyle::GetAppStyleSetName(), CompileStatusBackground, NAME_None, CompileStatusError);
 			}
 
 			return FSlateIcon(FAppStyle::GetAppStyleSetName(), CompileStatusBackground, NAME_None, CompileStatusGood);
diff --git a/Engine/Plugins/Runtime/StateTree/Source/StateTreeEditorModule/Private/StateTreeViewModel.cpp b/Engine/Plugins/Runtime/StateTree/Source/StateTreeEditorModule/Private/StateTreeViewModel.cpp
--- a/Engine/Plugins/Runtime/StateTree/Source/StateTreeEditorModule/Private/StateTreeViewModel.cpp
+++ b/Engine/Plugins/Runtime/StateTree/Source/StateTreeEditorModule/Private/StateTreeViewModel.cpp
@@ -840,21 +840,11 @@ void FStateTreeViewModel::RefreshDebuggerBreakpoints()
 void FStateTreeViewModel::NotifyAssetChangedExternally() const
 {
 	OnAssetChanged.Broadcast();
-
-	if (UStateTreeEditorData* TreeData = TreeDataWeak.Get())
-	{
-		UStateTreeEditingSubsystem::MarkAsModified(TreeData->GetTypedOuter<UStateTree>());
-	}
 }
 
 void FStateTreeViewModel::NotifyStatesChangedExternally(const TSet<UStateTreeState*>& ChangedStates, const FPropertyChangedEvent& PropertyChangedEvent) const
 {
 	OnStatesChanged.Broadcast(ChangedStates, PropertyChangedEvent);
-
-	if (UStateTreeEditorData* TreeData = TreeDataWeak.Get())
-	{
-		UStateTreeEditingSubsystem::MarkAsModified(TreeData->GetTypedOuter<UStateTree>());
-	}
 }
 
 TArray<TObjectPtr<UStateTreeState>>* FStateTreeViewModel::GetSubTrees() const
@@ -959,11 +949,21 @@ void FStateTreeViewModel::PostUndo(bool bSuccess)
 {
 	// TODO: see if we can narrow this down.
 	NotifyAssetChangedExternally();
+
+	if (UStateTreeEditorData* TreeData = TreeDataWeak.Get())
+	{
+		UStateTreeEditingSubsystem::MarkAsModified(TreeData->GetTypedOuter<UStateTree>());
+	}
 }
 
 void FStateTreeViewModel::PostRedo(bool bSuccess)
 {
 	NotifyAssetChangedExternally();
+
+	if (UStateTreeEditorData* TreeData = TreeDataWeak.Get())
+	{
+		UStateTreeEditingSubsystem::MarkAsModified(TreeData->GetTypedOuter<UStateTree>());
+	}
 }
 
 void FStateTreeViewModel::ClearSelection()
@@ -1728,6 +1728,11 @@ void FStateTreeViewModel::RenameState(UStateTreeState* State, FName NewName)
 	FProperty* NameProperty = FindFProperty<FProperty>(UStateTreeState::StaticClass(), GET_MEMBER_NAME_CHECKED(UStateTreeState, Name));
 	FPropertyChangedEvent PropertyChangedEvent(NameProperty, EPropertyChangeType::ValueSet);
 	NotifyStatesChangedExternally(AffectedStates, PropertyChangedEvent);
+
+	if (UStateTreeEditorData* TreeData = TreeDataWeak.Get())
+	{
+		UStateTreeEditingSubsystem::MarkAsModified(TreeData->GetTypedOuter<UStateTree>());
+	}
 }
 
 void FStateTreeViewModel::RemoveSelectedStates()
@@ -2045,6 +2050,10 @@ void FStateTreeViewModel::SetSelectedStatesEnabled(const bool bEnable)
 		}
 
 		NotifyAssetChangedExternally();
+		if (UStateTreeEditorData* TreeData = TreeDataWeak.Get())
+		{
+			UStateTreeEditingSubsystem::MarkAsModified(TreeData->GetTypedOuter<UStateTree>());
+		}
 	}
 }
 
diff --git a/Engine/Plugins/Runtime/StateTree/Source/StateTreeEditorModule/Public/StateTreeEditorMode.h b/Engine/Plugins/Runtime/StateTree/Source/StateTreeEditorModule/Public/StateTreeEditorMode.h
--- a/Engine/Plugins/Runtime/StateTree/Source/StateTreeEditorModule/Public/StateTreeEditorMode.h
+++ b/Engine/Plugins/Runtime/StateTree/Source/StateTreeEditorModule/Public/StateTreeEditorMode.h
@@ -93,7 +93,9 @@ protected:
 	friend class FStateTreeEditorModeToolkit;
 
 protected:
+	UE_DEPRECATED(5.8, "UStateTree::IsEditorDataDirty")
 	uint32 EditorDataHash = 0;
+	UE_DEPRECATED(5.8, "UStateTree::CompileStatus")
 	bool bLastCompileSucceeded = true;
 	bool bForceAssetDetailViewToRefresh = false;
 
diff --git a/Engine/Plugins/Runtime/StateTree/Source/StateTreeModule/Private/StateTree.cpp b/Engine/Plugins/Runtime/StateTree/Source/StateTreeModule/Private/StateTree.cpp
--- a/Engine/Plugins/Runtime/StateTree/Source/StateTreeModule/Private/StateTree.cpp
+++ b/Engine/Plugins/Runtime/StateTree/Source/StateTreeModule/Private/StateTree.cpp
@@ -375,7 +375,8 @@ void UStateTree::ResetCompiled()
 	
 	ResetLinked();
 
-	CompileStatus = ECompileStatus::PublicDirty;
+	CompileStatus = ECompileStatus::Public;
+	EditorDataDirtyStatus = EDirtyStatus::Public;
 	LastCompiledEditorDataHash = 0;
 
 	// Remove objects created from last compilation.
@@ -397,9 +398,21 @@ void UStateTree::ResetCompiled()
 	}
 }
 
+bool UStateTree::IsEditorDataDirty() const
+{
+	return EditorDataDirtyStatus != EDirtyStatus::None;
+}
+
 void UStateTree::MarkAsModified(const bool bPubliclyModified)
 {
-	CompileStatus = bPubliclyModified ? ECompileStatus::PublicDirty : ECompileStatus::InternalDirty;
+	if (bPubliclyModified)
+	{
+		EditorDataDirtyStatus = EDirtyStatus::Public;
+	}
+	else if (EditorDataDirtyStatus != EDirtyStatus::Public)
+	{
+		EditorDataDirtyStatus = EDirtyStatus::Internal;
+	}
 }
 
 void UStateTree::PreSave(FObjectPreSaveContext Context)
@@ -553,7 +566,7 @@ void UStateTree::PostLoad()
 	}
 
 	// Request a link
-	CompileStatus = ECompileStatus::LinkDirty;
+	CompileStatus = ECompileStatus::Link;
 
 #if WITH_EDITOR
 	if (EditorData)
@@ -567,11 +580,14 @@ void UStateTree::PostLoad()
 	// Request a compilation (can be async or sync).
 	UE::StateTree::Delegates::Private::OnStateTreeAssetLoaded.ExecuteIfBound(this);
 
-	if (CompileStatus == ECompileStatus::LinkDirty && LastCompiledEditorDataHash != 0)
+	// The OnStateTreeAssetLoaded can trigger a compilation. That would also trigger a Link()
+	// Only Link() if no compilation occurs.
+	// LastCompiledEditorDataHash is only set if a compilation is valid. A valid compilation set CompileStatus to Executable.
+	if (CompileStatus == ECompileStatus::Link && LastCompiledEditorDataHash != 0)
 	{
 		if (!Link())
 		{
-			UE_CLOG(CompileStatus == ECompileStatus::LinkDirty, LogStateTree, Warning, TEXT("%s failed to link. Asset will not be usable at runtime."), *GetPathName());
+			UE_CLOG(CompileStatus == ECompileStatus::Link, LogStateTree, Warning, TEXT("%s failed to link. Asset will not be usable at runtime."), *GetPathName());
 		}
 	}
 }
@@ -623,8 +639,14 @@ void UStateTree::ResetLinked()
 {
 	if (CompileStatus == ECompileStatus::Executable)
 	{
-		CompileStatus = ECompileStatus::LinkDirty;
+		CompileStatus = ECompileStatus::Link;
 	}
+#if WITH_EDITOR
+	if (EditorDataDirtyStatus == EDirtyStatus::None)
+	{
+		EditorDataDirtyStatus = EDirtyStatus::Link;
+	}
+#endif
 
 	ExternalDataDescs.Reset();
 
@@ -783,7 +805,7 @@ bool UStateTree::ValidateInstanceData()
 
 bool UStateTree::Link()
 {
-	if (CompileStatus == ECompileStatus::PublicDirty || CompileStatus == ECompileStatus::InternalDirty)
+	if (CompileStatus == ECompileStatus::Public || CompileStatus == ECompileStatus::Internal)
 	{
 		return false;
 	}
@@ -907,6 +929,12 @@ bool UStateTree::Link()
 
 	// Link succeeded, setup tree to be ready to run
 	CompileStatus = ECompileStatus::Executable;
+#if WITH_EDITOR
+	if (EditorDataDirtyStatus == EDirtyStatus::Link)
+	{
+		EditorDataDirtyStatus = EDirtyStatus::None;
+	}
+#endif
 	
 	return true;
 }
@@ -1570,9 +1598,7 @@ TArray<FStateTreeMemoryUsage> UStateTree::CalculateEstimatedMemoryUsage() const
 
 void UStateTree::CompileIfChanged()
 {
-	// Same check are in the compiler manager. But do it here also as a quick optimization.
-	if (CompileStatus == UStateTree::ECompileStatus::PublicDirty
-		|| CompileStatus == UStateTree::ECompileStatus::InternalDirty)
+	if (IsEditorDataDirty())
 	{
 		UE::StateTree::Delegates::Private::OnCompileIfChanged.ExecuteIfBound(this);
 	}
diff --git a/Engine/Plugins/Runtime/StateTree/Source/StateTreeModule/Public/StateTree.h b/Engine/Plugins/Runtime/StateTree/Source/StateTreeModule/Public/StateTree.h
--- a/Engine/Plugins/Runtime/StateTree/Source/StateTreeModule/Public/StateTree.h
+++ b/Engine/Plugins/Runtime/StateTree/Source/StateTreeModule/Public/StateTree.h
@@ -187,8 +187,11 @@ public:
 	/** @return The global parameters of the state tree. Global parameter values can be overridden at runtime by the execution context. */
 	UE_API FConstStructView GetDefaultGlobalParameters() const;
 
-	/** @return true if the tree asset can be used at runtime. */
-	UE_API bool IsReadyToRun() const;
+	/**
+	 * @return true if the tree asset can be used at runtime.
+	 * @note the asset can be in the compilation queue. This will make sure it is up to date.
+	 */
+	[[nodiscard]] UE_API bool IsReadyToRun() const;
 
 	/** @return schema that was used to compile the StateTree. */
 	const UStateTreeSchema* GetSchema() const
@@ -334,12 +337,18 @@ public:
 	UE_API void ResetCompiled();
 
 	/**
-	 * Mark the StateTree asset dirty.
+	 * Mark the state tree dirty.
 	 * It will require the StateTree to be compiled before it can be used.
 	 * @param bPubliclyModified Whether to also mark the public dependency dirty. It will affect the other asset that references the StateTree asset.
 	 */
 	UE_API void MarkAsModified(bool bPubliclyModified);
 
+	/** @return whether the state tree asset is dirty and needs to be compiled. */
+	[[nodiscard]] UE_API bool IsEditorDataDirty() const;
+
+	/** Compile the state tree if the editor data as changed since the last compilation. */
+	UE_API void CompileIfChanged();
+
 	/** Calculates runtime memory usage for different sections of the tree. */
 	UE_API TArray<FStateTreeMemoryUsage> CalculateEstimatedMemoryUsage() const;
 
@@ -349,9 +358,6 @@ public:
 	{
 	}
 
-	/** Compile the state trees if the editor hash data as changed since the last compilation. */
-	UE_API void CompileIfChanged();
-
 	/** The expected result of GetAssetSavedVersion. If the values do not match, the UI will suggest saving the asset. */
 	UE_API static const int32 LatestCustomAssetSavedVersion;
 
@@ -570,16 +576,36 @@ private:
 
 	enum class ECompileStatus : uint8
 	{
-		PublicDirty,		// The exported dependencies need to be compiled. 
-		InternalDirty,		// The public compiled and the internal dependencies need to be compiled.
-		LinkDirty,			// Compiled and not yet link.
-		Executable,			// Compiled and ready to use.
+		Public,			// The exported dependencies need to be compiled. 
+		Internal,		// The public compiled and the internal dependencies need to be compiled.
+		Link,			// Compiled and not yet link.
+		Executable,		// Compiled and ready to use.
 	};
-	ECompileStatus CompileStatus = ECompileStatus::LinkDirty;
+	/**
+	 * The current state of the asset.
+	 * The compilation could have failed or is pending (in the compiler queue).
+	 * This is different from EditorDataDirtyStatus.
+	 * The compilation/link could have failed but the data didn't change since the last compilation/link.
+	 */
+	ECompileStatus CompileStatus = ECompileStatus::Link;
 
 #if WITH_EDITORONLY_DATA
-	/** The hash during the last compilation. Valid whether the compilation succeeded or failed. */
-	uint32 LastCompilationEditorDataHash = 0;
+	enum class EDirtyStatus : uint8
+	{
+		Public,			// Publicly produced dependencies changed since the last compilation.
+		Internal,		// Internal/Private nodes/properties changed since the last compilation.
+						//A new task is considered a modification.
+						//A property value changed on a task is considered a modification.
+						//A BP task that added a property is not considered a modification. It wouldn't affect the compilation result.
+		Link,			// Internal/Private dependencies that have a binding changed since the last compilation.
+						//A BP task that removed a property and that property had a binding is not considered a modification.
+		None			// Nothing changed since the last compilation.
+	};
+	/**
+	 * Whether the editor data dirtied since the last compilation.
+	 * After a compilation (success or failure), the dirty status returns to none.
+	 */
+	EDirtyStatus EditorDataDirtyStatus = EDirtyStatus::None;
 #endif
 
 	/** True if any global task is a transition task. */
diff --git a/Engine/Plugins/VirtualProduction/Avalanche/Source/AvalancheTransitionEditor/Private/Compiler/AvaTransitionCompiler.cpp b/Engine/Plugins/VirtualProduction/Avalanche/Source/AvalancheTransitionEditor/Private/Compiler/AvaTransitionCompiler.cpp
--- a/Engine/Plugins/VirtualProduction/Avalanche/Source/AvalancheTransitionEditor/Private/Compiler/AvaTransitionCompiler.cpp
+++ b/Engine/Plugins/VirtualProduction/Avalanche/Source/AvalancheTransitionEditor/Private/Compiler/AvaTransitionCompiler.cpp
@@ -140,7 +140,6 @@ FSlateIcon FAvaTransitionCompiler::GetCompileStatusIcon() const
 	static const FName CompileStatusUnknown("Blueprint.CompileStatus.Overlay.Unknown");
 	static const FName CompileStatusError("Blueprint.CompileStatus.Overlay.Error");
 	static const FName CompileStatusGood("Blueprint.CompileStatus.Overlay.Good");
-	static const FName CompileStatusWarning("Blueprint.CompileStatus.Overlay.Warning");
 
 	UAvaTransitionTree* TransitionTree = TransitionTreeWeak.Get();
 	if (!TransitionTree)
@@ -148,16 +147,13 @@ FSlateIcon FAvaTransitionCompiler::GetCompileStatusIcon() const
 		return FSlateIcon(FAppStyle::GetAppStyleSetName(), CompileStatusBackground, NAME_None, CompileStatusUnknown);
 	}
 
-	const bool bCompiledDataResetDuringLoad = TransitionTree->LastCompiledEditorDataHash == EditorDataHash && !TransitionTree->IsReadyToRun();
-
-	if (!bLastCompileSucceeded || bCompiledDataResetDuringLoad)
+	if (TransitionTree->IsEditorDataDirty())
 	{
-		return FSlateIcon(FAppStyle::GetAppStyleSetName(), CompileStatusBackground, NAME_None, CompileStatusError);
+		return FSlateIcon(FAppStyle::GetAppStyleSetName(), CompileStatusBackground, NAME_None, CompileStatusUnknown);
 	}
-
-	if (TransitionTree->LastCompiledEditorDataHash != EditorDataHash)
+	else if (!TransitionTree->IsReadyToRun())
 	{
-		return FSlateIcon(FAppStyle::GetAppStyleSetName(), CompileStatusBackground, NAME_None, CompileStatusUnknown);
+		return FSlateIcon(FAppStyle::GetAppStyleSetName(), CompileStatusBackground, NAME_None, CompileStatusError);
 	}
 
 	return FSlateIcon(FAppStyle::GetAppStyleSetName(), CompileStatusBackground, NAME_None, CompileStatusGood);
