[State Tree] set cvar CVarAutoRemoveSourcePathInvalidPropertyBindings default to false

#Synced-CL 51266215

diff --git a/Engine/Plugins/Runtime/StateTree/Source/StateTreeEditorModule/Private/StateTreeEditorData.cpp b/Engine/Plugins/Runtime/StateTree/Source/StateTreeEditorModule/Private/StateTreeEditorData.cpp
--- a/Engine/Plugins/Runtime/StateTree/Source/StateTreeEditorModule/Private/StateTreeEditorData.cpp
+++ b/Engine/Plugins/Runtime/StateTree/Source/StateTreeEditorModule/Private/StateTreeEditorData.cpp
@@ -42,7 +42,7 @@ namespace UE::StateTree::Editor
 
 	FAutoConsoleVariable CVarAutoRemoveSourcePathInvalidPropertyBindings(
 		TEXT("StateTree.Editor.AutoRemoveSourcePathInvalidEditorBindings"),
-		true,
+		false,
 		TEXT("If set to true, EditorBindings whose source path does not resolve anymore,")
 		TEXT("will be automatically removed when the asset is loaded, when the users authors in the editor,")
 		TEXT("when state tree gets compiled.\n")
