[StateTree] Add bCompilationPending to be more explicit that a StateTree asset is in the compiler queue and the compilation is not completed.

#Synced-CL 51260989

diff --git a/Engine/Plugins/Runtime/StateTree/Source/StateTreeEditorModule/Private/StateTreeCompilerManager.cpp b/Engine/Plugins/Runtime/StateTree/Source/StateTreeEditorModule/Private/StateTreeCompilerManager.cpp
--- a/Engine/Plugins/Runtime/StateTree/Source/StateTreeEditorModule/Private/StateTreeCompilerManager.cpp
+++ b/Engine/Plugins/Runtime/StateTree/Source/StateTreeEditorModule/Private/StateTreeCompilerManager.cpp
@@ -457,6 +457,8 @@ void FCompilerManagerImpl::AddToQueue_AnyThread(TNotNull<UStateTree*> StateTree)
 	TObjectKey<UStateTree> StateTreeKey = StateTree;
 	DirtyAssets_AnyThread.AddUnique(StateTreeKey);
 
+	StateTree->bCompilationPending = true;
+
 	SendToBatchCompilationQueue_AnyThread();
 }
 
@@ -577,6 +579,8 @@ void FCompilerManagerImpl::CompileStateTree_AnyThread(TNotNull<UStateTree*> Stat
 	const uint32 NewHash = UStateTreeEditingSubsystem::CalculateStateTreeHash(StateTree);
 	checkf(StateTree->LastCompiledEditorDataHash == 0, TEXT("The compilation has to reset the statetree data. Including the hash value."));
 
+	StateTree->bCompilationPending = false;
+
 	if (bCompilationResult)
 	{
 		StateTree->LastCompiledEditorDataHash = NewHash;
diff --git a/Engine/Plugins/Runtime/StateTree/Source/StateTreeModule/Private/StateTree.cpp b/Engine/Plugins/Runtime/StateTree/Source/StateTreeModule/Private/StateTree.cpp
--- a/Engine/Plugins/Runtime/StateTree/Source/StateTreeModule/Private/StateTree.cpp
+++ b/Engine/Plugins/Runtime/StateTree/Source/StateTreeModule/Private/StateTree.cpp
@@ -57,8 +57,11 @@ FConstStructView UStateTree::GetDefaultGlobalParameters() const
 bool UStateTree::IsReadyToRun() const
 {
 #if WITH_EDITOR
-	// The asset can be in the compiler queue. Compile the asset sync.
-	const_cast<UStateTree*>(this)->CompileIfChanged();
+	if (bCompilationPending)
+	{
+		// The asset can be in the compiler queue. Compile the asset sync.
+		const_cast<UStateTree*>(this)->CompileIfChanged();
+	}
 #endif
 
 	// Valid tree must have at least one state and valid instance data.
diff --git a/Engine/Plugins/Runtime/StateTree/Source/StateTreeModule/Public/StateTree.h b/Engine/Plugins/Runtime/StateTree/Source/StateTreeModule/Public/StateTree.h
--- a/Engine/Plugins/Runtime/StateTree/Source/StateTreeModule/Public/StateTree.h
+++ b/Engine/Plugins/Runtime/StateTree/Source/StateTreeModule/Public/StateTree.h
@@ -606,6 +606,9 @@ private:
 	 * After a compilation (success or failure), the dirty status returns to none.
 	 */
 	EDirtyStatus EditorDataDirtyStatus = EDirtyStatus::None;
+
+	/** The asset compilation was requested. The asset is in the pending queue. */
+	bool bCompilationPending = false;
 #endif
 
 	/** True if any global task is a transition task. */
