[StateTree] The StateTree's schema can modify the PropertyBag's schema.

#Synced-CL 51266215

diff --git a/Engine/Plugins/Runtime/StateTree/Source/StateTreeEditorModule/Private/Customizations/StateTreeEditorDataDetails.cpp b/Engine/Plugins/Runtime/StateTree/Source/StateTreeEditorModule/Private/Customizations/StateTreeEditorDataDetails.cpp
--- a/Engine/Plugins/Runtime/StateTree/Source/StateTreeEditorModule/Private/Customizations/StateTreeEditorDataDetails.cpp
+++ b/Engine/Plugins/Runtime/StateTree/Source/StateTreeEditorModule/Private/Customizations/StateTreeEditorDataDetails.cpp
@@ -15,6 +15,7 @@
 #include "StateTreeEditorUserSettings.h"
 #include "StateTreeGraphSchema.h"
 #include "StateTreeSchema.h"
+#include "StructUtilsMetadata.h"
 #include "Styling/StyleColors.h"
 #include "Widgets/Images/SImage.h"
 #include "Widgets/Layout/SBorder.h"
@@ -160,9 +161,17 @@ void FStateTreeEditorDataDetails::CustomizeDetails(IDetailLayoutBuilder& DetailB
 		ParametersCategory.SetColor(UE::StateTree::Colors::Blue);
 		{
 			// Show parameters as a category.
-			TSharedPtr<IPropertyHandle> PropertyBagParametersProperty = DetailBuilder.GetProperty(GET_MEMBER_NAME_CHECKED(UStateTreeEditorData, RootParameterPropertyBag)); // FInstancedPropertyBag
+			TSharedPtr<IPropertyHandle> PropertyBagParametersProperty = DetailBuilder.GetProperty(GET_MEMBER_NAME_CHECKED(UStateTreeEditorData, RootParameterPropertyBag));
 			PropertyBagParametersProperty->MarkHiddenByCustomization();
-			check(PropertyBagParametersProperty);
+
+			if (EditorSchema)
+			{
+				TOptional<FString> SchemaMetaData = EditorSchema->GetParametersSchemaClass();
+				if (SchemaMetaData.IsSet())
+				{
+					PropertyBagParametersProperty->SetInstanceMetaData(UE::StructUtils::Metadata::PropertyBagSchemaClassName, SchemaMetaData.GetValue());
+				}
+			}
 
 			const TSharedRef<SHorizontalBox> HeaderContentWidget = SNew(SHorizontalBox)
 				.IsEnabled(PropUtils, &IPropertyUtilities::IsPropertyEditingEnabled)
diff --git a/Engine/Plugins/Runtime/StateTree/Source/StateTreeEditorModule/Private/Customizations/StateTreeStateDetails.cpp b/Engine/Plugins/Runtime/StateTree/Source/StateTreeEditorModule/Private/Customizations/StateTreeStateDetails.cpp
--- a/Engine/Plugins/Runtime/StateTree/Source/StateTreeEditorModule/Private/Customizations/StateTreeStateDetails.cpp
+++ b/Engine/Plugins/Runtime/StateTree/Source/StateTreeEditorModule/Private/Customizations/StateTreeStateDetails.cpp
@@ -11,12 +11,14 @@
 #include "StateTreeEditor.h"
 #include "StateTreeEditorData.h"
 #include "StateTreeEditorNodeUtils.h"
+#include "StateTreeEditorSchema.h"
 #include "StateTreeEditorStyle.h"
 #include "StateTreeEditorDataExtension.h"
 #include "StateTreeEditorUserSettings.h"
 #include "StateTreePropertyHelpers.h"
 #include "StateTreeSchema.h"
 #include "StateTreeStateParametersDetails.h"
+#include "StructUtilsMetadata.h"
 #include "Debugger/StateTreeDebuggerUIExtensions.h"
 #include "SWarningOrErrorBox.h"
 #include "Widgets/SBoxPanel.h"
@@ -76,6 +78,7 @@ void FStateTreeStateDetails::CustomizeDetails(IDetailLayoutBuilder& DetailBuilde
 	WeakState = State;
 
 	const UStateTreeSchema* Schema = EditorData ? EditorData->Schema.Get() : nullptr;
+	const UStateTreeEditorSchema* EditorSchema = EditorData ? EditorData->EditorSchema.Get() : nullptr;
 	const FString SchemaPath = Schema ? Schema->GetClass()->GetPathName() : FString();
 	WeakEditorData = EditorData;
 
@@ -318,6 +321,16 @@ void FStateTreeStateDetails::CustomizeDetails(IDetailLayoutBuilder& DetailBuilde
 
 		if (!bFixedLayout)
 		{
+			if (EditorSchema)
+			{
+				TOptional<FString> SchemaMetaData = EditorSchema->GetParametersSchemaClass();
+				if (SchemaMetaData.IsSet())
+				{
+					using namespace UE::StructUtils::Metadata;
+					ParametersProperty->SetInstanceMetaData(PropertyBagSchemaClassName, SchemaMetaData.GetValue());
+				}
+			}
+
 			HeaderContentWidget->AddSlot()
 				.FillWidth(1.f)
 				.HAlign(HAlign_Right)
diff --git a/Engine/Plugins/Runtime/StateTree/Source/StateTreeEditorModule/Public/StateTreeEditorSchema.h b/Engine/Plugins/Runtime/StateTree/Source/StateTreeEditorModule/Public/StateTreeEditorSchema.h
--- a/Engine/Plugins/Runtime/StateTree/Source/StateTreeEditorModule/Public/StateTreeEditorSchema.h
+++ b/Engine/Plugins/Runtime/StateTree/Source/StateTreeEditorModule/Public/StateTreeEditorSchema.h
@@ -43,6 +43,17 @@ public:
 		return true;
 	}
 
+public:
+	/**
+	 * @returns the path name of a Property Bag Schema class used by the state tree schema.
+	 * It is used by the Global Parameters and the State Parameters.
+	 * The format is "/Script/MyModule.MySchema".
+	 */
+	virtual TOptional<FString> GetParametersSchemaClass() const
+	{
+		return TOptional<FString>();
+	}
+
 public:
 	/*
 	 * Validates and applies the schema restrictions on the StateTree.
