Fix state centric transitions not restoring with proper collapse state.
- Issue was UI behavior changes only existed on change, moved to on node LOD refresh.

#tests Editor
#Synced-CL 50972851

diff --git a/Engine/Plugins/Runtime/StateTree/Source/StateTreeEditorModule/Private/StateCentricView/View/SMainStateTransitionExtensionNode.cpp b/Engine/Plugins/Runtime/StateTree/Source/StateTreeEditorModule/Private/StateCentricView/View/SMainStateTransitionExtensionNode.cpp
--- a/Engine/Plugins/Runtime/StateTree/Source/StateTreeEditorModule/Private/StateCentricView/View/SMainStateTransitionExtensionNode.cpp
+++ b/Engine/Plugins/Runtime/StateTree/Source/StateTreeEditorModule/Private/StateCentricView/View/SMainStateTransitionExtensionNode.cpp
@@ -129,6 +129,19 @@ TSharedRef<SWidget> SMainStateTransitionExtensionNode::Construct_GenerateNodeHea
 
 		UpdateHideButtonAttributes();
 
+		if (TransitionDirection == EStateTransitionDirection::In)
+		{
+			StateCentricEditorData->InTransitionSplitterSizeRule = StateCentricEditorData->PerStateTreeUserSettings.GetInTransitionLOD() == EExtendableNodeLOD::Hidden
+				? SSplitter::ESizeRule::SizeToContent
+				: SSplitter::ESizeRule::FractionOfParent;
+		}
+		else if (ensure(TransitionDirection == EStateTransitionDirection::Out))
+		{
+			StateCentricEditorData->OutTransitionSplitterSizeRule = StateCentricEditorData->PerStateTreeUserSettings.GetOutTransitionLOD() == EExtendableNodeLOD::Hidden
+				? SSplitter::ESizeRule::SizeToContent
+				: SSplitter::ESizeRule::FractionOfParent;
+		}
+
 		SAssignNew(HideButton, SButton)
 			.OnClicked(this, &SMainStateTransitionExtensionNode::HandleOnCollapseClicked)
 			.ButtonStyle(&FAppStyle::GetWidgetStyle<FButtonStyle>("NoBorder"))
@@ -281,12 +294,10 @@ FReply SMainStateTransitionExtensionNode::HandleOnCollapseClicked()
 			if (bHidden)
 			{
 				StateCentricEditorData->PerStateTreeUserSettings.SetInTransitionLOD({});
-				StateCentricEditorData->InTransitionSplitterSizeRule = SSplitter::ESizeRule::FractionOfParent;
 			}
 			else
 			{
 				StateCentricEditorData->PerStateTreeUserSettings.SetInTransitionLOD(EExtendableNodeLOD::Hidden);
-				StateCentricEditorData->InTransitionSplitterSizeRule = SSplitter::ESizeRule::SizeToContent;
 			}
 		}
 		else if (ensure(TransitionDirection == EStateTransitionDirection::Out))
@@ -296,12 +307,10 @@ FReply SMainStateTransitionExtensionNode::HandleOnCollapseClicked()
 			if (bHidden)
 			{
 				StateCentricEditorData->PerStateTreeUserSettings.SetOutTransitionLOD({});
-				StateCentricEditorData->OutTransitionSplitterSizeRule = SSplitter::ESizeRule::FractionOfParent;
 			}
 			else
 			{
 				StateCentricEditorData->PerStateTreeUserSettings.SetOutTransitionLOD(EExtendableNodeLOD::Hidden);
-				StateCentricEditorData->OutTransitionSplitterSizeRule = SSplitter::ESizeRule::SizeToContent;
 			}
 		}
 	}
