Add ability to hide parent states in transitions

#tests editor

#Synced-CL 50972851

diff --git a/Engine/Plugins/Runtime/StateTree/Source/StateTreeEditorModule/Private/StateCentricView/StateCentricViewPerUserSettings.cpp b/Engine/Plugins/Runtime/StateTree/Source/StateTreeEditorModule/Private/StateCentricView/StateCentricViewPerUserSettings.cpp
--- a/Engine/Plugins/Runtime/StateTree/Source/StateTreeEditorModule/Private/StateCentricView/StateCentricViewPerUserSettings.cpp
+++ b/Engine/Plugins/Runtime/StateTree/Source/StateTreeEditorModule/Private/StateCentricView/StateCentricViewPerUserSettings.cpp
@@ -13,7 +13,8 @@ UStateCentricViewPerUserSettings::UStateCentricViewPerUserSettings()
 	DefaultPerUserSettings.InTransitionViewLOD = EExtendableNodeLOD::Minimal;
 	DefaultPerUserSettings.OutTransitionLOD = EExtendableNodeLOD::Full;
 	// @TODO: Feedback here was 2 in passing, but that seems super low.
-	DefaultPerUserSettings.NumParentsShown = 4; 
+	DefaultPerUserSettings.NumParentsShown = 2;
+	DefaultPerUserSettings.bHideParentTransitions = false;
 	PerSchemaSettings.Emplace(UStateTreeSchema::StaticClass(), DefaultPerUserSettings);
 }
 
diff --git a/Engine/Plugins/Runtime/StateTree/Source/StateTreeEditorModule/Private/StateCentricView/StateCentricViewPerUserSettings.h b/Engine/Plugins/Runtime/StateTree/Source/StateTreeEditorModule/Private/StateCentricView/StateCentricViewPerUserSettings.h
--- a/Engine/Plugins/Runtime/StateTree/Source/StateTreeEditorModule/Private/StateCentricView/StateCentricViewPerUserSettings.h
+++ b/Engine/Plugins/Runtime/StateTree/Source/StateTreeEditorModule/Private/StateCentricView/StateCentricViewPerUserSettings.h
@@ -13,6 +13,7 @@
 
 DECLARE_MULTICAST_DELEGATE_OneParam(FOnStateCentricViewLOD_Changed, TOptional<EExtendableNodeLOD> /*InLOD*/);
 DECLARE_MULTICAST_DELEGATE_OneParam(FOnStateCentricViewIntValue_Changed, TOptional<uint32> /*InValue*/);
+DECLARE_MULTICAST_DELEGATE_OneParam(FOnStateCentricViewBoolValue_Changed, TOptional<bool> /*InValue*/);
 
 /**
  * Experimental. Settings per user. This exists locally per user per schema. But also per state tree.
@@ -46,6 +47,10 @@ private:
 	UPROPERTY(EditAnywhere, Category="Layout", Config)
 	TOptional<uint32> NumParentsShown = {};
 
+	/** True if we should hide parents in transition view. */
+	UPROPERTY(EditAnywhere, Category = "Layout", Config)
+	TOptional<bool> bHideParentTransitions = {};
+
 public:
 
 	UE_API TOptional<EExtendableNodeLOD> GetParentHireachyLOD() const
@@ -68,6 +73,11 @@ public:
 		return NumParentsShown;
 	}
 
+	UE_API TOptional<bool> GetHideParentTransitions() const
+	{
+		return bHideParentTransitions;
+	}
+
 	UE_API void SetParentHireachyLOD(TOptional<EExtendableNodeLOD> InLOD)
 	{
 		ensure(!InLOD
@@ -110,6 +120,15 @@ public:
 		}
 	}
 
+	UE_API void SetHideParentTransitions(TOptional<bool> bInHideParentTransitions)
+	{
+		if (bHideParentTransitions != bInHideParentTransitions)
+		{
+			bHideParentTransitions = bInHideParentTransitions;
+			OnHideParentTransitionsChanged.Broadcast(bHideParentTransitions);
+		}
+	}
+
 public:
 
 	/** Callback whenever default parent hierarchy LOD changed. */
@@ -123,6 +142,9 @@ public:
 
 	/** Callback whenever num parents shown has changed. */
 	FOnStateCentricViewIntValue_Changed OnNumParentsShownChanged;
+
+	/** Callback whenever hide parent transitions has changed. */
+	FOnStateCentricViewBoolValue_Changed OnHideParentTransitionsChanged;
 };
 
 
diff --git a/Engine/Plugins/Runtime/StateTree/Source/StateTreeEditorModule/Private/StateCentricView/View/SStateCentricToolbar.cpp b/Engine/Plugins/Runtime/StateTree/Source/StateTreeEditorModule/Private/StateCentricView/View/SStateCentricToolbar.cpp
--- a/Engine/Plugins/Runtime/StateTree/Source/StateTreeEditorModule/Private/StateCentricView/View/SStateCentricToolbar.cpp
+++ b/Engine/Plugins/Runtime/StateTree/Source/StateTreeEditorModule/Private/StateCentricView/View/SStateCentricToolbar.cpp
@@ -71,6 +71,17 @@ void SStateCentricToolbar::Construct( const FArguments& InArgs, TSharedRef<FStat
 		}
 	};
 	
+	auto HandleOnHideParentTransitionsClicked = [WeakViewModel]()
+	{
+		if (TSharedPtr<FStateTreeViewModel> ViewModelPinned = WeakViewModel.Pin())
+		{
+			UStateTreeEditorData* EditorData = const_cast<UStateTreeEditorData*>(ViewModelPinned->GetStateTreeEditorData());
+			UStateCentricViewEditorDataExtension* StateCentricEditorData = EditorData->GetOrCreateExtension<UStateCentricViewEditorDataExtension>();
+
+			StateCentricEditorData->PerStateTreeUserSettings.SetHideParentTransitions(!StateCentricEditorData->PerStateTreeUserSettings.GetHideParentTransitions().Get(false));
+		}
+	};
+	
 	auto HandleOnClearPerStateExpansions = [WeakViewModel]()
 	{
 		if (TSharedPtr<FStateTreeViewModel> ViewModelPinned = WeakViewModel.Pin())
@@ -106,6 +117,13 @@ void SStateCentricToolbar::Construct( const FArguments& InArgs, TSharedRef<FStat
 		LOCTEXT("ExpandParentsTip", "Clears all per-state expansions, then expands parents."),
 		FSlateIcon());
 
+	ToolbarBuilder.AddToolBarButton(
+		FUIAction(FExecuteAction::CreateLambda(HandleOnHideParentTransitionsClicked)),
+		NAME_None,
+		LOCTEXT("HideParentTransitions", "Hide Parent Transitions"),
+		LOCTEXT("HideParentTransitionsTip", "Hides parents from transitions."),
+		FSlateIcon());
+
 	ToolbarBuilder.AddToolBarButton(
 		FUIAction(FExecuteAction::CreateLambda(HandleOnClearPerStateExpansions)),
 		NAME_None,
diff --git a/Engine/Plugins/Runtime/StateTree/Source/StateTreeEditorModule/Private/StateCentricView/View/STransitionNode.cpp b/Engine/Plugins/Runtime/StateTree/Source/StateTreeEditorModule/Private/StateCentricView/View/STransitionNode.cpp
--- a/Engine/Plugins/Runtime/StateTree/Source/StateTreeEditorModule/Private/StateCentricView/View/STransitionNode.cpp
+++ b/Engine/Plugins/Runtime/StateTree/Source/StateTreeEditorModule/Private/StateCentricView/View/STransitionNode.cpp
@@ -5,6 +5,7 @@
 #include "Customizations/StateTreeBindingExtension.h"
 #include "StateCentricView/StateCentricViewSettings.h"
 #include "StateCentricView/View/StateCentricViewUtils.h"
+#include "StateCentricView/ViewModel/StateCentricViewEditorDataExtension.h"
 #include "StateTreeEditorStyle.h"
 
 #include "Containers/Ticker.h"
@@ -39,6 +40,14 @@ STransitionNode::~STransitionNode()
 	if (ViewModel)
 	{
 		ViewModel->GetTransitionViewModel().GetOnTransitionModified().RemoveAll(this);
+
+		if (UStateTreeEditorData* EditorData = const_cast<UStateTreeEditorData*>(ViewModel->GetStateTreeEditorData()))
+		{
+			if (UStateCentricViewEditorDataExtension* StateCentricEditorData = EditorData->GetExtension<UStateCentricViewEditorDataExtension>())
+			{
+				StateCentricEditorData->PerStateTreeUserSettings.OnHideParentTransitionsChanged.RemoveAll(this);
+			}
+		}
 	}
 }
 
@@ -55,7 +64,22 @@ void STransitionNode::Construct(const SExtendableNode::FArguments& InArgs
 
 void STransitionNode::Construct_InitializeNodePersistentPreLOD()
 {
-	ViewModel->GetTransitionViewModel().GetOnTransitionModified().AddSP(this, &STransitionNode::HandleOnTransitionModified);
+	if (ViewState.IsValid())
+	{
+		UStateTreeEditorData* EditorData = const_cast<UStateTreeEditorData*>(ViewModel->GetStateTreeEditorData());
+		UStateCentricViewEditorDataExtension* StateCentricEditorData = EditorData->GetOrCreateExtension<UStateCentricViewEditorDataExtension>();
+
+		// Remove any prior bindings in case of issues with lifetime
+		{
+			ViewModel->GetTransitionViewModel().GetOnTransitionModified().RemoveAll(this);
+
+			StateCentricEditorData->PerStateTreeUserSettings.OnHideParentTransitionsChanged.RemoveAll(this);
+		}
+
+		ViewModel->GetTransitionViewModel().GetOnTransitionModified().AddSP(this, &STransitionNode::HandleOnTransitionModified);
+
+		StateCentricEditorData->PerStateTreeUserSettings.OnHideParentTransitionsChanged.AddSP(this, &STransitionNode::HandleOnHideParentTransitionsChanged);
+	}
 }
 
 void STransitionNode::Construct_GenerateBuiltInExtensions(const EExtendableNodeLOD InLOD)
@@ -66,11 +90,16 @@ void STransitionNode::Construct_GenerateBuiltInExtensions(const EExtendableNodeL
 	FStateTreeTransition* Transition = ViewModel->GetMutableTransitionByID(TransitionID);
 	if (Transition && ViewState.IsValid())
 	{
+		UStateTreeEditorData* EditorData = const_cast<UStateTreeEditorData*>(ViewModel->GetStateTreeEditorData());
+		UStateCentricViewEditorDataExtension* StateCentricEditorData = EditorData->GetOrCreateExtension<UStateCentricViewEditorDataExtension>();
+
 		// Parent state name as top extension
-		if (const UStateTreeState* ParentState = TransitionDirection == EStateTransitionDirection::In 
-			? ViewModel->GetSourceParentStateFromTransition(Transition)
-			: ViewModel->GetTargetParentStateFromTransition(Transition))
+		if (!StateCentricEditorData->PerStateTreeUserSettings.GetHideParentTransitions().Get(false))
 		{
+			const UStateTreeState* ParentState = TransitionDirection == EStateTransitionDirection::In
+				? ViewModel->GetSourceParentStateFromTransition(Transition)
+				: ViewModel->GetTargetParentStateFromTransition(Transition);
+
 			TransitionParentStateBrush = FStateCentricViewUtils::MakeStateSubnodeBackgroundOutlined(ViewModel.ToSharedRef(), ParentState);
 
 			FScopedMainNodeLayoutInfo ScopedMainLayoutInfo = SetScopedMainNodeLayoutInfo()
@@ -438,28 +467,38 @@ TSharedRef<SWidget> STransitionNode::Construct_GenerateNodeBody(const EExtendabl
 	return SNullWidget::NullWidget;
 }
 
-void STransitionNode::HandleOnTransitionModified(const UStateTreeState* OwningState, const FGuid& InTransitionID)
+void STransitionNode::RefreshNodeDelayed()
 {
-	if (TransitionID == InTransitionID)
+	TWeakPtr<STransitionNode> ThisWeak = StaticCastSharedRef<STransitionNode>(AsShared());
+	auto RefreshTransitionNode = [ThisWeak](float)
 	{
-		TWeakPtr<STransitionNode> ThisWeak = StaticCastSharedRef<STransitionNode>(AsShared());
-		auto RefreshTransitionNode = [ThisWeak](float)
+		if (TSharedPtr<STransitionNode> ThisPinned = ThisWeak.Pin())
 		{
-			if (TSharedPtr<STransitionNode> ThisPinned = ThisWeak.Pin())
-			{
-				ThisPinned->RegenerateWidgetForLOD(ThisPinned->CurrentLOD);
-			}
+			ThisPinned->RegenerateWidgetForLOD(ThisPinned->CurrentLOD);
+		}
+
+		// Run just once
+		return false;
+	};
 
-			// Run just once
-			return false;
-		};
+	// Refresh on a delay. We might have been modified from our own details panel
+	// If we instantly refresh we can kill that details mid-update which can remove property handles & cause update issues
+	FTSTicker::GetCoreTicker().AddTicker(FTickerDelegate::CreateLambda(RefreshTransitionNode));
+}
 
-		// Refresh on a delay. We might have been modified from our own details panel
-		// If we instantly refresh we can kill that details mid-update which can remove property handles & cause update issues
-		FTSTicker::GetCoreTicker().AddTicker(FTickerDelegate::CreateLambda(RefreshTransitionNode));
+void STransitionNode::HandleOnTransitionModified(const UStateTreeState* OwningState, const FGuid& InTransitionID)
+{
+	if (TransitionID == InTransitionID)
+	{
+		RefreshNodeDelayed();
 	}
 }
 
+void STransitionNode::HandleOnHideParentTransitionsChanged(TOptional<bool> /* Unused */)
+{
+	RefreshNodeDelayed();
+}
+
 FReply STransitionNode::HandleOnDeleteTransitionClicked()
 {
 	FStateTreeTransition* Transition = ViewModel->GetMutableTransitionByID(TransitionID);
diff --git a/Engine/Plugins/Runtime/StateTree/Source/StateTreeEditorModule/Private/StateCentricView/View/STransitionNode.h b/Engine/Plugins/Runtime/StateTree/Source/StateTreeEditorModule/Private/StateCentricView/View/STransitionNode.h
--- a/Engine/Plugins/Runtime/StateTree/Source/StateTreeEditorModule/Private/StateCentricView/View/STransitionNode.h
+++ b/Engine/Plugins/Runtime/StateTree/Source/StateTreeEditorModule/Private/StateCentricView/View/STransitionNode.h
@@ -47,9 +47,15 @@ protected:
 
 protected:
 
+	/** Method to refresh this node delayed */
+	void RefreshNodeDelayed();
+
 	/** Callback for transtion modified */
 	void HandleOnTransitionModified(const UStateTreeState* OwningState, const FGuid& InTransitionID);
 
+	/** Callback for hide parents modified */
+	void HandleOnHideParentTransitionsChanged(TOptional<bool> /* Unused */);
+
 	/** Callback for delete transition button */
 	FReply HandleOnDeleteTransitionClicked();
 
