Add missing null parent check

#tests editor

#Synced-CL 50972851

diff --git a/Engine/Plugins/Runtime/StateTree/Source/StateTreeEditorModule/Private/StateCentricView/View/STransitionNode.cpp b/Engine/Plugins/Runtime/StateTree/Source/StateTreeEditorModule/Private/StateCentricView/View/STransitionNode.cpp
--- a/Engine/Plugins/Runtime/StateTree/Source/StateTreeEditorModule/Private/StateCentricView/View/STransitionNode.cpp
+++ b/Engine/Plugins/Runtime/StateTree/Source/StateTreeEditorModule/Private/StateCentricView/View/STransitionNode.cpp
@@ -94,12 +94,12 @@ void STransitionNode::Construct_GenerateBuiltInExtensions(const EExtendableNodeL
 		UStateCentricViewEditorDataExtension* StateCentricEditorData = EditorData->GetOrCreateExtension<UStateCentricViewEditorDataExtension>();
 
 		// Parent state name as top extension
-		if (!StateCentricEditorData->PerStateTreeUserSettings.GetHideParentTransitions().Get(false))
-		{
-			const UStateTreeState* ParentState = TransitionDirection == EStateTransitionDirection::In
-				? ViewModel->GetSourceParentStateFromTransition(Transition)
-				: ViewModel->GetTargetParentStateFromTransition(Transition);
+		const UStateTreeState* ParentState = TransitionDirection == EStateTransitionDirection::In
+			? ViewModel->GetSourceParentStateFromTransition(Transition)
+			: ViewModel->GetTargetParentStateFromTransition(Transition);
 
+		if (ParentState && !StateCentricEditorData->PerStateTreeUserSettings.GetHideParentTransitions().Get(false))
+		{
 			TransitionParentStateBrush = FStateCentricViewUtils::MakeStateSubnodeBackgroundOutlined(ViewModel.ToSharedRef(), ParentState);
 
 			FScopedMainNodeLayoutInfo ScopedMainLayoutInfo = SetScopedMainNodeLayoutInfo()
