Wider transition bridge

#tests editor

#Synced-CL 50972851

diff --git a/Engine/Plugins/Runtime/StateTree/Source/StateTreeEditorModule/Private/StateCentricView/StateCentricViewSettings.h b/Engine/Plugins/Runtime/StateTree/Source/StateTreeEditorModule/Private/StateCentricView/StateCentricViewSettings.h
--- a/Engine/Plugins/Runtime/StateTree/Source/StateTreeEditorModule/Private/StateCentricView/StateCentricViewSettings.h
+++ b/Engine/Plugins/Runtime/StateTree/Source/StateTreeEditorModule/Private/StateCentricView/StateCentricViewSettings.h
@@ -61,6 +61,10 @@ public:
 	UPROPERTY(EditAnywhere, Category="Layout")
 	bool bEnableFillOnTransitionNodes = true;
 
+	/** Width for bridge connecting transitions nodes to center nodes*/
+	UPROPERTY(EditAnywhere, Category = "Layout")
+	float TransitionNodeBridgeWidth = 48.0f;
+
 	/** 
 	 * Max number of parents we allow users to show. At the moment we do not encourage changing this as showing too many parents 
 	 * can detract from the intended workflow of focusing on the central state and jumping around if needed.
diff --git a/Engine/Plugins/Runtime/StateTree/Source/StateTreeEditorModule/Private/StateCentricView/View/SFakeTransitionInfoNode.cpp b/Engine/Plugins/Runtime/StateTree/Source/StateTreeEditorModule/Private/StateCentricView/View/SFakeTransitionInfoNode.cpp
--- a/Engine/Plugins/Runtime/StateTree/Source/StateTreeEditorModule/Private/StateCentricView/View/SFakeTransitionInfoNode.cpp
+++ b/Engine/Plugins/Runtime/StateTree/Source/StateTreeEditorModule/Private/StateCentricView/View/SFakeTransitionInfoNode.cpp
@@ -87,6 +87,8 @@ void SFakeTransitionInfoNode::Construct_GenerateBuiltInExtensions(const EExtenda
 		}
 
 		// Small arrow used to indicate transition on left / right depending on direction
+		const float TransitionNodeBridgeWidth = UStateCentricViewSettings::GetSchemaSettingsDefaultFallback(ViewModel.Get()).TransitionNodeBridgeWidth;
+		const float TransitionNodeBridgeHalfWidth = TransitionNodeBridgeWidth / 2.0f;
 		{
 			FScopedMainNodeLayoutInfo ScopedMainLayoutInfo = SetScopedMainNodeLayoutInfo()
 				.SetFillSize(1.0f)
@@ -96,7 +98,7 @@ void SFakeTransitionInfoNode::Construct_GenerateBuiltInExtensions(const EExtenda
 				SNew(SButton)
 				.NormalPaddingOverride(FMargin(2.0, 4.0))
 				.PressedPaddingOverride(FMargin(2.0, 4.0))
-				.ContentPadding(FMargin(0.0f))
+				.ContentPadding(FMargin(TransitionNodeBridgeHalfWidth - 6.0f, 0.0f, TransitionNodeBridgeHalfWidth - 6.0f, 0))
 				[
 					SNew(SHorizontalBox)
 						+ SHorizontalBox::Slot()
@@ -137,8 +139,8 @@ void SFakeTransitionInfoNode::Construct_GenerateBuiltInExtensions(const EExtenda
 				.SetPadding(FMargin(4.0));
 
 			FMargin TransitionArrowPadding = GetTransitionDirection() == EStateTransitionDirection::In
-				? FMargin(0.0, 0.0, 16.0, 0.0)
-				: FMargin(16.0, 0.0, 0.0, 0.0);
+				? FMargin(0.0, 0.0, 4.0 + TransitionNodeBridgeWidth, 0.0)
+				: FMargin(4.0 + TransitionNodeBridgeWidth, 0.0, 0.0, 0.0);
 
 			const FSlateBrush* SubNodeBackgroundBrush = FStateCentricViewUtils::GetSubNodeBackground();
 
diff --git a/Engine/Plugins/Runtime/StateTree/Source/StateTreeEditorModule/Private/StateCentricView/View/STransitionNode.cpp b/Engine/Plugins/Runtime/StateTree/Source/StateTreeEditorModule/Private/StateCentricView/View/STransitionNode.cpp
--- a/Engine/Plugins/Runtime/StateTree/Source/StateTreeEditorModule/Private/StateCentricView/View/STransitionNode.cpp
+++ b/Engine/Plugins/Runtime/StateTree/Source/StateTreeEditorModule/Private/StateCentricView/View/STransitionNode.cpp
@@ -198,6 +198,8 @@ void STransitionNode::Construct_GenerateBuiltInExtensions(const EExtendableNodeL
 		}
 
 		// Small arrow used to indicate transition on left / right depending on direction
+		const float TransitionNodeBridgeWidth = UStateCentricViewSettings::GetSchemaSettingsDefaultFallback(ViewModel.Get()).TransitionNodeBridgeWidth;
+		const float TransitionNodeBridgeHalfWidth = TransitionNodeBridgeWidth / 2.0f;
 		{
 			FScopedMainNodeLayoutInfo ScopedMainLayoutInfo = SetScopedMainNodeLayoutInfo()
 				.SetFillSize(1.0f)
@@ -207,7 +209,7 @@ void STransitionNode::Construct_GenerateBuiltInExtensions(const EExtendableNodeL
 				SNew(SButton)
 				.NormalPaddingOverride(FMargin(2.0, 4.0))
 				.PressedPaddingOverride(FMargin(2.0, 4.0))
-				.ContentPadding(FMargin(0.0f))
+				.ContentPadding(FMargin(TransitionNodeBridgeHalfWidth - 6.0f, 0.0f, TransitionNodeBridgeHalfWidth - 6.0f, 0))
 				[
 					SNew(SHorizontalBox)
 						+ SHorizontalBox::Slot()
@@ -248,8 +250,8 @@ void STransitionNode::Construct_GenerateBuiltInExtensions(const EExtendableNodeL
 				.SetPadding(FMargin(4.0));
 
 			FMargin TransitionArrowPadding = TransitionDirection == EStateTransitionDirection::In
-				? FMargin(0.0, 0.0, 16.0, 0.0)
-				: FMargin(16.0, 0.0, 0.0, 0.0);
+				? FMargin(0.0, 0.0, 4.0 + TransitionNodeBridgeWidth, 0.0)
+				: FMargin(4.0 + TransitionNodeBridgeWidth, 0.0, 0.0, 0.0);
 
 			const FSlateBrush* SubNodeBackgroundBrush = FStateCentricViewUtils::GetSubNodeBackground();
 
