[StateTree] Add unittest to check if global task get the enter/exit when transitioning between different root from the same tree.

#rnx

#Synced-CL 51260989

diff --git a/Engine/Plugins/Runtime/StateTree/Source/StateTreeTestSuite/Private/StateTreeTransitionTest.cpp b/Engine/Plugins/Runtime/StateTree/Source/StateTreeTestSuite/Private/StateTreeTransitionTest.cpp
--- a/Engine/Plugins/Runtime/StateTree/Source/StateTreeTestSuite/Private/StateTreeTransitionTest.cpp
+++ b/Engine/Plugins/Runtime/StateTree/Source/StateTreeTestSuite/Private/StateTreeTransitionTest.cpp
@@ -2407,6 +2407,92 @@ struct FStateTreeTest_MultipleTransition_TemporaryFrame_GlobalTask : FStateTreeT
 };
 IMPLEMENT_STATE_TREE_INSTANT_TEST(FStateTreeTest_MultipleTransition_TemporaryFrame_GlobalTask, "System.StateTree.Transition.MultipleTransition.TemporaryFrame.GlobalTask");
 
+struct FStateTreeTest_Transition_DifferentRoot_GlobalTask : FStateTreeTestBase
+{
+	virtual bool InstantTest() override
+	{
+		/*
+		- Tree1: PrintTask
+			- Root1
+				- State1 Transition1 -> Goto Root2
+			- Root2
+				- State2
+		*/
+
+		UStateTree& StateTree1 = NewStateTree();
+
+		{
+			UStateTreeEditorData& Tree1EditorData = *Cast<UStateTreeEditorData>(StateTree1.EditorData);
+			UStateTreeState& Root1 = Tree1EditorData.AddSubTree(FName("Tree1Root1"));
+			UStateTreeState& State1 = Root1.AddChildState(FName("Tree1State1"));
+			UStateTreeState& Root2 = Tree1EditorData.AddSubTree(FName("Tree1Root2"));
+			UStateTreeState& State2 = Root2.AddChildState(FName("Tree1State2"));
+
+			FStateTreeTransition& LowPriorityTrans = State1.AddTransition(EStateTreeTransitionTrigger::OnTick, EStateTreeTransitionType::GotoState, &Root2);
+
+			Tree1EditorData.AddGlobalTask<FTestTask_StandNoTick>(FName("Tree1GlobalTask1"));
+			Tree1EditorData.AddGlobalTask<FTestTask_StandNoTick>(FName("Tree1GlobalTask2"));
+			Root1.AddTask<FTestTask_StandNoTick>(FName("Tree1Root1Task1"));
+			State1.AddTask<FTestTask_StandNoTick>(FName("Tree1State1Task1"));
+			Root2.AddTask<FTestTask_StandNoTick>(FName("Tree1Root2Task1"));
+			State2.AddTask<FTestTask_StandNoTick>(FName("Tree1State2Task1"));
+		}
+		{
+			FStateTreeCompilerLog Log;
+			FStateTreeCompiler Compiler(Log);
+			const bool bResult = Compiler.Compile(StateTree1);
+			AITEST_TRUE("StateTree should get compiled", bResult)
+		}
+
+		const TCHAR* EnterStateChangedStr = TEXT("EnterState=Changed");
+		const TCHAR* ExitStateChangedStr = TEXT("ExitState=Changed");
+		const TCHAR* ExitStateStr = TEXT("ExitState");
+
+		{
+			FStateTreeInstanceData InstanceData;
+			FTestStateTreeExecutionContext Exec(StateTree1, StateTree1, InstanceData);
+			const bool bInitSucceeded = Exec.IsValid();
+			AITEST_TRUE("StateTree should init", bInitSucceeded);
+
+			EStateTreeRunStatus Status = Exec.Start();
+			AITEST_TRUE(TEXT("Valid states"), Exec.ExpectInActiveStates("Tree1Root1", "Tree1State1"));
+			FTestStateTreeExecutionContext::FLogOrder LogOrder = Exec.Expect("Tree1GlobalTask1", EnterStateChangedStr);
+			AITEST_TRUE(TEXT("Tree1GlobalTask1 enter"), LogOrder);
+			LogOrder = LogOrder.Then("Tree1GlobalTask2", EnterStateChangedStr);
+			AITEST_TRUE(TEXT("Tree1GlobalTask2 enter"), LogOrder);
+			LogOrder = LogOrder.Then("Tree1Root1Task1", EnterStateChangedStr);
+			AITEST_TRUE(TEXT("Tree1Root1Task1 enter"), LogOrder);
+			LogOrder = LogOrder.Then("Tree1State1Task1", EnterStateChangedStr);
+			AITEST_TRUE(TEXT("Tree1State1Task1 enter"), LogOrder);
+			Exec.LogClear();
+
+			Exec.Tick(0.1f);
+			AITEST_TRUE(TEXT("Valid states"), Exec.ExpectInActiveStates("Tree1Root2", "Tree1State2"));
+			LogOrder = Exec.Expect("Tree1State1Task1", ExitStateChangedStr);
+			AITEST_TRUE(TEXT("Tree1State1Task1 exit"), LogOrder);
+			LogOrder = LogOrder.Then("Tree1Root1Task1", ExitStateChangedStr);
+			AITEST_TRUE(TEXT("Tree1Root1Task1 exit"), LogOrder);
+			LogOrder = LogOrder.Then("Tree1Root2Task1", EnterStateChangedStr);
+			AITEST_TRUE(TEXT("Tree1Root2Task1 enter"), LogOrder);
+			LogOrder = LogOrder.Then("Tree1State2Task1", EnterStateChangedStr);
+			AITEST_TRUE(TEXT("Tree1State2Task1 enter"), LogOrder);
+
+			AITEST_FALSE(TEXT("Tree1GlobalTask1 changed"), Exec.Expect("Tree1GlobalTask1", ExitStateStr));
+			AITEST_FALSE(TEXT("Tree1GlobalTask2 changed"), Exec.Expect("Tree1GlobalTask2", ExitStateStr));
+			Exec.LogClear();
+
+			Exec.Stop();
+			LogOrder = Exec.Expect("Tree1GlobalTask2", ExitStateChangedStr);
+			AITEST_TRUE(TEXT("Tree1GlobalTask2 exit"), LogOrder);
+			LogOrder = LogOrder.Then("Tree1GlobalTask1", ExitStateChangedStr);
+			AITEST_TRUE(TEXT("Tree1GlobalTask1 exit"), LogOrder);
+		}
+
+		return true;
+	}
+};
+IMPLEMENT_STATE_TREE_INSTANT_TEST(FStateTreeTest_Transition_DifferentRoot_GlobalTask, "System.StateTree.Transition.BetweenDifferentRoot.GlobalTasksShouldStay");
+
 struct FStateTreeTest_TransitionTypes : FStateTreeTestBase
 {
 	virtual bool InstantTest() override
