[State Tree] Events added during TriggerTransitions will be queued and only reset in the next TriggerTransitions. This allows multiple parallel trees to share event with each other, and fixes events sent in task on enter won't be accessible in the next tick in parallel trees.

- internal struct to scope new events marked as pending for next transition processing phase or not.

#Synced-CL 51266215

diff --git a/Engine/Plugins/Runtime/StateTree/Source/StateTreeModule/Private/StateTreeEvents.cpp b/Engine/Plugins/Runtime/StateTree/Source/StateTreeModule/Private/StateTreeEvents.cpp
--- a/Engine/Plugins/Runtime/StateTree/Source/StateTreeModule/Private/StateTreeEvents.cpp
+++ b/Engine/Plugins/Runtime/StateTree/Source/StateTreeModule/Private/StateTreeEvents.cpp
@@ -10,6 +10,17 @@
 // FStateTreeSharedEvent
 //----------------------------------------------------------------//
 
+UE::StateTree::Event::FEventsPendingForNextTransitionProcessingScope::FEventsPendingForNextTransitionProcessingScope(TNotNull<FStateTreeEventQueue*> EventQueue): ScopedEventQueue(EventQueue)
+{
+	bSavedIsPendingForNextTransitionProcessingScope = ScopedEventQueue->bNewlyAddedEventsShouldBePendingForNextTransitionProcessing;
+	ScopedEventQueue->bNewlyAddedEventsShouldBePendingForNextTransitionProcessing = true;
+}
+
+UE::StateTree::Event::FEventsPendingForNextTransitionProcessingScope::~FEventsPendingForNextTransitionProcessingScope()
+{
+	ScopedEventQueue->bNewlyAddedEventsShouldBePendingForNextTransitionProcessing = bSavedIsPendingForNextTransitionProcessingScope;
+}
+
 void FStateTreeSharedEvent::AddStructReferencedObjects(FReferenceCollector& Collector)
 {
 	if (IsValid())
@@ -37,7 +48,12 @@ bool FStateTreeEventQueue::SendEvent(const UObject* Owner, const FGameplayTag& T
 		return false;
 	}
 
-	SharedEvents.Emplace(Tag, Payload, Origin);
+	FStateTreeSharedEvent& SharedEvent = SharedEvents.Emplace_GetRef(Tag, Payload, Origin);
+	if (bNewlyAddedEventsShouldBePendingForNextTransitionProcessing)
+	{
+		SharedEvent->bIsPendingForNextTransitionProcessing = true;
+	}
+
 	return true;
 }
 
@@ -49,3 +65,17 @@ bool FStateTreeEventQueue::ConsumeEvent(const FStateTreeSharedEvent& Event)
 	});
 	return Count > 0;
 }
+
+void FStateTreeEventQueue::ClearEventsForCurrentTransitionProcessingPhase()
+{
+	SharedEvents.RemoveAllSwap([](FStateTreeSharedEvent& Event)
+	{
+		if (!Event->bIsPendingForNextTransitionProcessing)
+		{
+			return true;
+		}
+
+		Event->bIsPendingForNextTransitionProcessing = false;
+		return false;
+	});
+}
diff --git a/Engine/Plugins/Runtime/StateTree/Source/StateTreeModule/Private/StateTreeExecutionContext.cpp b/Engine/Plugins/Runtime/StateTree/Source/StateTreeModule/Private/StateTreeExecutionContext.cpp
--- a/Engine/Plugins/Runtime/StateTree/Source/StateTreeModule/Private/StateTreeExecutionContext.cpp
+++ b/Engine/Plugins/Runtime/StateTree/Source/StateTreeModule/Private/StateTreeExecutionContext.cpp
@@ -6017,7 +6017,7 @@ PRAGMA_ENABLE_DEPRECATION_WARNINGS
 	// All events have had the change to be reacted to, clear the event queue (if this instance owns it).
 	if (InstanceData.IsOwningEventQueue() && EventQueue)
 	{
-		EventQueue->Reset();
+		EventQueue->ClearEventsForCurrentTransitionProcessingPhase();
 	}
 
 	//
diff --git a/Engine/Plugins/Runtime/StateTree/Source/StateTreeModule/Private/Tasks/StateTreeRunParallelStateTreeTask.cpp b/Engine/Plugins/Runtime/StateTree/Source/StateTreeModule/Private/Tasks/StateTreeRunParallelStateTreeTask.cpp
--- a/Engine/Plugins/Runtime/StateTree/Source/StateTreeModule/Private/Tasks/StateTreeRunParallelStateTreeTask.cpp
+++ b/Engine/Plugins/Runtime/StateTree/Source/StateTreeModule/Private/Tasks/StateTreeRunParallelStateTreeTask.cpp
@@ -107,10 +107,18 @@ void FStateTreeRunParallelStateTreeTask::TriggerTransitions(FStateTreeExecutionC
 		return;
 	}
 
-	const EStateTreeRunStatus LastTreeRunStatus = InstanceData.TreeInstanceData.GetStorage().GetExecutionState().TreeRunStatus;
-	ParallelTreeContext.TickTriggerTransitions();
+	const EStateTreeRunStatus LastTreeRunStatus = InstanceData.TreeInstanceData.GetExecutionState()->TreeRunStatus;
 
-	const EStateTreeRunStatus NewTreeRunStatus = InstanceData.TreeInstanceData.GetStorage().GetExecutionState().TreeRunStatus;
+	{
+		// The parent tree owning the event queue clears events at end of transition processing phase
+		// Parallel Trees process state selection, exiting/entering states inside parent tree's transition processing phase
+		// Events need to live till the parent tree's next transition processing phase so events sent during task's enter state can be sustained correctly
+		using namespace UE::StateTree::Event;
+		FEventsPendingForNextTransitionProcessingScope PendingEventsScope(&InstanceData.TreeInstanceData.GetMutableEventQueue());
+		ParallelTreeContext.TickTriggerTransitions();
+	}
+
+	const EStateTreeRunStatus NewTreeRunStatus = InstanceData.TreeInstanceData.GetExecutionState()->TreeRunStatus;
 	if (LastTreeRunStatus != NewTreeRunStatus)
 	{
 		ensure(NewTreeRunStatus != EStateTreeRunStatus::Running);
diff --git a/Engine/Plugins/Runtime/StateTree/Source/StateTreeModule/Public/StateTreeEvents.h b/Engine/Plugins/Runtime/StateTree/Source/StateTreeModule/Public/StateTreeEvents.h
--- a/Engine/Plugins/Runtime/StateTree/Source/StateTreeModule/Public/StateTreeEvents.h
+++ b/Engine/Plugins/Runtime/StateTree/Source/StateTreeModule/Public/StateTreeEvents.h
@@ -9,6 +9,27 @@
 
 #define UE_API STATETREEMODULE_API
 
+struct FStateTreeEventQueue;
+
+namespace UE::StateTree::Event
+{
+	/**
+	* State Tree Events are cleared at end of each transition processing phase. (Inside TriggerTransitions()).
+	* When in scope, events sent will live through the current phase and cleared on the next transition processing phase.
+	* Used internally by TriggerTransitions and ParallelTree.
+	*/
+	struct FEventsPendingForNextTransitionProcessingScope
+	{
+		explicit FEventsPendingForNextTransitionProcessingScope(TNotNull<FStateTreeEventQueue*> EventQueue);
+
+		~FEventsPendingForNextTransitionProcessingScope();
+
+		TNotNull<FStateTreeEventQueue*> ScopedEventQueue;
+
+		uint8 bSavedIsPendingForNextTransitionProcessingScope : 1 = false;
+	};
+}
+
 /** Enum used for flow control during event iteration. */
 UENUM()
 enum class EStateTreeLoopEvents : uint8
@@ -66,6 +87,11 @@ struct FStateTreeEvent
 	/** Optional info to describe who sent the event. */
 	UPROPERTY(EditAnywhere, BlueprintReadWrite, Category = "Default")
 	FName Origin;
+
+private:
+	uint8 bIsPendingForNextTransitionProcessing : 1 = false;
+
+	friend struct FStateTreeEventQueue;
 };
 
 /**
@@ -172,6 +198,7 @@ struct FStateTreeEventQueue
 		SharedEvents.Reset();
 	}
 
+
 	/** @return true if the queue has any events. */
 	bool HasEvents() const
 	{
@@ -194,6 +221,12 @@ struct FStateTreeEventQueue
 	 */
 	UE_API bool ConsumeEvent(const FStateTreeSharedEvent& Event);
 
+	/**
+	 * Clear all events in the queue that are not marked pending for next transition processing phase,
+	 * And unmarks events that are pending for next transition processing since we are advancing to the next.
+	 */
+	UE_API void ClearEventsForCurrentTransitionProcessingPhase();
+	
 	/**
 	 * Iterates over all events.
 	 * @param Function a lambda which takes const FStateTreeSharedEvent& Event, and returns EStateTreeLoopEvents.
@@ -221,8 +254,10 @@ protected:
 
 	UPROPERTY()
 	TArray<FStateTreeSharedEvent> SharedEvents;
+	
+	uint8 bNewlyAddedEventsShouldBePendingForNextTransitionProcessing : 1 = false;
 
-	friend struct FStateTreeInstanceData;
+	friend UE::StateTree::Event::FEventsPendingForNextTransitionProcessingScope;
 };
 
 #undef UE_API
