[State Tree] set cvar CVarAutoRemoveSourcePathInvalidPropertyBindings default to false. Enable retaining invalid bindings.

#Synced-CL 51399697

diff --git a/Engine/Plugins/Runtime/StateTree/Source/StateTreeEditorModule/Private/StateTreeCompiler.cpp b/Engine/Plugins/Runtime/StateTree/Source/StateTreeEditorModule/Private/StateTreeCompiler.cpp
--- a/Engine/Plugins/Runtime/StateTree/Source/StateTreeEditorModule/Private/StateTreeCompiler.cpp
+++ b/Engine/Plugins/Runtime/StateTree/Source/StateTreeEditorModule/Private/StateTreeCompiler.cpp
@@ -65,7 +65,7 @@ namespace UE::StateTree::Compiler
 	// Version used by AppendToClassSchema. Bump this version number when a recompile of all assets is needed.
 	//Only bump when only the compiler changes.
 	//You do not need to bump the version when the memory layout of UStateTree has changed.
-	static constexpr int32 ClassSchemaVersion = 5;
+	static constexpr int32 ClassSchemaVersion = 6;
 
 	FAutoConsoleVariable CVarLogEnableParameterDelegateDispatcherBinding(
 		TEXT("StateTree.Compiler.EnableParameterDelegateDispatcherBinding"),
diff --git a/Engine/Plugins/Runtime/StateTree/Source/StateTreeEditorModule/Private/StateTreeEditorData.cpp b/Engine/Plugins/Runtime/StateTree/Source/StateTreeEditorModule/Private/StateTreeEditorData.cpp
--- a/Engine/Plugins/Runtime/StateTree/Source/StateTreeEditorModule/Private/StateTreeEditorData.cpp
+++ b/Engine/Plugins/Runtime/StateTree/Source/StateTreeEditorModule/Private/StateTreeEditorData.cpp
@@ -42,7 +42,7 @@ namespace UE::StateTree::Editor
 
 	FAutoConsoleVariable CVarAutoRemoveSourcePathInvalidPropertyBindings(
 		TEXT("StateTree.Editor.AutoRemoveSourcePathInvalidEditorBindings"),
-		true,
+		false,
 		TEXT("If set to true, EditorBindings whose source path does not resolve anymore,")
 		TEXT("will be automatically removed when the asset is loaded, when the users authors in the editor,")
 		TEXT("when state tree gets compiled.\n")
