[StateTree] Add unit test for success transition in Parallel task that is postpone by one frame. Fixed with CL 44413799
#rnx

#Synced-CL 51260989

diff --git a/Engine/Plugins/Runtime/StateTree/Source/StateTreeTestSuite/Private/StateTreeTest.cpp b/Engine/Plugins/Runtime/StateTree/Source/StateTreeTestSuite/Private/StateTreeTest.cpp
--- a/Engine/Plugins/Runtime/StateTree/Source/StateTreeTestSuite/Private/StateTreeTest.cpp
+++ b/Engine/Plugins/Runtime/StateTree/Source/StateTreeTestSuite/Private/StateTreeTest.cpp
@@ -2241,6 +2241,101 @@ struct FStateTreeTest_ParallelEventPriority_High : FStateTreeTest_ParallelEventP
 };
 IMPLEMENT_STATE_TREE_INSTANT_TEST(FStateTreeTest_ParallelEventPriority_High, "System.StateTree.ParallelEventPriority.High");
 
+// Test 2 parallel tree inside the same state with an All completion.
+struct FStateTreeTest_ParallelTask_TaskCompletion : FStateTreeTestBase
+{
+	virtual bool InstantTest() override
+	{
+		//Tree 1
+		//	Root (with 2 tasks that runs Tree 2)
+		//Tree 2/3
+		//	Root wait x frames to completes.
+
+		UStateTree& StateTree1 = NewStateTree();
+		UStateTree& StateTree2 = NewStateTree();
+		UStateTree& StateTree3 = NewStateTree();
+		auto Compile = [](TNotNull<UStateTree*> StateTree)
+			{
+				FStateTreeCompilerLog Log;
+				FStateTreeCompiler Compiler(Log);
+				return Compiler.Compile(StateTree);
+			};
+
+		auto BuildLinkedTree = [](TNotNull<UStateTree*> StateTree, FName TaskName, int32 TickToCompletion)
+			{
+				{
+					UStateTreeEditorData& EditorData = *Cast<UStateTreeEditorData>(StateTree->EditorData);
+					UStateTreeState& Root = EditorData.AddSubTree("Root");
+
+					TStateTreeEditorNode<FTestTask_Stand>& TaskA = Root.AddTask<FTestTask_Stand>(TaskName);
+					TaskA.GetNode().TicksToCompletion = TickToCompletion;
+
+					Root.AddTransition(EStateTreeTransitionTrigger::OnStateCompleted, EStateTreeTransitionType::Succeeded);
+				}
+			};
+
+		{
+			UStateTreeEditorData& EditorData1 = *Cast<UStateTreeEditorData>(StateTree1.EditorData);
+			UStateTreeState& Root1 = EditorData1.AddSubTree("Tree1StateRoot");
+
+			Root1.TasksCompletion = EStateTreeTaskCompletionType::All;
+			TStateTreeEditorNode<FStateTreeRunParallelStateTreeTask>& Task1 = Root1.AddTask<FStateTreeRunParallelStateTreeTask>();
+			Task1.GetInstanceData().StateTree.SetStateTree(&StateTree2);
+			TStateTreeEditorNode<FStateTreeRunParallelStateTreeTask>& Task2 = Root1.AddTask<FStateTreeRunParallelStateTreeTask>();
+			Task2.GetInstanceData().StateTree.SetStateTree(&StateTree3);
+			TStateTreeEditorNode<FTestTask_Stand>& Task3 = Root1.AddTask<FTestTask_Stand>(FName("Tree1Root1Task3"));
+			Task3.GetNode().TicksToCompletion = 0;
+			Root1.AddTransition(EStateTreeTransitionTrigger::OnStateCompleted, EStateTreeTransitionType::Succeeded);
+		}
+		BuildLinkedTree(&StateTree2, "Tree2Root1Task1", 1);
+		BuildLinkedTree(&StateTree3, "Tree3Root1Task1", 2);
+
+		AITEST_TRUE(TEXT("StateTree3 should get compiled"), Compile(&StateTree3));
+		AITEST_TRUE(TEXT("StateTree2 should get compiled"), Compile(&StateTree2));
+		AITEST_TRUE(TEXT("StateTree1 should get compiled"), Compile(&StateTree1));
+
+		{
+			FStateTreeInstanceData InstanceData;
+			FTestStateTreeExecutionContext Exec(StateTree1, StateTree1, InstanceData);
+			const bool bInitSucceeded = Exec.IsValid();
+			AITEST_TRUE(TEXT("StateTree should init"), bInitSucceeded);
+
+			{
+				const TCHAR* EnterStateStr(TEXT("EnterState"));
+				const TCHAR* ExitStateStr(TEXT("ExitState"));
+
+				EStateTreeRunStatus Status = Exec.Start();
+				AITEST_EQUAL(TEXT("Tick should complete with Running"), Status, EStateTreeRunStatus::Running);
+				AITEST_TRUE(TEXT("StateTree Active States"), Exec.ExpectInActiveStates("Tree1StateRoot"));
+				AITEST_TRUE(TEXT("Tree2Root1Task1 enter state"), Exec.Expect("Tree2Root1Task1", EnterStateStr));
+				AITEST_TRUE(TEXT("Tree3Root1Task1 enter state"), Exec.Expect("Tree3Root1Task1", EnterStateStr));
+				AITEST_TRUE(TEXT("Tree1Root1Task3 enter state"), Exec.Expect("Tree1Root1Task3", EnterStateStr));
+				Exec.LogClear();
+
+				Status = Exec.Tick(0.1f);
+				AITEST_EQUAL(TEXT("Tick should complete with Running"), Status, EStateTreeRunStatus::Running);
+				AITEST_TRUE(TEXT("StateTree Active States"), Exec.ExpectInActiveStates("Tree1StateRoot"));
+				AITEST_TRUE(TEXT("Tree2Root1Task1 exit state"), Exec.Expect("Tree2Root1Task1", ExitStateStr));
+				AITEST_FALSE(TEXT("Tree3Root1Task1 exit state"), Exec.Expect("Tree3Root1Task1", ExitStateStr));
+				AITEST_FALSE(TEXT("Tree1Root1Task3 exit state"), Exec.Expect("Tree1Root1Task3", ExitStateStr));
+				Exec.LogClear();
+
+				Status = Exec.Tick(2.0f);
+				AITEST_EQUAL(TEXT("Tick should complete with Succeeded"), Status, EStateTreeRunStatus::Succeeded);
+				AITEST_FALSE(TEXT("Tree2Root1Task1 exit state"), Exec.Expect("Tree2Root1Task1", ExitStateStr));
+				AITEST_TRUE(TEXT("Tree3Root1Task1 exit state"), Exec.Expect("Tree3Root1Task1", ExitStateStr));
+				AITEST_TRUE(TEXT("Tree1Root1Task3 exit state"), Exec.Expect("Tree1Root1Task3", ExitStateStr));
+				Exec.LogClear();
+
+				Exec.Stop();
+			}
+		}
+		return true;
+	}
+};
+
+IMPLEMENT_STATE_TREE_INSTANT_TEST(FStateTreeTest_ParallelTask_TaskCompletion, "System.StateTree.ParallelTask.Completion");
+
 struct FStateTreeTest_SubTreeTransition : FStateTreeTestBase
 {
 	virtual bool InstantTest() override
