[State Tree] StateTreeViewModel saves the selected states. Sync DetailsView to selected states when a new details view is created.

#Synced-CL 51445612

diff --git a/Engine/Plugins/Runtime/StateTree/Source/StateTreeEditorModule/Private/StateTreeEditor.cpp b/Engine/Plugins/Runtime/StateTree/Source/StateTreeEditorModule/Private/StateTreeEditor.cpp
--- a/Engine/Plugins/Runtime/StateTree/Source/StateTreeEditorModule/Private/StateTreeEditor.cpp
+++ b/Engine/Plugins/Runtime/StateTree/Source/StateTreeEditorModule/Private/StateTreeEditor.cpp
@@ -325,7 +325,12 @@ TSharedRef<SDockTab> FStateTreeEditor::SpawnTab_SelectionDetails(const FSpawnTab
 	DetailsViewArgs.NameAreaSettings = FDetailsViewArgs::HideNameArea;
 
 	SelectionDetailsView = PropertyEditorModule.CreateDetailView(DetailsViewArgs);
-	SelectionDetailsView->SetObject(nullptr);
+
+	TArray<UStateTreeState*> SelectedStates;
+	StateTreeViewModel->GetSelectedStates(SelectedStates);
+
+	const TArray<UObject*> SelectedObjects = TArray<UObject*>(SelectedStates);
+	SelectionDetailsView->SetObjects(SelectedObjects);
 
 	TSharedRef<SDockTab> SpawnedTab = SNew(SDockTab)
 		.Label(NSLOCTEXT("StateTreeEditor", "SelectionDetailsTab", "Details"))
