[State Tree] Unit test for parallel tree events handling.

Parallel trees should be able to share events with each other
Parallel trees global tasks and state tasks on enter should retain events correctly

#Synced-CL 51445612

diff --git a/Engine/Plugins/Runtime/StateTree/Source/StateTreeTestSuite/Private/StateTreeTest.cpp b/Engine/Plugins/Runtime/StateTree/Source/StateTreeTestSuite/Private/StateTreeTest.cpp
--- a/Engine/Plugins/Runtime/StateTree/Source/StateTreeTestSuite/Private/StateTreeTest.cpp
+++ b/Engine/Plugins/Runtime/StateTree/Source/StateTreeTestSuite/Private/StateTreeTest.cpp
@@ -2096,6 +2096,237 @@ struct FStateTreeTest_RecursiveParallelTask : FStateTreeTestBase
 
 IMPLEMENT_STATE_TREE_INSTANT_TEST(FStateTreeTest_RecursiveParallelTask, "System.StateTree.RecursiveParallelTask");
 
+struct FStateTreeTest_ParallelTreeSendEvents : FStateTreeTestBase
+{
+	virtual bool InstantTest() override
+	{
+		// MainTree
+		//  Root (Run ParallelTree1, Run ParallelTree2)
+		
+		// ParallelTree1
+		// PT1Root On Event 4 -> PT1Child4 Critical
+		//  PT1Child1 On Tick -> PT1Child2
+		//  PT1Child2 SendEvent 5 on Enter, On Event 5 -> PT1Child3
+		//  PT1Child3 -> Linked Asset
+		//  PT1Child4
+
+		// Linked Asset global : Send Event 1 on Enter
+		// LARoot
+		//  LAChild1 On Event 1 -> LAChild2
+		//  LAChild2 Send Event 2 on enter
+
+		// ParallelTree2
+		// PT2Root
+		//  PT2Child1 On Event 2 -> PT2Child 2
+		//  PT2Child2 SendEvent 4 on Enter 
+
+		auto MakeSendEventOnEnterTask = []<typename T, typename... TArgs>(TNotNull<T*> StateOrTree, const FGameplayTag EventTag, TArgs&&... InArgs)->TStateTreeEditorNode<FTestTask_PrintValue>&
+			requires TIsDerivedFrom<T, UStateTreeState>::IsDerived || TIsDerivedFrom<T, UStateTreeEditorData>::IsDerived 
+			{
+				TStateTreeEditorNode<FTestTask_PrintValue>* PrintValueTaskPtr;
+				
+				if constexpr (TIsDerivedFrom<T, UStateTreeState>::IsDerived)
+				{
+					PrintValueTaskPtr = &StateOrTree->template AddTask<FTestTask_PrintValue>(InArgs...);
+					
+				}
+				else if constexpr (TIsDerivedFrom<T, UStateTreeEditorData>::IsDerived)
+				{
+					PrintValueTaskPtr = &StateOrTree->template AddGlobalTask<FTestTask_PrintValue>(InArgs...);
+
+				}
+
+				PrintValueTaskPtr->GetNode().CustomEnterStateFunc = [EventTag](FStateTreeExecutionContext& Context, const FTestTask_PrintValue* PrintTask)
+					{
+						Context.SendEvent(EventTag);
+					};
+
+				return *PrintValueTaskPtr;
+			};
+
+		auto HasEvent = [](const FStateTreeEventQueue& EventQueue, const FGameplayTag EventTag)
+			{
+				TConstArrayView<FStateTreeSharedEvent> EventsView = EventQueue.GetEventsView();
+
+				for (const FStateTreeSharedEvent& Event : EventsView)
+				{
+					if (Event->Tag == EventTag)
+					{
+						return true;
+					}
+				}
+
+				return false;
+			};
+
+		// Linked Asset
+		UStateTree& LinkedAsset = NewStateTree();
+		{
+			UStateTreeEditorData& LAEditorData = *Cast<UStateTreeEditorData>(LinkedAsset.EditorData);
+
+			MakeSendEventOnEnterTask(TNotNull<UStateTreeEditorData*>(&LAEditorData), GetTestTag1(), "LAGlobalTask1");
+
+			UStateTreeState& LARoot = LAEditorData.AddSubTree("LARoot");
+
+			UStateTreeState& LAChild1 = LARoot.AddChildState("LAChild1");
+			TStateTreeEditorNode<FTestTask_PrintValue>& LAChild1Task1 = LAChild1.AddTask<FTestTask_PrintValue>("LAChild1Task1");
+			
+			UStateTreeState& LAChild2 = LARoot.AddChildState("LAChild2");
+			TStateTreeEditorNode<FTestTask_PrintValue>& LAChild2Task1 = MakeSendEventOnEnterTask(TNotNull<UStateTreeState*>(&LAChild2), GetTestTag2(), "LAChild2Task1");
+
+			LAChild1.AddTransition(EStateTreeTransitionTrigger::OnEvent, GetTestTag1(), EStateTreeTransitionType::GotoState, &LAChild2);
+
+			{
+				FStateTreeCompilerLog Log;
+				FStateTreeCompiler Compiler(Log);
+				const bool bResult = Compiler.Compile(LinkedAsset);
+				AITEST_TRUE("LinkedAsset should get compiled", bResult);
+			}
+		}
+
+		// Parallel Tree 1
+		UStateTree& ParallelTree1 = NewStateTree();
+		{
+			UStateTreeEditorData& PT1EditorData = *Cast<UStateTreeEditorData>(ParallelTree1.EditorData);
+
+			UStateTreeState& PT1Root = PT1EditorData.AddSubTree("PT1Root");
+
+			UStateTreeState& PT1Child1 = PT1Root.AddChildState("PT1Child1");
+			TStateTreeEditorNode<FTestTask_PrintValue>& PT1Child1Task1 = PT1Child1.AddTask<FTestTask_PrintValue>("PT1Child1Task1");
+
+			UStateTreeState& PT1Child2 = PT1Root.AddChildState("PT1Child2");
+			TStateTreeEditorNode<FTestTask_PrintValue>& PT1Child2Task1 = MakeSendEventOnEnterTask(TNotNull<UStateTreeState*>(&PT1Child2), GetTestTag5(), "PT1Child2Task1");
+			
+			UStateTreeState& PT1Child3 = PT1Root.AddChildState("PT1Child3", EStateTreeStateType::LinkedAsset);
+			PT1Child3.SetLinkedStateAsset(&LinkedAsset);
+			
+			UStateTreeState& PT1Child4 = PT1Root.AddChildState("PT1Child4");
+			TStateTreeEditorNode<FTestTask_PrintValue>& PT1Child4Task1 = PT1Child4.AddTask<FTestTask_PrintValue>("PT1Child4Task1");
+
+			{
+				FStateTreeTransition& PT1RootTrans = PT1Root.AddTransition(EStateTreeTransitionTrigger::OnEvent, GetTestTag4(), EStateTreeTransitionType::GotoState, &PT1Child4);
+				PT1RootTrans.Priority = EStateTreeTransitionPriority::Critical;
+
+				PT1Child1.AddTransition(EStateTreeTransitionTrigger::OnTick, EStateTreeTransitionType::GotoState, &PT1Child2);
+				PT1Child2.AddTransition(EStateTreeTransitionTrigger::OnEvent, GetTestTag5(), EStateTreeTransitionType::GotoState, &PT1Child3);
+			}
+
+			{
+				FStateTreeCompilerLog Log;
+				FStateTreeCompiler Compiler(Log);
+				const bool bResult = Compiler.Compile(ParallelTree1);
+				AITEST_TRUE("ParallelTree1 should get compiled", bResult);
+			}
+		}
+
+		// Parallel Tree2
+		UStateTree& ParallelTree2 = NewStateTree();
+		{
+			UStateTreeEditorData& PT2EditorData = *Cast<UStateTreeEditorData>(ParallelTree2.EditorData);
+
+			UStateTreeState& PT2Root = PT2EditorData.AddSubTree("PT2Root");
+
+			UStateTreeState& PT2Child1 = PT2Root.AddChildState("PT2Child1");
+			TStateTreeEditorNode<FTestTask_PrintValue>& PT2Child1Task1 = PT2Child1.AddTask<FTestTask_PrintValue>("PT2Child1Task1");
+
+			UStateTreeState& PT2Child2 = PT2Root.AddChildState("PT2Child2");
+			TStateTreeEditorNode<FTestTask_PrintValue>& PT2Child2Task1 = MakeSendEventOnEnterTask(TNotNull<UStateTreeState*>(&PT2Child2), GetTestTag4(), "PT2Child2Task1");
+
+			PT2Child1.AddTransition(EStateTreeTransitionTrigger::OnEvent, GetTestTag2(), EStateTreeTransitionType::GotoState, &PT2Child2);
+
+			{
+				FStateTreeCompilerLog Log;
+				FStateTreeCompiler Compiler(Log);
+				const bool bResult = Compiler.Compile(ParallelTree2);
+				AITEST_TRUE("ParallelTree2 should get compiled", bResult);
+			}
+		}
+
+		// Main Tree
+		UStateTree& MainTree = NewStateTree();
+		{
+			UStateTreeEditorData& MainEditorData = *Cast<UStateTreeEditorData>(MainTree.EditorData);
+
+			UStateTreeState& Root = MainEditorData.AddSubTree("Root");
+
+			TStateTreeEditorNode<FStateTreeRunParallelStateTreeTask>& PT1Task = Root.AddTask<FStateTreeRunParallelStateTreeTask>();
+			PT1Task.GetInstanceData().StateTree.SetStateTree(&ParallelTree1);
+
+			TStateTreeEditorNode<FStateTreeRunParallelStateTreeTask>& PT2Task = Root.AddTask<FStateTreeRunParallelStateTreeTask>();
+			PT2Task.GetInstanceData().StateTree.SetStateTree(&ParallelTree2);
+
+			{
+				FStateTreeCompilerLog Log;
+				FStateTreeCompiler Compiler(Log);
+				const bool bResult = Compiler.Compile(MainTree);
+				AITEST_TRUE("MainTree should get compiled", bResult);
+			}
+		}
+
+		constexpr TCHAR TickStr[] = TEXT("Tick0");
+		constexpr TCHAR EnterStateStr[] = TEXT("EnterState0");
+		constexpr TCHAR ExitStateStr[] = TEXT("ExitState0");
+
+		EStateTreeRunStatus Status = EStateTreeRunStatus::Unset;
+		FStateTreeInstanceData InstanceData;
+		const FStateTreeEventQueue& EventQueue = InstanceData.GetEventQueue();
+		FTestStateTreeExecutionContext Exec(MainTree, MainTree, InstanceData);
+		const bool bInitSucceeded = Exec.IsValid();
+		AITEST_TRUE(TEXT("StateTree should init"), bInitSucceeded);
+
+		{
+			Status = Exec.Start();
+			AITEST_EQUAL(TEXT("Start should complete with Running"), Status, EStateTreeRunStatus::Running);
+			AITEST_TRUE(TEXT("StateTree should enter PT2 Child1"), Exec.Expect("PT2Child1Task1", EnterStateStr));
+			AITEST_TRUE(TEXT("StateTree should enter PT1 Child1"), Exec.Expect("PT1Child1Task1", EnterStateStr));
+			AITEST_TRUE(TEXT("Event queue should be empty"), !EventQueue.HasEvents());
+			Exec.LogClear();
+		}
+
+		{
+			Status = Exec.Tick(0.1f);
+			AITEST_EQUAL(TEXT("Tick should complete with Running"), Status, EStateTreeRunStatus::Running);
+			AITEST_FALSE(TEXT("StateTree should not enter PT2 Child2"), Exec.Expect("PT2Child2Task1", EnterStateStr));
+			AITEST_TRUE(TEXT("StateTree should enter PT1 Child2"), Exec.Expect("PT1Child2Task1", EnterStateStr));
+			AITEST_TRUE(TEXT("Event queue should have correct events"), HasEvent(EventQueue, GetTestTag5()));
+			Exec.LogClear();
+		}
+
+		{
+			Status = Exec.Tick(0.1f);
+			AITEST_EQUAL(TEXT("Tick should complete with Running"), Status, EStateTreeRunStatus::Running);
+			AITEST_FALSE(TEXT("StateTree should not enter PT2 Child2"), Exec.Expect("PT2Child2Task1", EnterStateStr));
+			AITEST_TRUE(TEXT("StateTree should enter PT1 Child3, LA Child1"), Exec.Expect("LAGlobalTask1", EnterStateStr).Then("LAChild1Task1", EnterStateStr));
+			AITEST_TRUE(TEXT("Event queue should have correct events"), HasEvent(EventQueue, GetTestTag1()));
+			Exec.LogClear();
+		}
+
+		{
+			Status = Exec.Tick(0.1f);
+			AITEST_EQUAL(TEXT("Tick should complete with Running"), Status, EStateTreeRunStatus::Running);
+			AITEST_FALSE(TEXT("StateTree should not enter PT2 Child2"), Exec.Expect("PT2Child2Task1", EnterStateStr));
+			AITEST_TRUE(TEXT("StateTree should enter LA Child2"), Exec.Expect("LAChild2Task1", EnterStateStr));
+			AITEST_TRUE(TEXT("Event queue should have correct events"), HasEvent(EventQueue, GetTestTag2()));
+			Exec.LogClear();
+		}
+
+		{
+			Status = Exec.Tick(0.1f);
+			AITEST_EQUAL(TEXT("Tick should complete with Running"), Status, EStateTreeRunStatus::Running);
+			AITEST_TRUE(TEXT("StateTree should enter PT2 Child2"), Exec.Expect("PT2Child2Task1", EnterStateStr));
+			AITEST_TRUE(TEXT("StateTree should enter PT1 Child4"), Exec.Expect("PT1Child4Task1", EnterStateStr));
+			AITEST_TRUE(TEXT("Event queue should be empty"), !EventQueue.HasEvents());
+			Exec.LogClear();
+		}
+
+		Status = Exec.Stop();
+		Exec.LogClear();
+
+		return true;
+	}
+};
+IMPLEMENT_STATE_TREE_INSTANT_TEST(FStateTreeTest_ParallelTreeSendEvents, "System.StateTree.ParallelTreeSendEvents");
+
 // Test parallel tree event priority handling.
 struct FStateTreeTest_ParallelEventPriority : FStateTreeTestBase
 {
