[State Tree] Initialize local pointer to nullptr in StateTreeTest.cpp

Initialize PrintValueTaskPtr to nullptr at declaration to follow best practices and avoid potential undefined behavior if neither constexpr branch is taken.

#Synced-CL 51445612

diff --git a/Engine/Plugins/Runtime/StateTree/Source/StateTreeTestSuite/Private/StateTreeTest.cpp b/Engine/Plugins/Runtime/StateTree/Source/StateTreeTestSuite/Private/StateTreeTest.cpp
--- a/Engine/Plugins/Runtime/StateTree/Source/StateTreeTestSuite/Private/StateTreeTest.cpp
+++ b/Engine/Plugins/Runtime/StateTree/Source/StateTreeTestSuite/Private/StateTreeTest.cpp
@@ -2123,7 +2123,7 @@ struct FStateTreeTest_ParallelTreeSendEvents : FStateTreeTestBase
 		auto MakeSendEventOnEnterTask = []<typename T, typename... TArgs>(TNotNull<T*> StateOrTree, const FGameplayTag EventTag, TArgs&&... InArgs)->TStateTreeEditorNode<FTestTask_PrintValue>&
 			requires TIsDerivedFrom<T, UStateTreeState>::IsDerived || TIsDerivedFrom<T, UStateTreeEditorData>::IsDerived 
 			{
-				TStateTreeEditorNode<FTestTask_PrintValue>* PrintValueTaskPtr;
+				TStateTreeEditorNode<FTestTask_PrintValue>* PrintValueTaskPtr = nullptr;
 				
 				if constexpr (TIsDerivedFrom<T, UStateTreeState>::IsDerived)
 				{
