[StateTree] Add missing consideration to the editor node.

#Synced-CL 51494761

diff --git a/Engine/Plugins/Runtime/StateTree/Source/StateTreeEditorModule/Private/Customizations/StateTreeEditorNodeDetails.cpp b/Engine/Plugins/Runtime/StateTree/Source/StateTreeEditorModule/Private/Customizations/StateTreeEditorNodeDetails.cpp
--- a/Engine/Plugins/Runtime/StateTree/Source/StateTreeEditorModule/Private/Customizations/StateTreeEditorNodeDetails.cpp
+++ b/Engine/Plugins/Runtime/StateTree/Source/StateTreeEditorModule/Private/Customizations/StateTreeEditorNodeDetails.cpp
@@ -2010,7 +2010,8 @@ FText FStateTreeEditorNodeDetails::GetNodePickerTooltip() const
 		{
 			if (ScriptStruct->IsChildOf(FStateTreeBlueprintEvaluatorWrapper::StaticStruct())
 				|| ScriptStruct->IsChildOf(FStateTreeBlueprintTaskWrapper::StaticStruct())
-				|| ScriptStruct->IsChildOf(FStateTreeBlueprintConditionWrapper::StaticStruct()))
+				|| ScriptStruct->IsChildOf(FStateTreeBlueprintConditionWrapper::StaticStruct())
+				|| ScriptStruct->IsChildOf(FStateTreeBlueprintConsiderationWrapper::StaticStruct()))
 			{
 				if (Node->InstanceObject != nullptr
 					&& Node->InstanceObject->GetClass() != nullptr)
diff --git a/Engine/Plugins/Runtime/StateTree/Source/StateTreeEditorModule/Private/StateTreeEditorNode.cpp b/Engine/Plugins/Runtime/StateTree/Source/StateTreeEditorModule/Private/StateTreeEditorNode.cpp
--- a/Engine/Plugins/Runtime/StateTree/Source/StateTreeEditorModule/Private/StateTreeEditorNode.cpp
+++ b/Engine/Plugins/Runtime/StateTree/Source/StateTreeEditorModule/Private/StateTreeEditorNode.cpp
@@ -3,6 +3,7 @@
 #include "StateTreeEditorNode.h"
 
 #include "Blueprint/StateTreeConditionBlueprintBase.h"
+#include "Blueprint/StateTreeConsiderationBlueprintBase.h"
 #include "Blueprint/StateTreeEvaluatorBlueprintBase.h"
 #include "Blueprint/StateTreeTaskBlueprintBase.h"
 
@@ -23,7 +24,8 @@ FName FStateTreeEditorNode::GetName() const
 			if (InstanceObject &&
 					(NodeType->IsChildOf(TBaseStructure<FStateTreeBlueprintTaskWrapper>::Get())
 					|| NodeType->IsChildOf(TBaseStructure<FStateTreeBlueprintEvaluatorWrapper>::Get())
-					|| NodeType->IsChildOf(TBaseStructure<FStateTreeBlueprintConditionWrapper>::Get())))
+					|| NodeType->IsChildOf(TBaseStructure<FStateTreeBlueprintConditionWrapper>::Get())
+					|| NodeType->IsChildOf(TBaseStructure<FStateTreeBlueprintConsiderationWrapper>::Get())))
 			{
 				return FName(InstanceObject->GetClass()->GetDisplayNameText().ToString());
 			}
