[StateTree] Add a menu context to retreive the StateTree asset before the menu is built.

#Synced-CL 51494761

diff --git a/Engine/Plugins/Runtime/StateTree/Source/StateTreeEditorModule/Private/StateTreeEditor.cpp b/Engine/Plugins/Runtime/StateTree/Source/StateTreeEditorModule/Private/StateTreeEditor.cpp
--- a/Engine/Plugins/Runtime/StateTree/Source/StateTreeEditorModule/Private/StateTreeEditor.cpp
+++ b/Engine/Plugins/Runtime/StateTree/Source/StateTreeEditorModule/Private/StateTreeEditor.cpp
@@ -269,6 +269,15 @@ void FStateTreeEditor::OnToolkitHostingStarted(const TSharedRef<IToolkit>& Toolk
 	HostedToolkit = Toolkit;
 }
 
+void FStateTreeEditor::InitToolMenuContext(FToolMenuContext& MenuContext)
+{
+	IStateTreeEditor::InitToolMenuContext(MenuContext);
+
+	UStateTreeEditorContext* EditorContext = NewObject<UStateTreeEditorContext>();
+	EditorContext->EditorHostInterface = EditorHost;
+	MenuContext.AddObject(EditorContext);
+}
+
 void FStateTreeEditor::OnToolkitHostingFinished(const TSharedRef<IToolkit>& Toolkit)
 {
 	ModeUILayer->OnToolkitHostingFinished(Toolkit);
diff --git a/Engine/Plugins/Runtime/StateTree/Source/StateTreeEditorModule/Private/StateTreeEditor.h b/Engine/Plugins/Runtime/StateTree/Source/StateTreeEditorModule/Private/StateTreeEditor.h
--- a/Engine/Plugins/Runtime/StateTree/Source/StateTreeEditorModule/Private/StateTreeEditor.h
+++ b/Engine/Plugins/Runtime/StateTree/Source/StateTreeEditorModule/Private/StateTreeEditor.h
@@ -42,7 +42,8 @@ public:
 	//~ Begin FAssetEditorToolkit Interface
 	virtual void PostInitAssetEditor() override;
 	virtual void OnToolkitHostingFinished(const TSharedRef<IToolkit>& Toolkit) override;
-	virtual void OnToolkitHostingStarted(const TSharedRef<IToolkit>& Toolkit) override;	
+	virtual void OnToolkitHostingStarted(const TSharedRef<IToolkit>& Toolkit) override;
+	virtual void InitToolMenuContext(FToolMenuContext& MenuContext) override;
 	//~ End FAssetEditorToolkit Interface
 
 	//~ Begin IToolkit Interface
diff --git a/Engine/Plugins/Runtime/StateTree/Source/StateTreeEditorModule/Private/StateTreeEditorModeToolkit.cpp b/Engine/Plugins/Runtime/StateTree/Source/StateTreeEditorModule/Private/StateTreeEditorModeToolkit.cpp
--- a/Engine/Plugins/Runtime/StateTree/Source/StateTreeEditorModule/Private/StateTreeEditorModeToolkit.cpp
+++ b/Engine/Plugins/Runtime/StateTree/Source/StateTreeEditorModule/Private/StateTreeEditorModeToolkit.cpp
@@ -447,6 +447,11 @@ void FStateTreeEditorModeToolkit::ExtendSecondaryModeToolbar(UToolMenu* ToolBar)
 	}
 }
 
+UStateTree* FStateTreeEditorModeToolkit::GetStateTree() const
+{
+	return EditorHost ? EditorHost->GetStateTree() : nullptr;
+}
+
 void FStateTreeEditorModeToolkit::OnStateTreeChanged()
 {
 	// Update underlying state tree state
diff --git a/Engine/Plugins/Runtime/StateTree/Source/StateTreeEditorModule/Public/StateTreeEditorModeToolkit.h b/Engine/Plugins/Runtime/StateTree/Source/StateTreeEditorModule/Public/StateTreeEditorModeToolkit.h
--- a/Engine/Plugins/Runtime/StateTree/Source/StateTreeEditorModule/Public/StateTreeEditorModeToolkit.h
+++ b/Engine/Plugins/Runtime/StateTree/Source/StateTreeEditorModule/Public/StateTreeEditorModeToolkit.h
@@ -8,6 +8,7 @@
 
 #define UE_API STATETREEEDITORMODULE_API
 
+class UStateTree;
 class UStateTreeEditorMode;
 struct FPropertyBindingBindingCollection;
 namespace UE::StateTreeEditor
@@ -30,7 +31,8 @@ public:
 	UE_API virtual void RequestModeUITabs() override;
 	
 	UE_API virtual void ExtendSecondaryModeToolbar(UToolMenu* InModeToolbarMenu) override;
-		
+
+	UE_API UStateTree* GetStateTree() const;
 	UE_API void OnStateTreeChanged();
 
 protected:
