[State Tree] Fix crash in Enum Compare Condition when resetting right field default value

Fix wrong enum pointer checked in GetDescriptionForOperand ??? was checking Left.Enum instead of Right.Enum for the right operand, causing a crash when the right field's default value is reset.

#Synced-CL 51445612
#Synced-CL 51445612

diff --git a/Engine/Plugins/Runtime/StateTree/Source/StateTreeModule/Private/Conditions/StateTreeCommonConditions.cpp b/Engine/Plugins/Runtime/StateTree/Source/StateTreeModule/Private/Conditions/StateTreeCommonConditions.cpp
--- a/Engine/Plugins/Runtime/StateTree/Source/StateTreeModule/Private/Conditions/StateTreeCommonConditions.cpp
+++ b/Engine/Plugins/Runtime/StateTree/Source/StateTreeModule/Private/Conditions/StateTreeCommonConditions.cpp
@@ -238,7 +238,7 @@ FText FStateTreeCompareEnumCondition::GetDescription(const FGuid& ID, FStateTree
 	FText RightValue = BindingLookup.GetBindingSourceDisplayName(FPropertyBindingPath(ID, GET_MEMBER_NAME_CHECKED(FInstanceDataType, Right)), Formatting);
 	if (RightValue.IsEmpty())
 	{
-		if (InstanceData->Left.Enum)
+		if (InstanceData->Right.Enum)
 		{
 			RightValue = InstanceData->Right.Enum->GetDisplayNameTextByValue(InstanceData->Right.Value);
 		}
