[State Tree] Remap state event IDs when cut/paste states

Preserve and remap state event IDs during copy/paste to fix compilation failures when a pasted state has nodes that bind to required events on the state's enter conditions.

#Synced-CL 51445612
#Synced-CL 51445612

diff --git a/Engine/Plugins/Runtime/StateTree/Source/StateTreeEditorModule/Private/StateTreeEditorData.cpp b/Engine/Plugins/Runtime/StateTree/Source/StateTreeEditorModule/Private/StateTreeEditorData.cpp
--- a/Engine/Plugins/Runtime/StateTree/Source/StateTreeEditorModule/Private/StateTreeEditorData.cpp
+++ b/Engine/Plugins/Runtime/StateTree/Source/StateTreeEditorModule/Private/StateTreeEditorData.cpp
@@ -1089,6 +1089,7 @@ void UStateTreeEditorData::DuplicateIDs()
 				if (State)
 				{
 					AddId(State->ID, false);
+					AddId(State->GetEventID(), false);
 				}
 
 				// Schema should not be included in OldToNewIDs
diff --git a/Engine/Plugins/Runtime/StateTree/Source/StateTreeEditorModule/Private/StateTreeViewModel.cpp b/Engine/Plugins/Runtime/StateTree/Source/StateTreeEditorModule/Private/StateTreeViewModel.cpp
--- a/Engine/Plugins/Runtime/StateTree/Source/StateTreeEditorModule/Private/StateTreeViewModel.cpp
+++ b/Engine/Plugins/Runtime/StateTree/Source/StateTreeEditorModule/Private/StateTreeViewModel.cpp
@@ -172,10 +172,14 @@ namespace UE::StateTree::Editor
 	{
 		State->Modify();
 
+		const FGuid OldStateEventID = State->GetEventID();
+
 		const FGuid NewStateID = FGuid::NewGuid();
 		IDsMap.Emplace(State->ID, NewStateID);
 		State->ID = NewStateID;
 		
+		IDsMap.Emplace(OldStateEventID, State->GetEventID());
+
 		const FGuid NewParametersID = FGuid::NewGuid();
 		IDsMap.Emplace(State->Parameters.ID, NewParametersID);
 		State->Parameters.ID = NewParametersID;
