[StateTree] Fix mutlithreaded access when a StateTree asset loads in the AsyncLoadingThread and another StateTree compiles on the main thread.

#Synced-CL 51494712

diff --git a/Engine/Plugins/Runtime/StateTree/Source/StateTreeEditorModule/Private/StateTreeState.cpp b/Engine/Plugins/Runtime/StateTree/Source/StateTreeEditorModule/Private/StateTreeState.cpp
--- a/Engine/Plugins/Runtime/StateTree/Source/StateTreeEditorModule/Private/StateTreeState.cpp
+++ b/Engine/Plugins/Runtime/StateTree/Source/StateTreeEditorModule/Private/StateTreeState.cpp
@@ -81,14 +81,19 @@ UStateTreeState::UStateTreeState(const FObjectInitializer& ObjectInitializer)
 
 UStateTreeState::~UStateTreeState()
 {
-	UE::StateTree::Delegates::OnPostCompile.RemoveAll(this);
+	UE::StateTree::Delegates::OnPostCompile.Remove(PostCompileHandle);
 }
 
 void UStateTreeState::PostInitProperties()
 {
 	Super::PostInitProperties();
-	
-	UE::StateTree::Delegates::OnPostCompile.AddUObject(this, &UStateTreeState::OnTreeCompiled);
+
+	// Async load should not be registered; it should only be registered in PostLoad.
+	//A newly created object in the editor should register now because it’s not loaded, and PostLoad won’t get called.
+	if (!HasAnyFlags(RF_NeedPostLoad | RF_ClassDefaultObject))
+	{
+		PostCompileHandle = UE::StateTree::Delegates::OnPostCompile.AddUObject(this, &UStateTreeState::OnTreeCompiled);
+	}
 }
 
 void UStateTreeState::OnTreeCompiled(const UStateTree& StateTree)
@@ -352,6 +357,11 @@ void UStateTreeState::PostLoad()
 		SetFlags(RF_Transactional);
 	}
 
+	if (!PostCompileHandle.IsValid() && !HasAnyFlags(RF_ClassDefaultObject))
+	{
+		PostCompileHandle = UE::StateTree::Delegates::OnPostCompile.AddUObject(this, &UStateTreeState::OnTreeCompiled);
+	}
+
 #if WITH_EDITORONLY_DATA
 PRAGMA_DISABLE_DEPRECATION_WARNINGS
 	const int32 CurrentStateTreeCustomVersion = GetLinkerCustomVersion(FStateTreeCustomVersion::GUID);
diff --git a/Engine/Plugins/Runtime/StateTree/Source/StateTreeEditorModule/Public/StateTreeState.h b/Engine/Plugins/Runtime/StateTree/Source/StateTreeEditorModule/Public/StateTreeState.h
--- a/Engine/Plugins/Runtime/StateTree/Source/StateTreeEditorModule/Public/StateTreeState.h
+++ b/Engine/Plugins/Runtime/StateTree/Source/StateTreeEditorModule/Public/StateTreeState.h
@@ -491,6 +491,8 @@ private:
 
 	/** Transaction temporary. Transitions we captured in PreEdit for remove / move comparision */
 	TArray<FStateTreeTransition> PreEditTransitions;
+	/** Handle for the UE::StateTree::Delegates::OnPostCompile delegate */
+	FDelegateHandle PostCompileHandle;
 };
 
 #undef UE_API
