[State Tree] Allow creating new StateTree assets from linked asset dropdown

- Add AllowCreate(true) to the linked asset picker
- Create an override UStateTreeFactory with schema class pre-set (skips schema picker dialog)
- Skip ConfigureProperties when schema class is already specified
- Change StateTreeSchemaClass UPROPERTY from Transient to serialized so it survives duplication
#Synced-CL 51445612

diff --git a/Engine/Plugins/Runtime/StateTree/Source/StateTreeEditorModule/Private/Customizations/StateTreeStateDetails.cpp b/Engine/Plugins/Runtime/StateTree/Source/StateTreeEditorModule/Private/Customizations/StateTreeStateDetails.cpp
--- a/Engine/Plugins/Runtime/StateTree/Source/StateTreeEditorModule/Private/Customizations/StateTreeStateDetails.cpp
+++ b/Engine/Plugins/Runtime/StateTree/Source/StateTreeEditorModule/Private/Customizations/StateTreeStateDetails.cpp
@@ -243,6 +243,12 @@ void FStateTreeStateDetails::CustomizeDetails(IDetailLayoutBuilder& DetailBuilde
 	}
 	else if (StateType == EStateTreeStateType::LinkedAsset)
 	{
+		if (!OverrideFactoryForLinkedAsset)
+		{
+			OverrideFactoryForLinkedAsset = TStrongObjectPtr<UStateTreeFactory>(NewObject<UStateTreeFactory>());
+			OverrideFactoryForLinkedAsset->SetSchemaClass(Schema->GetClass());
+		}
+
 		// Custom widget for the linked asset, to add filtering to the assets.
 		IDetailPropertyRow& Row = StateCategory.AddProperty(LinkedAssetProperty);
 		Row.CustomWidget()
@@ -255,6 +261,8 @@ void FStateTreeStateDetails::CustomizeDetails(IDetailLayoutBuilder& DetailBuilde
 			SNew(SObjectPropertyEntryBox)
 			.PropertyHandle(LinkedAssetProperty)
 			.AllowedClass(UStateTree::StaticClass())
+			.AllowCreate(true)
+			.NewAssetFactories(TArray<UFactory*>({ OverrideFactoryForLinkedAsset.Get() }))
 			.ThumbnailPool(DetailBuilder.GetPropertyUtilities()->GetThumbnailPool())
 			.OnShouldFilterAsset_Lambda([SchemaPath](const FAssetData& InAssetData)
 			{
diff --git a/Engine/Plugins/Runtime/StateTree/Source/StateTreeEditorModule/Private/Customizations/StateTreeStateDetails.h b/Engine/Plugins/Runtime/StateTree/Source/StateTreeEditorModule/Private/Customizations/StateTreeStateDetails.h
--- a/Engine/Plugins/Runtime/StateTree/Source/StateTreeEditorModule/Private/Customizations/StateTreeStateDetails.h
+++ b/Engine/Plugins/Runtime/StateTree/Source/StateTreeEditorModule/Private/Customizations/StateTreeStateDetails.h
@@ -5,6 +5,7 @@
 #include "IDetailCustomization.h"
 #include "EditorUndoClient.h"
 #include "StateTreeEditorData.h"
+#include "StateTreeFactory.h"
 #include "StateTreeState.h"
 
 class IDetailLayoutBuilder;
@@ -28,7 +29,6 @@ public:
 	//~ IDetailCustomization interface
 	virtual void PendingDelete() override;
 	virtual void CustomizeDetails(IDetailLayoutBuilder& DetailBuilder) override;
-
 protected:
 
 	//~ FEditorUndoClient
@@ -43,5 +43,7 @@ protected:
 	TWeakObjectPtr<UStateTreeEditorData> WeakEditorData;
 	TWeakObjectPtr<UStateTreeState> WeakState;
 
+	TStrongObjectPtr<UStateTreeFactory> OverrideFactoryForLinkedAsset;
+
 	bool bIsPendingDelete = false;
 };
diff --git a/Engine/Plugins/Runtime/StateTree/Source/StateTreeEditorModule/Private/StateTreeFactory.cpp b/Engine/Plugins/Runtime/StateTree/Source/StateTreeEditorModule/Private/StateTreeFactory.cpp
--- a/Engine/Plugins/Runtime/StateTree/Source/StateTreeEditorModule/Private/StateTreeFactory.cpp
+++ b/Engine/Plugins/Runtime/StateTree/Source/StateTreeEditorModule/Private/StateTreeFactory.cpp
@@ -58,6 +58,12 @@ UStateTreeFactory::UStateTreeFactory(const FObjectInitializer& ObjectInitializer
 
 bool UStateTreeFactory::ConfigureProperties()
 {
+	// If it is an overridden factory and the schema class has been specified by caller
+	if (StateTreeSchemaClass)
+	{
+		return true;
+	}
+
 	// Load the classviewer module to display a class picker
 	FClassViewerModule& ClassViewerModule = FModuleManager::LoadModuleChecked<FClassViewerModule>("ClassViewer");
 
diff --git a/Engine/Plugins/Runtime/StateTree/Source/StateTreeEditorModule/Public/StateTreeFactory.h b/Engine/Plugins/Runtime/StateTree/Source/StateTreeEditorModule/Public/StateTreeFactory.h
--- a/Engine/Plugins/Runtime/StateTree/Source/StateTreeEditorModule/Public/StateTreeFactory.h
+++ b/Engine/Plugins/Runtime/StateTree/Source/StateTreeEditorModule/Public/StateTreeFactory.h
@@ -27,7 +27,7 @@ class UStateTreeFactory : public UFactory
 
 protected:
 	
-	UPROPERTY(Transient)
+	UPROPERTY()
 	TObjectPtr<UClass> StateTreeSchemaClass;
 };
 
