[StateTree] StateTreeCompiler optimization.
Preallocate some containers.
Pre calculate the frame for each states.
Pre calculate the state path.
Map the binding per ID instead of going through them all.

#Synced-CL 51494761

diff --git a/Engine/Plugins/Runtime/StateTree/Source/StateTreeEditorModule/Private/StateTreeCompiler.cpp b/Engine/Plugins/Runtime/StateTree/Source/StateTreeEditorModule/Private/StateTreeCompiler.cpp
--- a/Engine/Plugins/Runtime/StateTree/Source/StateTreeEditorModule/Private/StateTreeCompiler.cpp
+++ b/Engine/Plugins/Runtime/StateTree/Source/StateTreeEditorModule/Private/StateTreeCompiler.cpp
@@ -495,6 +495,11 @@ bool FStateTreeCompiler::PreCompile(TNotNull<UStateTree*> InStateTree)
 	}
 
 	EditorData->GetAllStructValues(IDToStructValue);
+	for (const FPropertyBindingBinding& Binding : EditorData->EditorBindings.GetBindings())
+	{
+		BindingsByTargetStructID.Add(Binding.GetTargetPath().GetStructID(), &Binding);
+	}
+	EditorBindingsNum = EditorData->EditorBindings.GetBindings().Num();
 
 	return true;
 }
@@ -628,18 +633,21 @@ bool FStateTreeCompiler::CompileInternalImpl()
 	}
 
 	// Store mapping between node unique ID and their compiled index. Used for debugging purposes.
+	StateTree->IDToNodeMappings.Reserve(IDToNode.Num());
 	for (const TPair<FGuid, int32>& ToNode : IDToNode)
 	{
 		StateTree->IDToNodeMappings.Emplace(ToNode.Key, FStateTreeIndex16(ToNode.Value));
 	}
 
 	// Store mapping between state unique ID and state handle. Used for debugging purposes.
+	StateTree->IDToStateMappings.Reserve(IDToState.Num());
 	for (const TPair<FGuid, int32>& ToState : IDToState)
 	{
 		StateTree->IDToStateMappings.Emplace(ToState.Key, FStateTreeStateHandle(ToState.Value));
 	}
 
 	// Store mapping between state transition identifier and compact transition index. Used for debugging purposes.
+	StateTree->IDToTransitionMappings.Reserve(IDToTransition.Num());
 	for (const TPair<FGuid, int32>& ToTransition : IDToTransition)
 	{
 		StateTree->IDToTransitionMappings.Emplace(ToTransition.Key, FStateTreeIndex16(ToTransition.Value));
@@ -870,7 +878,13 @@ bool FStateTreeCompiler::CreateStateRecursive(UStateTreeState& State, const FSta
 			TEXT("Linked State cannot have child states, because the state selection will enter to the linked state on activation."));
 	}
 
+	if (const auto Validation = UE::StateTree::Compiler::IsValidCount16(StateTree->States.Num()); Validation.DidFail())
+	{
+		Validation.Log(Log, TEXT("States"));
+		return false;
+	}
 	const int32 StateIdx = StateTree->States.AddDefaulted();
+	const FStateTreeStateHandle StateHandle = FStateTreeStateHandle(StateIdx);
 	FCompactStateTreeState& CompactState = StateTree->States[StateIdx];
 	CompactState.Name = State.Name;
 	CompactState.Tag = State.Tag;
@@ -889,6 +903,26 @@ bool FStateTreeCompiler::CreateStateRecursive(UStateTreeState& State, const FSta
 	CompactState.Type = State.Type;
 	CompactState.SelectionBehavior = State.SelectionBehavior;
 
+	// create sub frame info
+	if (!Parent.IsValid())
+	{
+		const int32 NewFrameIndex = StateTree->Frames.Num();
+		FCompactStateTreeFrame& CompactFrame = StateTree->Frames.AddDefaulted_GetRef();
+		CompactFrame.RootState = FStateTreeStateHandle((uint16)StateIdx);
+		CompactFrame.NumberOfTasksStatusMasks = 0;
+
+		StateHandleToFrameIndex.Add(StateHandle, NewFrameIndex);
+	}
+	else
+	{
+		const int32* ParentFrameIndexPtr = StateHandleToFrameIndex.Find(Parent);
+		if (ensure(ParentFrameIndexPtr != nullptr))
+		{
+			const int32 ParentFrameIndex = *ParentFrameIndexPtr;
+			StateHandleToFrameIndex.Add(StateHandle, ParentFrameIndex);
+		}
+	}
+
 	if (!StateTree->Schema->IsStateTypeAllowed(CompactState.Type))
 	{
 		Log.Reportf(EMessageSeverity::Warning,
@@ -929,6 +963,7 @@ bool FStateTreeCompiler::CreateStateRecursive(UStateTreeState& State, const FSta
 	}
 
 	SourceStates.Add(&State);
+	StatePaths.Add(State.GetPath());
 	IDToState.Add(State.ID, StateIdx);
 
 	// Child states
@@ -944,7 +979,7 @@ bool FStateTreeCompiler::CreateStateRecursive(UStateTreeState& State, const FSta
 	{
 		if (Child != nullptr && Child->Type != EStateTreeStateType::Subtree)
 		{
-			if (!CreateStateRecursive(*Child, FStateTreeStateHandle((uint16)StateIdx)))
+			if (!CreateStateRecursive(*Child, StateHandle))
 			{
 				return false;
 			}
@@ -959,14 +994,6 @@ bool FStateTreeCompiler::CreateStateRecursive(UStateTreeState& State, const FSta
 	}
 	StateTree->States[StateIdx].ChildrenEnd = uint16(ChildrenEnd); // Not using CompactState here because the array may have changed.
 
-	// create sub frame info
-	if (!Parent.IsValid())
-	{
-		FCompactStateTreeFrame& CompactFrame = StateTree->Frames.AddDefaulted_GetRef();
-		CompactFrame.RootState = FStateTreeStateHandle((uint16)StateIdx);
-		CompactFrame.NumberOfTasksStatusMasks = 0;
-	}
-
 	return true;
 }
 
@@ -1123,6 +1150,7 @@ bool FStateTreeCompiler::CreateStateTasksAndParameters()
 		FCompactStateTreeState& CompactState = StateTree->States[StateIndex];
 		const FStateTreeStateHandle CompactStateHandle(StateIndex);
 		UStateTreeState* State = SourceStates[StateIndex];
+		const FString& StatePath = GetStatePath(StateIndex);
 		check(State != nullptr);
 
 		// Carry over instance data count from parent.
@@ -1169,8 +1197,6 @@ bool FStateTreeCompiler::CreateStateTasksAndParameters()
 
 		// @todo: We should be able to skip empty parameter data.
 
-		const FString StatePath = State->GetPath(); 
-		
 		// Binding target
 		FStateTreeBindableStructDesc LinkedParamsDesc = {
 			StatePath,
@@ -1373,29 +1399,23 @@ bool FStateTreeCompiler::CreateStateTasksAndParameters()
 
 			// Find Frame and update the number of masks.
 			{
-				FStateTreeStateHandle FrameHandle = CompactStateHandle;
-				while (true)
+				const int32* FrameIndexPtr = StateHandleToFrameIndex.Find(CompactStateHandle);
+				if (FrameIndexPtr == nullptr)
 				{
-					const FCompactStateTreeState* ParentState = StateTree->GetStateFromHandle(FrameHandle);
-					check(ParentState);
-					if (!ParentState->Parent.IsValid())
-					{
-						break;
-					}
-					FrameHandle = ParentState->Parent;
+					Log.Reportf(EMessageSeverity::Error, LinkedParamsDesc, TEXT("The parent frame can't be found"));
+					bSucceeded = false;
+					continue;
 				}
-				FCompactStateTreeFrame* FoundFrame = StateTree->Frames.FindByPredicate([FrameHandle](const FCompactStateTreeFrame& Frame)
-					{
-						return Frame.RootState == FrameHandle;
-					});
-				if (FoundFrame == nullptr)
+
+				if (!StateTree->Frames.IsValidIndex(*FrameIndexPtr))
 				{
-					Log.Reportf(EMessageSeverity::Error, LinkedParamsDesc, TEXT("The parent frame can't be found"));
+					Log.Reportf(EMessageSeverity::Error, LinkedParamsDesc, TEXT("The parent frame is invalid"));
 					bSucceeded = false;
 					continue;
 				}
 
-				FoundFrame->NumberOfTasksStatusMasks = FMath::Max(FoundFrame->NumberOfTasksStatusMasks, static_cast<uint8>(CompactState.CompletionTasksMaskBufferIndex+1));
+				FCompactStateTreeFrame& Frame = StateTree->Frames[*FrameIndexPtr];
+				Frame.NumberOfTasksStatusMasks = FMath::Max(Frame.NumberOfTasksStatusMasks, static_cast<uint8>(CompactState.CompletionTasksMaskBufferIndex+1));
 			}
 		}
 
@@ -1425,7 +1445,7 @@ bool FStateTreeCompiler::CreateStateTransitions()
 
 		FStateTreeCompilerLogStateScope LogStateScope(SourceState, Log);
 
-		const FString StatePath = SourceState->GetPath();
+		const FString& StatePath = GetStatePath(i);
 
 		// Enter conditions.
 		const int32 EnterConditionsBegin = Nodes.Num();
@@ -1785,7 +1805,7 @@ bool FStateTreeCompiler::CreateStateConsiderations()
 
 		FStateTreeCompilerLogStateScope LogStateScope(SourceState, Log);
 
-		const FString StatePath = SourceState->GetPath();
+		const FString& StatePath = GetStatePath(i);
 
 		const int32 UtilityConsiderationsBegin = Nodes.Num();
 		if (const auto Validation = UE::StateTree::Compiler::IsValidCount16(UtilityConsiderationsBegin); Validation.DidFail())
@@ -1900,11 +1920,11 @@ bool FStateTreeCompiler::CreateBindingsForNodes(TConstArrayView<FStateTreeEditor
 }
 
 bool FStateTreeCompiler::CreateBindingsForStruct(
-	const FStateTreeBindableStructDesc& TargetStruct, 
-	FStateTreeDataView TargetValue, 
-	FStateTreeIndex16 PropertyFuncsBegin, 
-	FStateTreeIndex16 PropertyFuncsEnd, 
-	FStateTreeIndex16& OutBatchIndex, 
+	const FStateTreeBindableStructDesc& TargetStruct,
+	FStateTreeDataView TargetValue,
+	FStateTreeIndex16 PropertyFuncsBegin,
+	FStateTreeIndex16 PropertyFuncsEnd,
+	FStateTreeIndex16& OutBatchIndex,
 	FStateTreeIndex16* OutOutputBindingBatchIndex /* nullptr */)
 {
 	auto CopyBatch = [Self = this, &TargetStruct]
@@ -1976,9 +1996,11 @@ bool FStateTreeCompiler::CreateBindingsForStruct(
 
 bool FStateTreeCompiler::CreatePropertyFunctionsForStruct(FGuid StructID)
 {
-	for (const FPropertyBindingBinding& Binding : EditorData->EditorBindings.GetBindings())
+	check(EditorBindingsNum == EditorData->EditorBindings.GetBindings().Num());
+	for (TMultiMap<FGuid, const FPropertyBindingBinding*>::TKeyIterator It = BindingsByTargetStructID.CreateKeyIterator(StructID); It; ++It)
 	{
-		if (Binding.GetTargetPath().GetStructID() != StructID)
+		const FPropertyBindingBinding& Binding = *(It.Value());
+		if (!ensure(Binding.GetTargetPath().GetStructID() == StructID))
 		{
 			continue;
 		}
@@ -3198,3 +3220,13 @@ void FStateTreeCompiler::CreateBindingSourceStructsForNode(const FStateTreeEdito
 		BindingsCompiler.AddSourceStruct(NodeDesc);
 	}
 }
+
+const FString& FStateTreeCompiler::GetStatePath(int32 StateIndex) const
+{
+	if (ensure(StatePaths.IsValidIndex(StateIndex)))
+	{
+		return StatePaths[StateIndex];
+	}
+	static FString EmptyString;
+	return EmptyString;
+}
diff --git a/Engine/Plugins/Runtime/StateTree/Source/StateTreeEditorModule/Public/StateTreeCompiler.h b/Engine/Plugins/Runtime/StateTree/Source/StateTreeEditorModule/Public/StateTreeCompiler.h
--- a/Engine/Plugins/Runtime/StateTree/Source/StateTreeEditorModule/Public/StateTreeCompiler.h
+++ b/Engine/Plugins/Runtime/StateTree/Source/StateTreeEditorModule/Public/StateTreeCompiler.h
@@ -129,6 +129,8 @@ private:
 	
 	void CreateBindingSourceStructsForNode(const FStateTreeEditorNode& EditorNode, const FStateTreeBindableStructDesc& InstanceDesc);
 
+	const FString& GetStatePath(int32 StateIndex) const;
+
 private:
 	FStateTreeCompilerLog& Log;
 	TObjectPtr<UStateTree> StateTree = nullptr;
@@ -137,13 +139,22 @@ private:
 	TMap<FGuid, int32> IDToState;
 	TMap<FGuid, int32> IDToTransition;
 	TMap<FGuid, const FStateTreeDataView > IDToStructValue;
+	TMap<FStateTreeStateHandle, int32> StateHandleToFrameIndex;
+	TMultiMap<FGuid, const FPropertyBindingBinding*> BindingsByTargetStructID;
 	TArray<TObjectPtr<UStateTreeState>> SourceStates;
+	TArray<FString> StatePaths;
 
 	TArray<FInstancedStruct> Nodes;
 	TArray<FInstancedStruct> InstanceStructs;
 	TArray<FInstancedStruct> SharedInstanceStructs;
 	TArray<FInstancedStruct> EvaluationScopeStructs;
 	TArray<FInstancedStruct> ExecutionRuntimeStructs;
+
+	/**
+	 * Number of binding in the editor data.
+	 * Used to check that pointer in BindingsByTargetStructID are valid.
+	 */
+	int32 EditorBindingsNum = 0;
 	
 	/** Cached result of MakeCompletionTasksMask for global tasks. Indicates where state tasks should start. */
 	int32 GlobalTaskEndBit = 0;
