[State Tree] Fixed an issue events are not visible to state required event to enter from on completion transition. Fixed an issue when a subtree completes, broadcasted delegates and events will be incorrectly cleared.

- Unit tests to verify the issues.

#Synced-CL 51445612

diff --git a/Engine/Plugins/Runtime/StateTree/Source/StateTreeModule/Private/StateTreeExecutionContext.cpp b/Engine/Plugins/Runtime/StateTree/Source/StateTreeModule/Private/StateTreeExecutionContext.cpp
--- a/Engine/Plugins/Runtime/StateTree/Source/StateTreeModule/Private/StateTreeExecutionContext.cpp
+++ b/Engine/Plugins/Runtime/StateTree/Source/StateTreeModule/Private/StateTreeExecutionContext.cpp
@@ -5651,10 +5651,34 @@ bool FStateTreeExecutionContext::TriggerTransitions()
 	}
 
 	RequestedTransition.Reset();
+	
+	// Transition() can TriggerTransitions() in a loop when a sub-frame completes.
+	// We do not want to evaluate the transition from that sub-frame.
+	const bool bTriggerTransitionsFromSubtreeCompletion = TriggerTransitionsFromFrameIndex.IsSet();
+	const int32 EndFrameIndex = TriggerTransitionsFromFrameIndex.Get(Exec.ActiveFrames.Num() - 1);
+	check(Exec.ActiveFrames.IsValidIndex(EndFrameIndex));
+
+	TriggerTransitionsFromFrameIndex.Reset();
 
-	check(CurrentlyProcessedBroadcastedDelegates.IsEmpty());
-	CurrentlyProcessedBroadcastedDelegates = Storage.StealBroadcastedDelegates();
-	ON_SCOPE_EXIT{ CurrentlyProcessedBroadcastedDelegates.Empty(); };
+	if (!bTriggerTransitionsFromSubtreeCompletion)
+	{
+		check(CurrentlyProcessedBroadcastedDelegates.IsEmpty());
+		CurrentlyProcessedBroadcastedDelegates = Storage.StealBroadcastedDelegates();
+	}
+
+	ON_SCOPE_EXIT
+	{
+		if (TriggerTransitionsFromFrameIndex.IsSet() == false)
+		{
+			CurrentlyProcessedBroadcastedDelegates.Empty();
+						
+			// All events have had the change to be reacted to, clear the event queue (if this instance owns it).
+			if (InstanceData.IsOwningEventQueue() && EventQueue)
+			{
+				EventQueue->ClearEventsForCurrentTransitionProcessingPhase();
+			}
+		}
+	};
 
 	//
 	// Process transition requests
@@ -5739,10 +5763,6 @@ bool FStateTreeExecutionContext::TriggerTransitions()
 		const bool bHasEvents = EventQueue && EventQueue->HasEvents();
 		const bool bHasBroadcastedDelegates = !CurrentlyProcessedBroadcastedDelegates.IsEmpty();
 
-		// Transition() can TriggerTransitions() in a loop when a sub-frame completes.
-		//We do not want to evaluate the transition from that sub-frame.
-		const int32 EndFrameIndex = TriggerTransitionsFromFrameIndex.Get(Exec.ActiveFrames.Num() - 1);
-		TriggerTransitionsFromFrameIndex.Reset();
 		for (int32 FrameIndex = EndFrameIndex; FrameIndex >= 0; FrameIndex--)
 		{
 			FStateTreeExecutionFrame& CurrentFrame = Exec.ActiveFrames[FrameIndex];
@@ -6014,12 +6034,6 @@ PRAGMA_ENABLE_DEPRECATION_WARNINGS
 		}
 	}
 
-	// All events have had the change to be reacted to, clear the event queue (if this instance owns it).
-	if (InstanceData.IsOwningEventQueue() && EventQueue)
-	{
-		EventQueue->ClearEventsForCurrentTransitionProcessingPhase();
-	}
-
 	//
 	// Check state completion transitions.
 	//
diff --git a/Engine/Plugins/Runtime/StateTree/Source/StateTreeTestSuite/Private/StateTreeLinkedStateTest.cpp b/Engine/Plugins/Runtime/StateTree/Source/StateTreeTestSuite/Private/StateTreeLinkedStateTest.cpp
--- a/Engine/Plugins/Runtime/StateTree/Source/StateTreeTestSuite/Private/StateTreeLinkedStateTest.cpp
+++ b/Engine/Plugins/Runtime/StateTree/Source/StateTreeTestSuite/Private/StateTreeLinkedStateTest.cpp
@@ -2234,6 +2234,160 @@ struct FStateTreeTest_AutoCompletedLinkedAsset : FStateTreeTestBase
 };
 IMPLEMENT_STATE_TREE_INSTANT_TEST(FStateTreeTest_AutoCompletedLinkedAsset, "System.StateTree.LinkedAsset.AutoCompleted");
 
+/*
+ * Test for when a subtree completes, events for current transition processing phase are correctly retained.
+ */
+struct FStateTreeTest_EventTransitionInParentTree : FStateTreeTestBase
+{
+	virtual bool InstantTest() override
+	{
+		// MainTree
+		//   Root
+		//     State1 (LinkedAsset -> LATree) OnEvent(Tag1) -> State2
+		//     State2
+		//
+		// LATree
+		//   LARoot OnTick(Critical) -> Succeeded
+
+		FStateTreeCompilerLog Log;
+
+		UStateTree& LATree = NewStateTree();
+		{
+			UStateTreeEditorData& LAEditorData = *Cast<UStateTreeEditorData>(LATree.EditorData);
+
+			UStateTreeState& LARoot = LAEditorData.AddSubTree("LARoot");
+			FStateTreeTransition& LATransition = LARoot.AddTransition(EStateTreeTransitionTrigger::OnTick, EStateTreeTransitionType::Succeeded);
+			LATransition.Priority = EStateTreeTransitionPriority::Critical;
+
+			FStateTreeCompiler Compiler(Log);
+			const bool bResult = Compiler.Compile(LATree);
+			AITEST_TRUE(TEXT("LATree should get compiled"), bResult);
+		}
+
+		UStateTree& MainTree = NewStateTree();
+		{
+			UStateTreeEditorData& EditorData = *Cast<UStateTreeEditorData>(MainTree.EditorData);
+
+			UStateTreeState& Root = EditorData.AddSubTree("Root");
+
+			UStateTreeState& State1 = Root.AddChildState("State1", EStateTreeStateType::LinkedAsset);
+			UStateTreeState& State2 = Root.AddChildState("State2");
+
+			State1.SetLinkedStateAsset(&LATree);
+			State1.AddTransition(EStateTreeTransitionTrigger::OnEvent, GetTestTag1(), EStateTreeTransitionType::GotoState, &State2);
+
+			FStateTreeCompiler Compiler(Log);
+			const bool bResult = Compiler.Compile(MainTree);
+			AITEST_TRUE(TEXT("MainTree should get compiled"), bResult);
+		}
+
+		FStateTreeInstanceData InstanceData;
+		FTestStateTreeExecutionContext Exec(MainTree, MainTree, InstanceData);
+		const bool bInitSucceeded = Exec.IsValid();
+		AITEST_TRUE(TEXT("StateTree should init"), bInitSucceeded);
+
+		{
+			EStateTreeRunStatus Status = Exec.Start();
+			AITEST_EQUAL(TEXT("Start should complete with Running"), Status, EStateTreeRunStatus::Running);
+			AITEST_TRUE(TEXT("Active states should be Root/State1/LARoot"), Exec.ExpectInActiveStates("Root", "State1", "LARoot"));
+			Exec.LogClear();
+		}
+
+		{
+			Exec.SendEvent(GetTestTag1());
+
+			EStateTreeRunStatus Status = Exec.Tick(0.1f);
+			AITEST_EQUAL(TEXT("Tick should complete with Running"), Status, EStateTreeRunStatus::Running);
+			AITEST_TRUE(TEXT("Active states should be Root/State2"), Exec.ExpectInActiveStates("Root", "State2"));
+			Exec.LogClear();
+		}
+
+		Exec.Stop();
+
+		return true;
+	}
+};
+IMPLEMENT_STATE_TREE_INSTANT_TEST(FStateTreeTest_EventTransitionInParentTree, "System.StateTree.LinkedAsset.EventTransitionInParentTree");
+
+/*
+ * Test for when a subtree completes, broadcasted delegates in parent frame are correctly retained.
+ */
+struct FStateTreeTest_DelegateTransitionInParentTree : FStateTreeTestBase
+{
+	virtual bool InstantTest() override
+	{
+		// MainTree
+		//   Root [BroadcastDelegateTask] OnDelegate -> State2
+		//     State1 (LinkedAsset -> LATree)
+		//     State2
+		//
+		// LATree
+		//   LARoot OnTick -> Succeeded (Critical)
+
+		FStateTreeCompilerLog Log;
+
+		UStateTree& LATree = NewStateTree();
+		{
+			UStateTreeEditorData& LAEditorData = *Cast<UStateTreeEditorData>(LATree.EditorData);
+
+			UStateTreeState& LARoot = LAEditorData.AddSubTree("LARoot");
+			FStateTreeTransition& LATransition = LARoot.AddTransition(EStateTreeTransitionTrigger::OnTick, EStateTreeTransitionType::Succeeded);
+			LATransition.Priority = EStateTreeTransitionPriority::Critical;
+
+			FStateTreeCompiler Compiler(Log);
+			const bool bResult = Compiler.Compile(LATree);
+			AITEST_TRUE(TEXT("LATree should get compiled"), bResult);
+		}
+
+		UStateTree& MainTree = NewStateTree();
+		{
+			UStateTreeEditorData& EditorData = *Cast<UStateTreeEditorData>(MainTree.EditorData);
+
+			UStateTreeState& Root = EditorData.AddSubTree("Root");
+			TStateTreeEditorNode<FTestTask_BroadcastDelegate>& BroadcastTask = Root.AddTask<FTestTask_BroadcastDelegate>("BroadcastTask");
+
+			UStateTreeState& State1 = Root.AddChildState("State1", EStateTreeStateType::LinkedAsset);
+			State1.SetLinkedStateAsset(&LATree);
+
+			UStateTreeState& State2 = Root.AddChildState("State2");
+
+			FStateTreeTransition& DelegateTransition = Root.AddTransition(EStateTreeTransitionTrigger::OnDelegate, EStateTreeTransitionType::GotoState, &State2);
+
+			EditorData.AddPropertyBinding(
+				FPropertyBindingPath(BroadcastTask.ID, GET_MEMBER_NAME_CHECKED(FTestTask_BroadcastDelegate_InstanceData, OnTickDelegate)),
+				FPropertyBindingPath(DelegateTransition.ID, GET_MEMBER_NAME_CHECKED(FStateTreeTransition, DelegateListener)));
+
+			FStateTreeCompiler Compiler(Log);
+			const bool bResult = Compiler.Compile(MainTree);
+			AITEST_TRUE(TEXT("MainTree should get compiled"), bResult);
+		}
+
+		FStateTreeInstanceData InstanceData;
+		FTestStateTreeExecutionContext Exec(MainTree, MainTree, InstanceData);
+		const bool bInitSucceeded = Exec.IsValid();
+		AITEST_TRUE(TEXT("StateTree should init"), bInitSucceeded);
+
+		{
+			EStateTreeRunStatus Status = Exec.Start();
+			AITEST_EQUAL(TEXT("Start should complete with Running"), Status, EStateTreeRunStatus::Running);
+			AITEST_TRUE(TEXT("Active states should be Root/State1/LARoot"), Exec.ExpectInActiveStates("Root", "State1", "LARoot"));
+			Exec.LogClear();
+		}
+
+		{
+			EStateTreeRunStatus Status = Exec.Tick(0.1f);
+			AITEST_EQUAL(TEXT("Tick should complete with Running"), Status, EStateTreeRunStatus::Running);
+			AITEST_TRUE(TEXT("Active states should be Root/State2"), Exec.ExpectInActiveStates("Root", "State2"));
+			Exec.LogClear();
+		}
+
+		Exec.Stop();
+
+		return true;
+	}
+};
+IMPLEMENT_STATE_TREE_INSTANT_TEST(FStateTreeTest_DelegateTransitionInParentTree, "System.StateTree.LinkedAsset.DelegateTransitionInParentTree");
+
 } // namespace UE::StateTree::Tests
 
 UE_ENABLE_OPTIMIZATION_SHIP
diff --git a/Engine/Plugins/Runtime/StateTree/Source/StateTreeTestSuite/Private/StateTreeTransitionTest.cpp b/Engine/Plugins/Runtime/StateTree/Source/StateTreeTestSuite/Private/StateTreeTransitionTest.cpp
--- a/Engine/Plugins/Runtime/StateTree/Source/StateTreeTestSuite/Private/StateTreeTransitionTest.cpp
+++ b/Engine/Plugins/Runtime/StateTree/Source/StateTreeTestSuite/Private/StateTreeTransitionTest.cpp
@@ -2797,6 +2797,111 @@ struct FStateTreeTest_TransitionTypes : FStateTreeTestBase
 };
 IMPLEMENT_STATE_TREE_INSTANT_TEST(FStateTreeTest_TransitionTypes, "System.StateTree.Transition.Types");
 
+/*
+ * Test state selection triggered by trantiions of different types can always access events.
+ */
+struct FStateTreeTest_TransitionToRequiredEventState : FStateTreeTestBase
+{
+	virtual bool InstantTest() override
+	{
+		// Root
+		//   A1 -> OnTick Transition to A3
+		//   A2 -> OnComplete Transition to A4
+		//   (Required Event To Enter) A3 -> On Complete Transition to A2
+		//   (Required Event To Enter) A4 
+
+		UStateTree& StateTree = NewStateTree();
+		UStateTreeEditorData& EditorData = *Cast<UStateTreeEditorData>(StateTree.EditorData);
+
+		const FGameplayTag TagA3 = GetTestTag1();
+		const FGameplayTag TagA4 = GetTestTag2();
+
+		UStateTreeState& Root = EditorData.AddSubTree(FName("Root"));
+
+		UStateTreeState& StateA1 = Root.AddChildState(FName("A1"));
+		{
+			TStateTreeEditorNode<FTestTask_Stand>& TaskA1 = StateA1.AddTask<FTestTask_Stand>(FName("TaskA1"));
+			TaskA1.GetNode().TicksToCompletion = 100;
+		}
+
+		UStateTreeState& StateA2 = Root.AddChildState(FName("A2"));
+
+		UStateTreeState& StateA3 = Root.AddChildState(FName("A3"));
+		{
+			StateA3.bHasRequiredEventToEnter = true;
+			StateA3.RequiredEventToEnter.Tag = TagA3;
+		}
+
+		UStateTreeState& StateA4 = Root.AddChildState(FName("A4"));
+		{
+			StateA4.bHasRequiredEventToEnter = true;
+			StateA4.RequiredEventToEnter.Tag = TagA4;
+		}
+
+		StateA1.AddTransition(EStateTreeTransitionTrigger::OnTick, EStateTreeTransitionType::GotoState, &StateA3);
+		StateA2.AddTransition(EStateTreeTransitionTrigger::OnStateCompleted, EStateTreeTransitionType::GotoState, &StateA4);
+		StateA3.AddTransition(EStateTreeTransitionTrigger::OnStateCompleted, EStateTreeTransitionType::GotoState, &StateA2);
+
+		{
+			FStateTreeCompilerLog Log;
+			FStateTreeCompiler Compiler(Log);
+			const bool bResult = Compiler.Compile(StateTree);
+			AITEST_TRUE("StateTree should get compiled", bResult);
+		}
+
+		FStateTreeInstanceData InstanceData;
+		const FStateTreeEventQueue& EventQueue = InstanceData.GetStorage().GetEventQueue();
+		FTestStateTreeExecutionContext Exec(StateTree, StateTree, InstanceData);
+		const bool bInitSucceeded = Exec.IsValid();
+		AITEST_TRUE("StateTree should be initialized", bInitSucceeded);
+
+		EStateTreeRunStatus Status = EStateTreeRunStatus::Unset;
+		{
+			Status = Exec.Start();
+			AITEST_TRUE("State Tree should be running", Status == EStateTreeRunStatus::Running);
+			AITEST_TRUE("Active states should be [Root, A1]", Exec.ExpectInActiveStates("Root", "A1"));
+			AITEST_TRUE(TEXT("Event Queue should be empty"), EventQueue.GetEventsView().IsEmpty());
+
+			Exec.LogClear();
+		}
+
+		{
+			Exec.SendEvent(TagA3);
+
+			Status = Exec.Tick(0.1f);
+			AITEST_TRUE("State Tree should be running", Status == EStateTreeRunStatus::Running);
+			AITEST_TRUE(TEXT("Active states should be [Root, A3]"), Exec.ExpectInActiveStates("Root", "A3"));
+			AITEST_TRUE(TEXT("Event Queue should be empty"), EventQueue.GetEventsView().IsEmpty());
+			
+			Exec.LogClear();
+		}
+
+		{
+			Status = Exec.Tick(0.1f);
+			AITEST_TRUE("State Tree should be running", Status == EStateTreeRunStatus::Running);
+			AITEST_TRUE(TEXT("Active states should be [Root, A2]"), Exec.ExpectInActiveStates("Root", "A2"));
+			AITEST_TRUE(TEXT("Event Queue should be empty"), EventQueue.GetEventsView().IsEmpty());
+
+			Exec.LogClear();
+		}
+
+		{
+			Exec.SendEvent(TagA4);
+			
+			Status = Exec.Tick(0.1f);
+			AITEST_TRUE("State Tree should be running", Status == EStateTreeRunStatus::Running);
+			AITEST_TRUE(TEXT("Active states should be [Root, A4]"), Exec.ExpectInActiveStates("Root", "A4"));
+			AITEST_TRUE(TEXT("Event Queue should be empty"), EventQueue.GetEventsView().IsEmpty());
+
+			Exec.LogClear();
+		}
+
+		Exec.Stop();
+
+		return true;
+	}
+};
+IMPLEMENT_STATE_TREE_INSTANT_TEST(FStateTreeTest_TransitionToRequiredEventState, "System.StateTree.Transition.RequiredEvent");
 
 } // namespace UE::StateTree::Tests
 
