[State Tree] removed check array index as it is redundant

#Synced-CL 51445612

diff --git a/Engine/Plugins/Runtime/StateTree/Source/StateTreeModule/Private/StateTreeExecutionContext.cpp b/Engine/Plugins/Runtime/StateTree/Source/StateTreeModule/Private/StateTreeExecutionContext.cpp
--- a/Engine/Plugins/Runtime/StateTree/Source/StateTreeModule/Private/StateTreeExecutionContext.cpp
+++ b/Engine/Plugins/Runtime/StateTree/Source/StateTreeModule/Private/StateTreeExecutionContext.cpp
@@ -5656,7 +5656,6 @@ bool FStateTreeExecutionContext::TriggerTransitions()
 	// We do not want to evaluate the transition from that sub-frame.
 	const bool bTriggerTransitionsFromSubtreeCompletion = TriggerTransitionsFromFrameIndex.IsSet();
 	const int32 EndFrameIndex = TriggerTransitionsFromFrameIndex.Get(Exec.ActiveFrames.Num() - 1);
-	check(Exec.ActiveFrames.IsValidIndex(EndFrameIndex));
 
 	TriggerTransitionsFromFrameIndex.Reset();
 
