[State Tree] greyed out paste binding option if multiple objects are selected, instead of check

#Synced-CL 51656013

diff --git a/Engine/Plugins/Runtime/StateTree/Source/StateTreeEditorModule/Private/Customizations/StateTreeBindingExtension.cpp b/Engine/Plugins/Runtime/StateTree/Source/StateTreeEditorModule/Private/Customizations/StateTreeBindingExtension.cpp
--- a/Engine/Plugins/Runtime/StateTree/Source/StateTreeEditorModule/Private/Customizations/StateTreeBindingExtension.cpp
+++ b/Engine/Plugins/Runtime/StateTree/Source/StateTreeEditorModule/Private/Customizations/StateTreeBindingExtension.cpp
@@ -912,7 +912,10 @@ void FStateTreeBindingExtension::CustomizeDetailWidgetRow(FDetailWidgetRow& InWi
 		{
 			TArray<UObject*> OuterObjects;
 			InPropertyHandle->GetOuterObjects(OuterObjects);
-			checkf(OuterObjects.Num() == 1, TEXT("Paste binding is only allowed on single selected object!"));
+			if (OuterObjects.Num() != 1)
+			{
+				return false;
+			}
 
 			constexpr UScriptStruct* TargetType = nullptr;
 			UE::StateTreeEditor::ImportTextAsClipboardEditorData(TargetType, Cast<UStateTreeEditorData>(InBindingsOwner), UE::StateTree::PropertyBinding::FindStateOrEditorDataOwner(OuterObjects[0]), OutClipboard, bInProcessBuffer);
