[State Tree] Recursively Trigger transitions from parent frame when a subtree completes. This solved issues caused by When hit the max iteration, the execution context will lose track of info in the next frame when continuing transition processing from the subtree completion.

- Added unit test to verify this.

#Synced-CL 51656013

diff --git a/Engine/Plugins/Runtime/StateTree/Source/StateTreeModule/Private/StateTreeExecutionContext.cpp b/Engine/Plugins/Runtime/StateTree/Source/StateTreeModule/Private/StateTreeExecutionContext.cpp
--- a/Engine/Plugins/Runtime/StateTree/Source/StateTreeModule/Private/StateTreeExecutionContext.cpp
+++ b/Engine/Plugins/Runtime/StateTree/Source/StateTreeModule/Private/StateTreeExecutionContext.cpp
@@ -1966,9 +1966,6 @@ void FStateTreeExecutionContext::TickTriggerTransitionsInternal()
 		return;
 	}
 
-	// Reset the completed subframe counter (for unit-test that do not recreate an execution context between each tick)
-	TriggerTransitionsFromFrameIndex.Reset();
-
 	// The state selection is repeated up to MaxIteration time. This allows failed EnterState() to potentially find a new state immediately.
 	// This helps event driven StateTrees to not require another event/tick to find a suitable state.
 	static constexpr int32 MaxIterations = 5;
@@ -5626,7 +5623,7 @@ PRAGMA_DISABLE_DEPRECATION_WARNINGS
 PRAGMA_ENABLE_DEPRECATION_WARNINGS
 }
 
-bool FStateTreeExecutionContext::TriggerTransitions()
+bool FStateTreeExecutionContext::TriggerTransitions(const FTriggerTransitionsInternalArgs Args)
 {
 	//1. Process transition requests. Keep the single request with the highest priority.
 	//2. Process tick/event/delegate transitions and tasks. TriggerTransitions, from bottom to top.
@@ -5651,26 +5648,20 @@ bool FStateTreeExecutionContext::TriggerTransitions()
 	}
 
 	RequestedTransition.Reset();
-	
-	// Transition() can TriggerTransitions() in a loop when a sub-frame completes.
-	// We do not want to evaluate the transition from that sub-frame.
-	const bool bTriggerTransitionsFromSubtreeCompletion = TriggerTransitionsFromFrameIndex.IsSet();
-	const int32 EndFrameIndex = TriggerTransitionsFromFrameIndex.Get(Exec.ActiveFrames.Num() - 1);
-
-	TriggerTransitionsFromFrameIndex.Reset();
 
-	if (!bTriggerTransitionsFromSubtreeCompletion)
+	if (!Args.CompletedSubtreeParentFrameIndex.IsSet())
 	{
 		check(CurrentlyProcessedBroadcastedDelegates.IsEmpty());
 		CurrentlyProcessedBroadcastedDelegates = Storage.StealBroadcastedDelegates();
 	}
 
+	bool bTransitionsTriggeredFromSubtreeCompletion = false;
 	ON_SCOPE_EXIT
 	{
-		if (TriggerTransitionsFromFrameIndex.IsSet() == false)
+		if (!bTransitionsTriggeredFromSubtreeCompletion)
 		{
 			CurrentlyProcessedBroadcastedDelegates.Empty();
-						
+
 			// All events have had the change to be reacted to, clear the event queue (if this instance owns it).
 			if (InstanceData.IsOwningEventQueue() && EventQueue)
 			{
@@ -5762,6 +5753,9 @@ bool FStateTreeExecutionContext::TriggerTransitions()
 		const bool bHasEvents = EventQueue && EventQueue->HasEvents();
 		const bool bHasBroadcastedDelegates = !CurrentlyProcessedBroadcastedDelegates.IsEmpty();
 
+		// Transition() can TriggerTransitions() in a loop when a sub-frame completes.
+		// We do not want to evaluate the transition from that sub-frame.
+		const int32 EndFrameIndex = Args.CompletedSubtreeParentFrameIndex.Get(Exec.ActiveFrames.Num() - 1);
 		for (int32 FrameIndex = EndFrameIndex; FrameIndex >= 0; FrameIndex--)
 		{
 			FStateTreeExecutionFrame& CurrentFrame = Exec.ActiveFrames[FrameIndex];
@@ -6265,13 +6259,11 @@ PRAGMA_ENABLE_DEPRECATION_WARNINGS
 				Exec.bHasPendingCompletedState = true;
 				Exec.LastTickStatus = RunStatus;
 
-				// Clear the transition and return that no transition took place.
-				// Since the LastTickStatus != running, the transition loop will try another transition
-				// now starting from the linked parent state. If we run out of retires in the selection loop (e.g. very deep hierarchy)
-				// we will continue on next tick.
-				TriggerTransitionsFromFrameIndex = ParentFrameIndex;
-				RequestedTransition.Reset();
-				return false;
+				Storage.ResetTemporaryInstances();
+
+				// now starting from the linked parent state.
+				bTransitionsTriggeredFromSubtreeCompletion = true;
+				return TriggerTransitions({ .CompletedSubtreeParentFrameIndex = ParentFrameIndex });
 			}
 		}
 	}
diff --git a/Engine/Plugins/Runtime/StateTree/Source/StateTreeModule/Public/StateTreeExecutionContext.h b/Engine/Plugins/Runtime/StateTree/Source/StateTreeModule/Public/StateTreeExecutionContext.h
--- a/Engine/Plugins/Runtime/StateTree/Source/StateTreeModule/Public/StateTreeExecutionContext.h
+++ b/Engine/Plugins/Runtime/StateTree/Source/StateTreeModule/Public/StateTreeExecutionContext.h
@@ -1236,7 +1236,12 @@ protected:
 	 * Test the completed state transitions.
 	 * @return if the transition is valid and should be entered.
 	 */
-	UE_API bool TriggerTransitions();
+	struct FTriggerTransitionsInternalArgs
+	{
+		TOptional<int32> CompletedSubtreeParentFrameIndex;
+	};
+
+	UE_API bool TriggerTransitions(const FTriggerTransitionsInternalArgs Args = FTriggerTransitionsInternalArgs());
 
 	UE_DEPRECATED(5.7, "Use ForceTransition")
 	/** Create a new transition result from a recorded transition result. It will fail if the recorded transition is malformed. */
@@ -1594,14 +1599,14 @@ PRAGMA_ENABLE_DEPRECATION_WARNINGS
 	UE_DEPRECATED(5.7, "Use RequestedTransition instead.")
 	/** Structure describing the origin of the state transition that caused the state change. */
 	FStateTreeTransitionSource NextTransitionSource;
+
+	UE_DEPRECATED(5.8, "No longer used.")
+	TOptional<int32> TriggerTransitionsFromFrameIndex;
 #endif 
 
 	/** Current selected transition. */
 	TUniquePtr<FRequestTransitionResult> RequestedTransition;
 
-	/** When set, start the transitions loop from TriggerTransitionsFromFrameIndex. */
-	TOptional<int32> TriggerTransitionsFromFrameIndex;
-
 	/** Current broadcasted delegates we're processing. */
 	TArray<FStateTreeDelegateDispatcher> CurrentlyProcessedBroadcastedDelegates;
 
diff --git a/Engine/Plugins/Runtime/StateTree/Source/StateTreeTestSuite/Private/StateTreeLinkedStateTest.cpp b/Engine/Plugins/Runtime/StateTree/Source/StateTreeTestSuite/Private/StateTreeLinkedStateTest.cpp
--- a/Engine/Plugins/Runtime/StateTree/Source/StateTreeTestSuite/Private/StateTreeLinkedStateTest.cpp
+++ b/Engine/Plugins/Runtime/StateTree/Source/StateTreeTestSuite/Private/StateTreeLinkedStateTest.cpp
@@ -2388,6 +2388,231 @@ struct FStateTreeTest_DelegateTransitionInParentTree : FStateTreeTestBase
 };
 IMPLEMENT_STATE_TREE_INSTANT_TEST(FStateTreeTest_DelegateTransitionInParentTree, "System.StateTree.LinkedAsset.DelegateTransitionInParentTree");
 
+/*
+ * Test for subtrees of deep hierarchy complete one after another.
+ */
+struct FStateTreeTest_DeepSubtreeChain : FStateTreeTestBase
+{
+	virtual bool InstantTest() override
+	{
+		// MainTree
+		//   Root
+		//     State1 (LinkedAsset -> LA1) OnCompleted -> GotoState State2
+		//     State2
+		//
+		// LA1
+		//   LA1Root (LinkedAsset -> LA2) OnCompleted -> Succeeded
+		//
+		// LA2
+		//   LA2Root (LinkedAsset -> LA3) OnTick -> Succeeded
+		//
+		// LA3
+		//   LA3Root (LinkedAsset -> LA4) OnCompleted -> Succeeded
+		//
+		// LA4
+		//   LA4Root (LinkedAsset -> LA5) OnTick -> Succeeded
+		//
+		// LA5
+		//   LA5Root (LinkedAsset -> LA6) OnCompleted -> Succeeded
+		//
+		// LA6
+		//   LA6Root (LinkedAsset -> LA7) OnTick -> Succeeded
+		//
+		// LA7
+		//   LA7Root (LinkedAsset -> LA8) OnCompleted -> Succeeded
+		//
+		// LA8
+		//   LA8Root OnTick -> Succeeded
+		//
+		// LA9
+		//   LA9Root (LinkedAsset -> LA10) OnCompleted -> Succeeded
+		//
+		// LA10
+		//   LA10Root OnTick -> Succeeded
+
+		FStateTreeCompilerLog Log;
+
+		bool bLinkedAssetCompilationSucceeded = false;
+
+		auto CreateLinkedAssetTree = [this, &Log, &bLinkedAssetCompilationSucceeded](const auto* RootName, const bool bUseOnTick, UStateTree* NextTree = nullptr) -> UStateTree&
+		{
+			UStateTree& Tree = NewStateTree();
+			UStateTreeEditorData& EditorData = *Cast<UStateTreeEditorData>(Tree.EditorData);
+				
+			UStateTreeState& Root = EditorData.AddSubTree(RootName, EStateTreeStateType::State);
+			if (NextTree)
+			{
+				Root.Type = EStateTreeStateType::LinkedAsset;
+				Root.SetLinkedStateAsset(NextTree);
+			}
+
+			const EStateTreeTransitionTrigger Trigger = bUseOnTick ? EStateTreeTransitionTrigger::OnTick : EStateTreeTransitionTrigger::OnStateCompleted;
+			FStateTreeTransition& Transition = Root.AddTransition(Trigger, EStateTreeTransitionType::Succeeded);
+
+			FStateTreeCompiler Compiler(Log);
+			bLinkedAssetCompilationSucceeded = Compiler.Compile(Tree);
+
+			return Tree;
+		};
+
+		UStateTree& LA10Tree = CreateLinkedAssetTree("LA10Root", true);
+		AITEST_TRUE(TEXT("LinkedAsset Tree should get compiled"), bLinkedAssetCompilationSucceeded);
+
+		UStateTree& LA9Tree = CreateLinkedAssetTree("LA9Root", false, &LA10Tree);
+		AITEST_TRUE(TEXT("LinkedAsset Tree should get compiled"), bLinkedAssetCompilationSucceeded);
+
+		UStateTree& LA8Tree = CreateLinkedAssetTree("LA8Root", true, &LA9Tree);
+		AITEST_TRUE(TEXT("LinkedAsset Tree should get compiled"), bLinkedAssetCompilationSucceeded);
+
+		UStateTree& LA7Tree = CreateLinkedAssetTree("LA7Root", false, &LA8Tree);
+		AITEST_TRUE(TEXT("LinkedAsset Tree should get compiled"), bLinkedAssetCompilationSucceeded);
+
+		UStateTree& LA6Tree = CreateLinkedAssetTree("LA6Root", true, &LA7Tree);
+		AITEST_TRUE(TEXT("LinkedAsset Tree should get compiled"), bLinkedAssetCompilationSucceeded);
+
+		UStateTree& LA5Tree = CreateLinkedAssetTree("LA5Root", false, &LA6Tree);
+		AITEST_TRUE(TEXT("LinkedAsset Tree should get compiled"), bLinkedAssetCompilationSucceeded);
+
+		UStateTree& LA4Tree = CreateLinkedAssetTree("LA4Root", true, &LA5Tree);
+		AITEST_TRUE(TEXT("LinkedAsset Tree should get compiled"), bLinkedAssetCompilationSucceeded);
+
+		UStateTree& LA3Tree = CreateLinkedAssetTree("LA3Root", false, &LA4Tree);
+		AITEST_TRUE(TEXT("LinkedAsset Tree should get compiled"), bLinkedAssetCompilationSucceeded);
+
+		UStateTree& LA2Tree = CreateLinkedAssetTree("LA2Root", true, &LA3Tree);
+		AITEST_TRUE(TEXT("LinkedAsset Tree should get compiled"), bLinkedAssetCompilationSucceeded);
+
+		UStateTree& LA1Tree = CreateLinkedAssetTree("LA1Root", false, &LA2Tree);
+		AITEST_TRUE(TEXT("LinkedAsset Tree should get compiled"), bLinkedAssetCompilationSucceeded);
+
+		UStateTree& MainTree = NewStateTree();
+		{
+			UStateTreeEditorData& EditorData = *Cast<UStateTreeEditorData>(MainTree.EditorData);
+			UStateTreeState& Root = EditorData.AddSubTree("Root");
+
+			UStateTreeState& State1 = Root.AddChildState("State1", EStateTreeStateType::LinkedAsset);
+			UStateTreeState& State2 = Root.AddChildState("State2");
+
+			State1.SetLinkedStateAsset(&LA1Tree);
+			State1.AddTransition(EStateTreeTransitionTrigger::OnStateCompleted, EStateTreeTransitionType::GotoState, &State2);
+
+			FStateTreeCompiler Compiler(Log);
+			const bool bResult = Compiler.Compile(MainTree);
+			AITEST_TRUE(TEXT("MainTree should get compiled"), bResult);
+		}
+
+		FStateTreeInstanceData InstanceData;
+		FTestStateTreeExecutionContext Exec(MainTree, MainTree, InstanceData);
+		const bool bInitSucceeded = Exec.IsValid();
+		AITEST_TRUE(TEXT("StateTree should init"), bInitSucceeded);
+
+		{
+			EStateTreeRunStatus Status = Exec.Start();
+			AITEST_EQUAL(TEXT("Start should complete with Running"), Status, EStateTreeRunStatus::Running);
+			AITEST_TRUE(TEXT("Active states should be Root/State1/LA1Root/.../LA10Root"),
+				Exec.ExpectInActiveStates("Root", "State1", "LA1Root", "LA2Root", "LA3Root", "LA4Root", "LA5Root", "LA6Root", "LA7Root", "LA8Root", "LA9Root", "LA10Root"));
+			Exec.LogClear();
+		}
+
+		{
+			EStateTreeRunStatus Status = Exec.Tick(0.1f);
+			AITEST_EQUAL(TEXT("Tick should complete with Running"), Status, EStateTreeRunStatus::Running);
+			AITEST_TRUE(TEXT("Active states should be Root/State2"),
+				Exec.ExpectInActiveStates("Root", "State2"));
+			Exec.LogClear();
+		}
+
+		Exec.Stop();
+
+		return true;
+	}
+};
+IMPLEMENT_STATE_TREE_INSTANT_TEST(FStateTreeTest_DeepSubtreeChain, "System.StateTree.LinkedAsset.DeepSubtreeChain");
+
+
+/*
+ * Test that a LinkedAsset state can reselect itself when the linked tree completes.
+ * The OnStateCompleted transition goes back to itself, causing the linked tree to re-enter each cycle.
+ */
+struct FStateTreeTest_LinkedAssetReselectOnCompletion : FStateTreeTestBase
+{
+	virtual bool InstantTest() override
+	{
+		// MainTree
+		//   Root
+		//     LinkedState (LinkedAsset -> LATree) [OnCompleted -> GotoState LinkedState (reselect)]
+		//
+		// LATree
+		//   LARoot [OnCompleted -> Succeeded]
+
+		FStateTreeCompilerLog Log;
+
+		// LATree
+		UStateTree& LATree = NewStateTree();
+		{
+			UStateTreeEditorData& LAEditorData = *Cast<UStateTreeEditorData>(LATree.EditorData);
+
+			UStateTreeState& LARoot = LAEditorData.AddSubTree("LARoot");
+			LARoot.AddTransition(EStateTreeTransitionTrigger::OnStateCompleted, EStateTreeTransitionType::Succeeded);
+
+			FStateTreeCompiler Compiler(Log);
+			const bool bResult = Compiler.Compile(LATree);
+			AITEST_TRUE(TEXT("LATree should get compiled"), bResult);
+		}
+
+		// MainTree
+		UStateTree& MainTree = NewStateTree();
+		{
+			UStateTreeEditorData& EditorData = *Cast<UStateTreeEditorData>(MainTree.EditorData);
+
+			UStateTreeState& Root = EditorData.AddSubTree("Root");
+			UStateTreeState& LinkedState = Root.AddChildState("LinkedState", EStateTreeStateType::LinkedAsset);
+			LinkedState.SetLinkedStateAsset(&LATree);
+			LinkedState.AddTransition(EStateTreeTransitionTrigger::OnStateCompleted, EStateTreeTransitionType::GotoState, &LinkedState);
+
+			FStateTreeCompiler Compiler(Log);
+			const bool bResult = Compiler.Compile(MainTree);
+			AITEST_TRUE(TEXT("MainTree should get compiled"), bResult);
+		}
+
+		const TCHAR* EnterStateStr = TEXT("EnterState");
+
+		FStateTreeInstanceData InstanceData;
+		FTestStateTreeExecutionContext Exec(MainTree, MainTree, InstanceData);
+		const bool bInitSucceeded = Exec.IsValid();
+		AITEST_TRUE(TEXT("StateTree should init"), bInitSucceeded);
+
+		// Start: enter Root/LinkedState/LARoot
+		{
+			const EStateTreeRunStatus Status = Exec.Start();
+			AITEST_EQUAL(TEXT("Start should complete with Running"), Status, EStateTreeRunStatus::Running);
+			AITEST_TRUE(TEXT("Active states should be Root/LinkedState/LARoot"), Exec.ExpectInActiveStates("Root", "LinkedState", "LARoot"));
+			Exec.LogClear();
+		}
+
+		// Tick 1: LATask completes -> LATree succeeds -> LinkedState reselects itself -> LATask re-enters
+		{
+			const EStateTreeRunStatus Status = Exec.Tick(0.1f);
+			AITEST_EQUAL(TEXT("Tick 1 should complete with Running"), Status, EStateTreeRunStatus::Running);
+			AITEST_TRUE(TEXT("Active states should be Root/LinkedState/LARoot after reselect"), Exec.ExpectInActiveStates("Root", "LinkedState", "LARoot"));
+			Exec.LogClear();
+		}
+
+		// Tick 2: Same cycle again
+		{
+			const EStateTreeRunStatus Status = Exec.Tick(0.1f);
+			AITEST_EQUAL(TEXT("Tick 2 should complete with Running"), Status, EStateTreeRunStatus::Running);
+			AITEST_TRUE(TEXT("Active states should still be Root/LinkedState/LARoot"), Exec.ExpectInActiveStates("Root", "LinkedState", "LARoot"));
+			Exec.LogClear();
+		}
+
+		Exec.Stop();
+
+		return true;
+	}
+};
+IMPLEMENT_STATE_TREE_INSTANT_TEST(FStateTreeTest_LinkedAssetReselectOnCompletion, "System.StateTree.LinkedAsset.ReselectOnCompletion");
+
 } // namespace UE::StateTree::Tests
 
 UE_ENABLE_OPTIMIZATION_SHIP
