[StateTree] Add the output binding batch index info to the compilation output log.

#Synced-CL 51494712

diff --git a/Engine/Plugins/Runtime/PropertyBindingUtils/Source/PropertyBindingUtils/Private/PropertyBindingBindingCollection.cpp b/Engine/Plugins/Runtime/PropertyBindingUtils/Source/PropertyBindingUtils/Private/PropertyBindingBindingCollection.cpp
--- a/Engine/Plugins/Runtime/PropertyBindingUtils/Source/PropertyBindingUtils/Private/PropertyBindingBindingCollection.cpp
+++ b/Engine/Plugins/Runtime/PropertyBindingUtils/Source/PropertyBindingUtils/Private/PropertyBindingBindingCollection.cpp
@@ -1496,7 +1496,7 @@ FString FPropertyBindingBindingCollection::DebugAsString() const
 	DebugString.Appendf(TEXT("\nCopyBatches (%d)\n"), CopyBatches.Num());
 	if (CopyBatches.Num())
 	{
-		DebugString.Appendf(TEXT("  [ (Idx) | %-45s | %-45s | %-8s [%-5s:%-5s[ | %-8s [%-5s:%-5s[ ]\n"),
+		DebugString.Appendf(TEXT("  [ (Idx) | %-45s | %-80s | %-8s [%-5s:%-5s[ | %-8s [%-5s:%-5s[ ]\n"),
 			TEXT("Target Type"), TEXT("Target Name"),
 			TEXT("Bindings"), TEXT("Beg"), TEXT("End"),
 			TEXT("ProFunc"), TEXT("Beg"), TEXT("End"));
@@ -1504,7 +1504,7 @@ FString FPropertyBindingBindingCollection::DebugAsString() const
 		{
 			const FPropertyBindingCopyInfoBatch& CopyBatch = CopyBatches[Index];
 			const FPropertyBindingBindableStructDescriptor& Descriptor = CopyBatch.TargetStruct.Get();
-			DebugString.Appendf(TEXT("  | (%3d) | %-45s | %-45s | %8s [%5d:%-5d[ | %8s [%5d:%-5d[ |\n"),
+			DebugString.Appendf(TEXT("  | (%3d) | %-45s | %-80s | %8s [%5d:%-5d[ | %8s [%5d:%-5d[ |\n"),
 				Index,
 				Descriptor.Struct ? *Descriptor.Struct->GetName() : TEXT("null"),
 				*Descriptor.ToString(),
diff --git a/Engine/Plugins/Runtime/StateTree/Source/StateTreeModule/Private/StateTree.cpp b/Engine/Plugins/Runtime/StateTree/Source/StateTreeModule/Private/StateTree.cpp
--- a/Engine/Plugins/Runtime/StateTree/Source/StateTreeModule/Private/StateTree.cpp
+++ b/Engine/Plugins/Runtime/StateTree/Source/StateTreeModule/Private/StateTree.cpp
@@ -1791,21 +1791,20 @@ FString UStateTree::DebugInternalLayoutAsString() const
 		}
 	}
 
-	// @todo: add output binding batch index info
-
 	// Evaluators
 	DebugString.Appendf(TEXT("\nEvaluators(%d)\n"), EvaluatorsNum);
 	if (EvaluatorsNum)
 	{
-		DebugString.Appendf(TEXT("  [ (Idx) | %-30s | %8s | %14s ]\n"),
-			TEXT("Name"), TEXT("Bindings"), TEXT("Struct Idx"));
+		DebugString.Appendf(TEXT("  [ (Idx) | %-30s | %12s | %12s | %14s ]\n"),
+			TEXT("Name"), TEXT("In Bindings"), TEXT("Out Binding"), TEXT("Struct Idx"));
 		for (int32 EvalIndex = EvaluatorsBegin; EvalIndex < (EvaluatorsBegin + EvaluatorsNum); EvalIndex++)
 		{
 			const FStateTreeEvaluatorBase& Eval = Nodes[EvalIndex].Get<const FStateTreeEvaluatorBase>();
-			DebugString.Appendf(TEXT("  | (%3d) | %-30s | %8d | %14s |\n"),
+			DebugString.Appendf(TEXT("  | (%3d) | %-30s | %12d | %12d | %14s |\n"),
 				EvalIndex,
 				*Eval.Name.ToString(),
 				Eval.BindingsBatch.Get(),
+				Eval.OutputBindingsBatch.Get(),
 				*Eval.InstanceDataHandle.Describe());
 		}
 	}
@@ -1817,8 +1816,8 @@ FString UStateTree::DebugInternalLayoutAsString() const
 		NumberOfTasks += State.TasksNum;
 	}
 
-	DebugString.Appendf(TEXT("\nTasks(%d)\n  [ (Idx) | %-30s | %-30s | %8s | %14s ]\n"),
-		NumberOfTasks, TEXT("State"), TEXT("Name"), TEXT("Bindings"), TEXT("Struct Idx"));
+	DebugString.Appendf(TEXT("\nTasks(%d)\n  [ (Idx) | %-30s | %-30s | %12s | %12s | %14s ]\n"),
+		NumberOfTasks, TEXT("State"), TEXT("Name"), TEXT("In Bindings"), TEXT("Out Bindings"), TEXT("Struct Idx"));
 	for (const FCompactStateTreeState& State : States)
 	{
 		if (State.TasksNum)
@@ -1826,11 +1825,12 @@ FString UStateTree::DebugInternalLayoutAsString() const
 			for (int32 TaskIndex = State.TasksBegin; TaskIndex < (State.TasksBegin + State.TasksNum); TaskIndex++)
 			{
 				const FStateTreeTaskBase& Task = Nodes[TaskIndex].Get<const FStateTreeTaskBase>();
-				DebugString.Appendf(TEXT("  | (%3d) | %-30s | %-30s | %8d | %14s |\n"),
+				DebugString.Appendf(TEXT("  | (%3d) | %-30s | %-30s | %12d | %12d | %14s |\n"),
 					TaskIndex,
 					*State.Name.ToString(),
 					*Task.Name.ToString(),
 					Task.BindingsBatch.Get(),
+					Task.OutputBindingsBatch.Get(),
 					*Task.InstanceDataHandle.Describe());
 			}
 		}
@@ -1838,11 +1838,12 @@ FString UStateTree::DebugInternalLayoutAsString() const
 	for (int32 TaskIndex = GlobalTasksBegin; TaskIndex < (GlobalTasksBegin + GlobalTasksNum); TaskIndex++)
 	{
 		const FStateTreeTaskBase& Task = Nodes[TaskIndex].Get<const FStateTreeTaskBase>();
-		DebugString.Appendf(TEXT("  | (%3d) | %-30s | %-30s | %8d | %14s |\n"),
+		DebugString.Appendf(TEXT("  | (%3d) | %-30s | %-30s | %8d | %12d | %14s |\n"),
 			TaskIndex,
 			TEXT("Global"),
 			*Task.Name.ToString(),
 			Task.BindingsBatch.Get(),
+			Task.OutputBindingsBatch.Get(),
 			*Task.InstanceDataHandle.Describe());
 	}
 
