[StateTree] The selection highlight color in SStateTreeViewRow was hardcoded as FColor(236, 134, 39). It is now registered as "Colors.StateSelection" in the StateTreeEditor style.

#Synced-CL 51494712

diff --git a/Engine/Plugins/Runtime/StateTree/Source/StateTreeEditorModule/Private/SStateTreeDiff.cpp b/Engine/Plugins/Runtime/StateTree/Source/StateTreeEditorModule/Private/SStateTreeDiff.cpp
--- a/Engine/Plugins/Runtime/StateTree/Source/StateTreeEditorModule/Private/SStateTreeDiff.cpp
+++ b/Engine/Plugins/Runtime/StateTree/Source/StateTreeEditorModule/Private/SStateTreeDiff.cpp
@@ -2,6 +2,7 @@
 
 #include "SStateTreeDiff.h"
 #include "Customizations/StateTreeBindingExtension.h"
+#include "StateTreeEditorStyle.h"
 #include "DetailsDiff.h"
 #include "DetailTreeNode.h"
 #include "Editor.h"
@@ -645,11 +646,11 @@ FLinearColor SDiffWidget::GetRowHighlightColor(const TUniquePtr<FAsyncDetailView
 	switch (DiffNode->DiffResult)
 	{
 	case ETreeDiffResult::MissingFromTree1:
-		return FLinearColor(0.f, 1.f, 0.f, .7f);
+		return FStateTreeEditorStyle::Get().GetSlateColor("Colors.Diff.MissingFromTree1").GetSpecifiedColor();
 	case ETreeDiffResult::MissingFromTree2:
-		return FLinearColor(1.f, 0.f, 0.f, .7f);
+		return FStateTreeEditorStyle::Get().GetSlateColor("Colors.Diff.MissingFromTree2").GetSpecifiedColor();
 	default:
-		return FLinearColor(1.f, 1.f, 0.f, .7f);
+		return FStateTreeEditorStyle::Get().GetSlateColor("Colors.Diff.Changed").GetSpecifiedColor();
 	}
 }
 
diff --git a/Engine/Plugins/Runtime/StateTree/Source/StateTreeEditorModule/Private/SStateTreeViewRow.cpp b/Engine/Plugins/Runtime/StateTree/Source/StateTreeEditorModule/Private/SStateTreeViewRow.cpp
--- a/Engine/Plugins/Runtime/StateTree/Source/StateTreeEditorModule/Private/SStateTreeViewRow.cpp
+++ b/Engine/Plugins/Runtime/StateTree/Source/StateTreeEditorModule/Private/SStateTreeViewRow.cpp
@@ -1289,8 +1289,7 @@ FSlateColor SStateTreeViewRow::GetActiveStateColor() const
 		}
 		if (StateTreeViewModel && StateTreeViewModel->IsSelected(State))
 		{
-			// @todo: change to the common selection color.
-			return FLinearColor(FColor(236, 134, 39));
+			return FStateTreeEditorStyle::Get().GetSlateColor("Colors.StateSelection");
 		}
 	}
 
diff --git a/Engine/Plugins/Runtime/StateTree/Source/StateTreeEditorModule/Private/StateTreeEditorStyle.cpp b/Engine/Plugins/Runtime/StateTree/Source/StateTreeEditorModule/Private/StateTreeEditorStyle.cpp
--- a/Engine/Plugins/Runtime/StateTree/Source/StateTreeEditorModule/Private/StateTreeEditorStyle.cpp
+++ b/Engine/Plugins/Runtime/StateTree/Source/StateTreeEditorModule/Private/StateTreeEditorStyle.cpp
@@ -483,6 +483,10 @@ FStateTreeEditorStyle::FStateTreeEditorStyle()
 	{
 		Set("Colors.StateLinkingIn", FLinearColor::Yellow);
 		Set("Colors.StateLinkedOut", FLinearColor::Green);
+		Set("Colors.StateSelection", FLinearColor(FColor(236, 134, 39)));
+		Set("Colors.Diff.MissingFromTree1", FLinearColor(0.f, 1.f, 0.f, .7f));
+		Set("Colors.Diff.MissingFromTree2", FLinearColor(1.f, 0.f, 0.f, .7f));
+		Set("Colors.Diff.Changed", FLinearColor(1.f, 1.f, 0.f, .7f));
 	}
 }
 
