[StateTree] Improve type safety in instance data mutable accessor .

#rnx
#Synced-CL 51494712

diff --git a/Engine/Plugins/Runtime/StateTree/Source/StateTreeModule/Public/StateTreeInstanceData.h b/Engine/Plugins/Runtime/StateTree/Source/StateTreeModule/Public/StateTreeInstanceData.h
--- a/Engine/Plugins/Runtime/StateTree/Source/StateTreeModule/Public/StateTreeInstanceData.h
+++ b/Engine/Plugins/Runtime/StateTree/Source/StateTreeModule/Public/StateTreeInstanceData.h
@@ -733,12 +733,6 @@ PRAGMA_ENABLE_DEPRECATION_WARNINGS
 			{
 				DataView = Storage.GetMutableTemporaryStruct(CurrentFrame->FrameID, DataHandle);
 			}
-
-
-			if constexpr (!std::is_same_v<T, void>)
-			{
-				check(DataView.GetStruct() == nullptr || DataView.GetStruct()->IsChildOf<T>());
-			}
 		}
 		else
 		{
@@ -761,7 +755,16 @@ PRAGMA_ENABLE_DEPRECATION_WARNINGS
 			}
 		}
 
-		return reinterpret_cast<T*>(DataView.GetMutableMemory());
+		if constexpr (!std::is_same_v<T, void>)
+		{
+			const bool bIsValidType = DataView.GetStruct() == nullptr || DataView.GetStruct()->IsChildOf<T>();
+			if (!ensure(bIsValidType))
+			{
+				return nullptr;
+			}
+		}
+
+		return static_cast<T*>(DataView.GetMutableMemory());
 	}
 
 protected:
