[StateTree] Debug macros for evaluating utility on states.

#Synced-CL 51494761

diff --git a/Engine/Plugins/Runtime/StateTree/Source/StateTreeModule/Private/Debugger/StateTreeDebug.cpp b/Engine/Plugins/Runtime/StateTree/Source/StateTreeModule/Private/Debugger/StateTreeDebug.cpp
--- a/Engine/Plugins/Runtime/StateTree/Source/StateTreeModule/Private/Debugger/StateTreeDebug.cpp
+++ b/Engine/Plugins/Runtime/StateTree/Source/StateTreeModule/Private/Debugger/StateTreeDebug.cpp
@@ -27,6 +27,7 @@ FNodeDelegate OnTickTask_AnyThread;
 FNodeDelegate OnTaskExitState_AnyThread;
 FEventSentDelegate OnEventSent_AnyThread;
 FEventConsumedDelegate OnEventConsumed_AnyThread;
+FStateUtilityEvaluatedDelegate OnStateUtilityEvaluated_AnyThread;
 
 FNodeReference::FNodeReference(TNotNull<const UStateTree*> InStateTree, const FStateTreeIndex16 InNodeIndex)
 	: StateTree(InStateTree)
@@ -161,6 +162,11 @@ void EventConsumed(const FStateTreeExecutionContext& ExecutionContext, const FSt
 	OnEventConsumed_AnyThread.Broadcast(ExecutionContext, Event);
 }
 
+void StateUtilityEvaluated(const FStateTreeExecutionContext& ExecutionContext, FStateTreeStateHandle StateHandle, float Score)
+{
+	OnStateUtilityEvaluated_AnyThread.Broadcast(ExecutionContext, StateHandle, Score);
+}
+
 }//namespace UE::StateTree::Debug
 
 #endif // WITH_STATETREE_DEBUG
diff --git a/Engine/Plugins/Runtime/StateTree/Source/StateTreeModule/Private/StateTreeExecutionContext.cpp b/Engine/Plugins/Runtime/StateTree/Source/StateTreeModule/Private/StateTreeExecutionContext.cpp
--- a/Engine/Plugins/Runtime/StateTree/Source/StateTreeModule/Private/StateTreeExecutionContext.cpp
+++ b/Engine/Plugins/Runtime/StateTree/Source/StateTreeModule/Private/StateTreeExecutionContext.cpp
@@ -5189,6 +5189,8 @@ float FStateTreeExecutionContext::EvaluateUtilityWithValidation(const FStateTree
 	FCurrentlyProcessedFrameScope FrameScope(*this, CurrentParentFrame, CurrentFrame);
 	FCurrentlyProcessedStateScope NextStateScope(*this, CurrentStateHandle);
 
+	UE_STATETREE_DEBUG_SCOPED_STATE_PHASE(this, CurrentStateHandle, EStateTreeUpdatePhase::EvaluateUtility);
+
 	TStaticArray<EStateTreeExpressionOperand, UE::StateTree::MaxExpressionIndent + 1> Operands(InPlace, EStateTreeExpressionOperand::Copy);
 	TStaticArray<float, UE::StateTree::MaxExpressionIndent + 1> Values(InPlace, false);
 
@@ -5280,7 +5282,11 @@ float FStateTreeExecutionContext::EvaluateUtilityWithValidation(const FStateTree
 		PopEvaluationScopeInstanceContainer(EvaluationScopeContainer);
 	}
 
-	return StateWeight * Values[0];
+	const float Score = StateWeight * Values[0];
+
+	UE_STATETREE_DEBUG_STATE_UTILITY_EVALUATED(this, CurrentStateHandle, Score);
+
+	return Score;
 }
 
 void FStateTreeExecutionContext::EvaluatePropertyFunctionsOnActiveInstances(const FStateTreeExecutionFrame* CurrentParentFrame, const FStateTreeExecutionFrame& CurrentFrame, FStateTreeIndex16 FuncsBegin, uint16 FuncsNum)
diff --git a/Engine/Plugins/Runtime/StateTree/Source/StateTreeModule/Public/Debugger/StateTreeDebug.h b/Engine/Plugins/Runtime/StateTree/Source/StateTreeModule/Public/Debugger/StateTreeDebug.h
--- a/Engine/Plugins/Runtime/StateTree/Source/StateTreeModule/Public/Debugger/StateTreeDebug.h
+++ b/Engine/Plugins/Runtime/StateTree/Source/StateTreeModule/Public/Debugger/StateTreeDebug.h
@@ -26,6 +26,8 @@ enum class EStateTreeTraceEventType : uint8;
 #define UE_STATETREE_DEBUG_TASK_TICK(InCtx, InStateTree, InNodeIndex) ::UE::StateTree::Debug::TaskTick(*(InCtx), ::UE::StateTree::Debug::FNodeReference((InStateTree), FStateTreeIndex16(InNodeIndex)))
 #define UE_STATETREE_DEBUG_TASK_EXIT_STATE(InCtx, InStateTree, InNodeIndex) ::UE::StateTree::Debug::TaskExitState(*(InCtx), ::UE::StateTree::Debug::FNodeReference((InStateTree), FStateTreeIndex16(InNodeIndex)))
 
+#define UE_STATETREE_DEBUG_STATE_UTILITY_EVALUATED(InCtx, InStateHandled, Score) ::UE::StateTree::Debug::StateUtilityEvaluated(*(InCtx), InStateHandled, Score)
+
 #define UE_STATETREE_DEBUG_SEND_EVENT(InCtx, InStateTree, InTag, InPayload, InOrigin) ::UE::StateTree::Debug::EventSent(*(InCtx), InStateTree, InTag, InPayload, InOrigin)
 #define UE_STATETREE_DEBUG_EVENT_CONSUMED(InExecutionContextPtr, Event) ::UE::StateTree::Debug::EventConsumed(*(InExecutionContextPtr), Event)
 
@@ -112,6 +114,7 @@ namespace UE::StateTree::Debug
 	DECLARE_TS_MULTICAST_DELEGATE_ThreeParams(FPhaseDelegate, const FStateTreeExecutionContext&, EStateTreeUpdatePhase, FStateTreeStateHandle);
 	DECLARE_TS_MULTICAST_DELEGATE_ThreeParams(FStateDelegate, const FStateTreeExecutionContext&, FStateTreeStateHandle, EStateTreeTraceEventType);
 	DECLARE_TS_MULTICAST_DELEGATE_ThreeParams(FTransitionDelegate, const FStateTreeExecutionContext&, const FStateTreeTransitionSource&, EStateTreeTraceEventType);
+	DECLARE_TS_MULTICAST_DELEGATE_ThreeParams(FStateUtilityEvaluatedDelegate, const FStateTreeExecutionContext&, FStateTreeStateHandle, float);
 
 	void ConditionEnterState(const FStateTreeExecutionContext& ExecutionContext, FNodeReference Node);
 	void ConditionTest(const FStateTreeExecutionContext& ExecutionContext, FNodeReference Node);
@@ -125,6 +128,8 @@ namespace UE::StateTree::Debug
 	void TaskTick(const FStateTreeExecutionContext& ExecutionContext, FNodeReference Node);
 	void TaskExitState(const FStateTreeExecutionContext& ExecutionContext, FNodeReference Node);
 
+	void StateUtilityEvaluated(const FStateTreeExecutionContext& ExecutionContext, FStateTreeStateHandle StateHandle, float Score);
+
 	void EventSent(const FStateTreeMinimalExecutionContext& ExecutionContext, TNotNull<const UStateTree*> StateTree, FGameplayTag Tag, FConstStructView Payload, FName Origin);
 	void EventConsumed(const FStateTreeExecutionContext& ExecutionContext, const FStateTreeSharedEvent& Event);
 
@@ -239,6 +244,13 @@ namespace UE::StateTree::Debug
 	 * @note The StateTree can execute on any thread.
 	 */
 	STATETREEMODULE_API extern FEventConsumedDelegate OnEventConsumed_AnyThread;
+
+	/**
+	 * Debugging callback for when an utility score is calculated for a state.
+	 * @note The callback executes inside the StateTree logic.
+	 * @note The StateTree can execute on any thread.
+	 */
+	STATETREEMODULE_API extern FStateUtilityEvaluatedDelegate OnStateUtilityEvaluated_AnyThread;
 } //UE::StateTree::Debug
 
 #else
@@ -255,6 +267,8 @@ namespace UE::StateTree::Debug
 #define UE_STATETREE_DEBUG_TASK_TICK(...)
 #define UE_STATETREE_DEBUG_TASK_EXIT_STATE(...)
 
+#define UE_STATETREE_DEBUG_STATE_UTILITY_EVALUATED(...)
+
 #define UE_STATETREE_DEBUG_SEND_EVENT(...)
 #define UE_STATETREE_DEBUG_EVENT_CONSUMED(...)
 
diff --git a/Engine/Plugins/Runtime/StateTree/Source/StateTreeModule/Public/StateTreeExecutionTypes.h b/Engine/Plugins/Runtime/StateTree/Source/StateTreeModule/Public/StateTreeExecutionTypes.h
--- a/Engine/Plugins/Runtime/StateTree/Source/StateTreeModule/Public/StateTreeExecutionTypes.h
+++ b/Engine/Plugins/Runtime/StateTree/Source/StateTreeModule/Public/StateTreeExecutionTypes.h
@@ -43,7 +43,8 @@ enum class EStateTreeUpdatePhase : uint8
 	EnterConditions					UMETA(DisplayName = "Enter conditions"),
 	EnterStates						UMETA(DisplayName = "Enter States"),
 	ExitStates						UMETA(DisplayName = "Exit States"),
-	StateCompleted					UMETA(DisplayName = "State(s) Completed")
+	StateCompleted					UMETA(DisplayName = "State(s) Completed"),
+	EvaluateUtility					UMETA(DisplayName = "Evaluate Utility")
 };
 
 
