[StateTree] Add the Output ,Input and Context category by default for Blueprint nodes.

#Synced-CL 51945425

diff --git a/Engine/Plugins/Runtime/StateTree/Source/StateTreeEditorModule/Private/StateTreeEditorModule.cpp b/Engine/Plugins/Runtime/StateTree/Source/StateTreeEditorModule/Private/StateTreeEditorModule.cpp
--- a/Engine/Plugins/Runtime/StateTree/Source/StateTreeEditorModule/Private/StateTreeEditorModule.cpp
+++ b/Engine/Plugins/Runtime/StateTree/Source/StateTreeEditorModule/Private/StateTreeEditorModule.cpp
@@ -196,6 +196,45 @@ void FStateTreeEditorModule::StartupModule()
 		OnPostEngineInitHandle = FCoreDelegates::GetOnPostEngineInit().AddLambda([this]()
 			{
 				NodeClassCache = UE::StateTree::Editor::InitNodeClassCache();
+
+				FBlueprintEditorModule& BlueprintEditorModule = FModuleManager::LoadModuleChecked<FBlueprintEditorModule>("Kismet");
+				{
+					UE::BlueprintEditor::FAdditionalCategory ContextCategory;
+					ContextCategory.Category = LOCTEXT("ContextCategory", "Context");
+					ContextCategory.ClassFilters.Add(FSoftClassPath(UStateTreeNodeBlueprintBase::StaticClass()));
+					ContextCategory.Placeholders.Emplace(
+						UE::BlueprintEditor::FAdditionalCategory::FPlaceholder{
+							.MenuDescription = LOCTEXT("ContextCategoryMenu", "Category for context bindings."),
+							.Tooltip = LOCTEXT("ContextCategoryTooltip", "The property is connected to one of the context objects, or can be overridden with property binding."),
+							.SectionID = NodeSectionID::VARIABLE
+						});
+					ContextCategoryHandle = BlueprintEditorModule.RegisterAdditionalBlueprintCategory(ContextCategory);
+				}
+				{
+					UE::BlueprintEditor::FAdditionalCategory InputCategory;
+					InputCategory.Category = LOCTEXT("InputCategory", "Input");
+					InputCategory.ClassFilters.Add(FSoftClassPath(UStateTreeNodeBlueprintBase::StaticClass()));
+					InputCategory.Placeholders.Emplace(
+						UE::BlueprintEditor::FAdditionalCategory::FPlaceholder{
+							.MenuDescription = LOCTEXT("InputCategoryMenu", "Category for input bindings."),
+							.Tooltip = LOCTEXT("InputCategoryTooltip", "The property is always expected to be bound to some other property."),
+							.SectionID = NodeSectionID::VARIABLE
+						});
+					InputCategoryHandle = BlueprintEditorModule.RegisterAdditionalBlueprintCategory(InputCategory);
+				}
+				{
+					UE::BlueprintEditor::FAdditionalCategory OutputCategory;
+					OutputCategory.Category = LOCTEXT("OutputCategory", "Output");
+					OutputCategory.ClassFilters.Add(FSoftClassPath(UStateTreeTaskBlueprintBase::StaticClass()));
+					OutputCategory.ClassFilters.Add(FSoftClassPath(UStateTreeEvaluatorBlueprintBase::StaticClass()));
+					OutputCategory.Placeholders.Emplace(
+						UE::BlueprintEditor::FAdditionalCategory::FPlaceholder{
+							.MenuDescription = LOCTEXT("OutputCategoryMenu", "Category for output bindings."),
+							.Tooltip = LOCTEXT("OutputCategoryTooltip", "The node will always set its value. The property can bind to another property to push the value. Other nodes can also bind to it to fetch the value."),
+							.SectionID = NodeSectionID::VARIABLE
+						});
+					OutputCategoryHandle = BlueprintEditorModule.RegisterAdditionalBlueprintCategory(OutputCategory);
+				}
 			});
 	}
 
@@ -280,6 +319,13 @@ void FStateTreeEditorModule::ShutdownModule()
 	}
 	NodeClassCache.Reset();
 
+	if (FBlueprintEditorModule* BlueprintEditorModule = FModuleManager::GetModulePtr<FBlueprintEditorModule>("Kismet"))
+	{
+		BlueprintEditorModule->UnregisterAdditionalBlueprintCategory(InputCategoryHandle);
+		BlueprintEditorModule->UnregisterAdditionalBlueprintCategory(ContextCategoryHandle);
+		BlueprintEditorModule->UnregisterAdditionalBlueprintCategory(OutputCategoryHandle);
+	}
+
 	UE::StateTree::Compiler::FCompilerManager::Shutdown();
 
 	UE::StateTree::Delegates::Private::OnAppendToClassSchema.Unbind();
diff --git a/Engine/Plugins/Runtime/StateTree/Source/StateTreeEditorModule/Public/StateTreeEditorModule.h b/Engine/Plugins/Runtime/StateTree/Source/StateTreeEditorModule/Public/StateTreeEditorModule.h
--- a/Engine/Plugins/Runtime/StateTree/Source/StateTreeEditorModule/Public/StateTreeEditorModule.h
+++ b/Engine/Plugins/Runtime/StateTree/Source/StateTreeEditorModule/Public/StateTreeEditorModule.h
@@ -2,6 +2,7 @@
 
 #pragma  once
 
+#include "BlueprintEditorModule.h"
 #include "Modules/ModuleInterface.h"
 #include "Logging/LogMacros.h"
 #include "Templates/NonNullSubclassOf.h"
@@ -163,6 +164,11 @@ protected:
 	FDelegateHandle OnUserDefinedStructReinstancedHandle;
 	FOnRegisterLayoutExtensions RegisterLayoutExtensions;
 
+	/** Handle for unregistering the categories registered with FBlueprintEditorModule. */
+	FBlueprintEditorModule::FAdditionalCategoryHandle ContextCategoryHandle;
+	FBlueprintEditorModule::FAdditionalCategoryHandle InputCategoryHandle;
+	FBlueprintEditorModule::FAdditionalCategoryHandle OutputCategoryHandle;
+
 	FOnAssetRegistryTags AssetRegistryTags;
 	FValidateStateTree PreValidateStateTree;
 	FValidateStateTree ValidateStateTree;
