[StateTree] Add an option to compile all dirty state trees when PIE starts.
Add option to use the hash-based approach to determine whether a state tree asset needs to be recompiled before entering PIE. While UStateTree::EditorDataDirtyStatus should be sufficient for this purpose, this option serves as a fallback for users who have not yet upgraded their code base.

#Synced-CL 51824547

diff --git a/Engine/Plugins/Runtime/StateTree/Source/StateTreeEditorModule/Private/StateTreeCompilerManager.cpp b/Engine/Plugins/Runtime/StateTree/Source/StateTreeEditorModule/Private/StateTreeCompilerManager.cpp
--- a/Engine/Plugins/Runtime/StateTree/Source/StateTreeEditorModule/Private/StateTreeCompilerManager.cpp
+++ b/Engine/Plugins/Runtime/StateTree/Source/StateTreeEditorModule/Private/StateTreeCompilerManager.cpp
@@ -36,14 +36,21 @@ FAutoConsoleVariableRef CVarUseDependenciesToTriggerCompilation(
 	TEXT("Use the build dependencies to detect when a state tree needs to be linked or compiled.")
 );
 
-bool bEnableCompileSynchronouslyInPIESession = true;
+bool bEnableForceCompileSynchronouslyInPIESession = true;
 FAutoConsoleVariableRef CVarEnableCompileSynchronouslyInPIESession(
-	TEXT("StateTree.Compiler.EnableCompileSynchronouslyInPIESession"),
-	bEnableCompileSynchronouslyInPIESession,
+	TEXT("StateTree.Compiler.EnableForceCompileSynchronouslyInPIESession"),
+	bEnableForceCompileSynchronouslyInPIESession,
 	TEXT("When PIE starts, flush the compilation queue.\n")
 	TEXT("Any new asset will be compile synchronously event is async was requested.")
 );
 
+bool bEnableCompileAllIfChangedOnBeginPIE = true;
+FAutoConsoleVariableRef CVarEnableCompileAllIfChangedOnBeginPIE(
+	TEXT("StateTree.Compiler.EnableCompileAllIfChangedOnBeginPIE"),
+	bEnableCompileAllIfChangedOnBeginPIE,
+	TEXT("When PIE starts, compile all state tree asset (if needed).\n")
+);
+
 int32 NumberOfQueuedCompilationPerBatch = 1;
 FAutoConsoleVariableRef CVarNumberOfQueuedCompilationPerBatch(
 	TEXT("StateTree.Compiler.NumberOfQueuedCompilationPerBatch"),
@@ -299,7 +306,7 @@ FCompilerManagerImpl::FCompilerManagerImpl()
 	// When running in standalone
 	if (IsRunningGame())
 	{
-		bCanQueueCompilation = !Private::bEnableCompileSynchronouslyInPIESession;
+		bCanQueueCompilation = !Private::bEnableForceCompileSynchronouslyInPIESession;
 	}
 }
 
@@ -348,7 +355,10 @@ TOptional<bool> FCompilerManagerImpl::CompileIfNeededSynchronously(TNotNull<USta
 {
 	check(IsInGameThread() || IsInAsyncLoadingThread());
 
-	if (StateTree->IsEditorDataDirty())
+	// Test if the asset is already in the queue. Faster than hashing the editor.
+	const bool bIsDirty = StateTree->bCompilationPending
+		|| UStateTreeEditingSubsystem::NeedsRecompile(StateTree);
+	if (bIsDirty)
 	{
 		FStateTreeCompilerLog Log;
 		const bool bResult = CompileInternalSynchronously(StateTree, Log);
@@ -734,13 +744,22 @@ void FCompilerManagerImpl::HandleStateTreeCooking(TNotNull<UStateTree*> StateTre
 
 void FCompilerManagerImpl::HandlePreBeginPIE(const bool bIsSimulating)
 {
-	bCanQueueCompilation = !Private::bEnableCompileSynchronouslyInPIESession;
+	bCanQueueCompilation = !Private::bEnableForceCompileSynchronouslyInPIESession;
 	FlushCompilationQueue();
+
+	if (Private::bEnableCompileAllIfChangedOnBeginPIE)
+	{
+		for (TObjectIterator<UStateTree> It; It; ++It)
+		{
+			check(!It->HasAnyFlags(RF_ClassDefaultObject));
+			CompileIfNeededSynchronously(*It);
+		}
+	}
 }
 
 void FCompilerManagerImpl::HandleEndPIE(const bool bIsSimulating)
 {
-	bCanQueueCompilation = Private::bEnableCompileSynchronouslyInPIESession;
+	bCanQueueCompilation = Private::bEnableForceCompileSynchronouslyInPIESession;
 }
 
 void FCompilerManagerImpl::HandleObjectsReinstanced(const FReplacementObjectMap& ObjectMap)
diff --git a/Engine/Plugins/Runtime/StateTree/Source/StateTreeEditorModule/Private/StateTreeEditingSubsystem.cpp b/Engine/Plugins/Runtime/StateTree/Source/StateTreeEditorModule/Private/StateTreeEditingSubsystem.cpp
--- a/Engine/Plugins/Runtime/StateTree/Source/StateTreeEditorModule/Private/StateTreeEditingSubsystem.cpp
+++ b/Engine/Plugins/Runtime/StateTree/Source/StateTreeEditorModule/Private/StateTreeEditingSubsystem.cpp
@@ -26,6 +26,13 @@ namespace UE::StateTree::Editor
 		bEnableQueuedCompilationWhenAssetDirty,
 		TEXT("When the asset is modified, automatically queue the asset compilation.")
 	);
+
+	bool bUseHashToDetermineIfAssetNeedsCompile = true;
+	FAutoConsoleVariableRef CVarUseHashToDetermineIfAssetNeedsCompile(
+		TEXT("StateTree.Compiler.UseHashToDetermineIfAssetNeedsCompile"),
+		bUseHashToDetermineIfAssetNeedsCompile,
+		TEXT("Use the hash to determine if the asset is dirty and requires compilation.\n")
+	);
 }
 
 UStateTreeEditingSubsystem::UStateTreeEditingSubsystem()
@@ -48,8 +55,13 @@ bool UStateTreeEditingSubsystem::CompileStateTree(TNotNull<UStateTree*> InStateT
 
 bool UStateTreeEditingSubsystem::NeedsRecompile(TNotNull<const UStateTree*> InStateTree)
 {
-	const uint32 CurrentHash = CalculateStateTreeHash(InStateTree);
-	return InStateTree->LastCompiledEditorDataHash != CurrentHash;
+	bool bIsDirty = InStateTree->IsEditorDataDirty();
+	if (!bIsDirty && UE::StateTree::Editor::bUseHashToDetermineIfAssetNeedsCompile)
+	{
+		const uint32 CurrentHash = CalculateStateTreeHash(InStateTree);
+		bIsDirty = InStateTree->LastCompiledEditorDataHash != CurrentHash;
+	}
+	return bIsDirty;
 }
 
 void UStateTreeEditingSubsystem::MarkAsPubliclyModified(TNotNull<UStateTree*> InStateTree)
diff --git a/Engine/Plugins/Runtime/StateTree/Source/StateTreeModule/Private/StateTree.cpp b/Engine/Plugins/Runtime/StateTree/Source/StateTreeModule/Private/StateTree.cpp
--- a/Engine/Plugins/Runtime/StateTree/Source/StateTreeModule/Private/StateTree.cpp
+++ b/Engine/Plugins/Runtime/StateTree/Source/StateTreeModule/Private/StateTree.cpp
@@ -477,12 +477,12 @@ EDataValidationResult UStateTree::IsDataValid(FDataValidationContext& Context) c
 				return EDataValidationResult::Invalid;
 			}
 		}
-	}
 
-	if (!const_cast<UStateTree*>(this)->Link())
-	{
-		Context.AddError(FText::FromString(FString::Printf(TEXT("%s failed to link. Please recompile the State Tree for more details errors."), *GetPathName())));
-		return EDataValidationResult::Invalid;
+		if (!const_cast<UStateTree*>(this)->Link())
+		{
+			Context.AddError(FText::FromString(FString::Printf(TEXT("%s failed to link. Please recompile the State Tree for more details errors."), *GetPathName())));
+			return EDataValidationResult::Invalid;
+		}
 	}
 
 	return Super::IsDataValid(Context);
@@ -1601,10 +1601,7 @@ TArray<FStateTreeMemoryUsage> UStateTree::CalculateEstimatedMemoryUsage() const
 
 void UStateTree::CompileIfChanged()
 {
-	if (IsEditorDataDirty())
-	{
-		UE::StateTree::Delegates::Private::OnCompileIfChanged.ExecuteIfBound(this);
-	}
+	UE::StateTree::Delegates::Private::OnCompileIfChanged.ExecuteIfBound(this);
 }
 #endif //WITH_EDITOR
 
