[StateTree] Add more ways to detect that the asset was modified since the last compilation. It improves performance because we do not need to hash the state tree asset when we already know the asset was modified.
#Synced-CL 52002607

diff --git a/Engine/Plugins/Runtime/StateTree/Source/StateTreeEditorModule/Private/StateTreeEditingSubsystem.cpp b/Engine/Plugins/Runtime/StateTree/Source/StateTreeEditorModule/Private/StateTreeEditingSubsystem.cpp
--- a/Engine/Plugins/Runtime/StateTree/Source/StateTreeEditorModule/Private/StateTreeEditingSubsystem.cpp
+++ b/Engine/Plugins/Runtime/StateTree/Source/StateTreeEditorModule/Private/StateTreeEditingSubsystem.cpp
@@ -139,7 +139,29 @@ void UStateTreeEditingSubsystem::ValidateStateTree(TNotNull<UStateTree*> InState
 {
 	TRACE_CPUPROFILER_EVENT_SCOPE(UStateTreeEditingSubsystem::ValidateStateTree)
 
-	auto FixChangedStateLinkName = [](FStateTreeStateLink& StateLink, const TMap<FGuid, FName>& IDToName) -> bool
+	auto ModifyStateTree = [InStateTree](const bool bPubliclyModified)
+		{
+			constexpr bool bMarkDirty = false;
+			InStateTree->Modify(bMarkDirty);
+			if (bPubliclyModified)
+			{
+				MarkAsPubliclyModified(InStateTree);
+			}
+			else
+			{
+				MarkAsModified(InStateTree);
+			}
+		};
+	auto ModifyState = [&ModifyStateTree](TNotNull<UStateTreeState*> State)
+		{
+			constexpr bool bMarkDirty = false;
+			State->Modify(bMarkDirty);
+
+			constexpr bool bPubliclyModified = false;
+			ModifyStateTree(bPubliclyModified);
+		};
+
+	auto FixChangedStateLinkName = [&ModifyState](TNotNull<UStateTreeState*> State, FStateTreeStateLink& StateLink, const TMap<FGuid, FName>& IDToName) -> bool
 		{
 			if (StateLink.ID.IsValid())
 			{
@@ -152,6 +174,7 @@ void UStateTreeEditingSubsystem::ValidateStateTree(TNotNull<UStateTree*> InState
 				if (StateLink.Name != *Name)
 				{
 					// Name changed, fix!
+					ModifyState(State);
 					StateLink.Name = *Name;
 					return true;
 				}
@@ -179,19 +202,19 @@ void UStateTreeEditingSubsystem::ValidateStateTree(TNotNull<UStateTree*> InState
 				State.Modify(bMarkDirty);
 				if (State.Type == EStateTreeStateType::Linked)
 				{
-					FixChangedStateLinkName(State.LinkedSubtree, IDToName);
+					FixChangedStateLinkName(&State, State.LinkedSubtree, IDToName);
 				}
 					
 				for (FStateTreeTransition& Transition : State.Transitions)
 				{
-					FixChangedStateLinkName(Transition.State, IDToName);
+					FixChangedStateLinkName(&State, Transition.State, IDToName);
 				}
 
 				return EStateTreeVisitor::Continue;
 			});
 		};
 
-	auto FixEditorData = [](TNotNull<UStateTree*> StateTree)
+	auto FixEditorData = [&ModifyStateTree](TNotNull<UStateTree*> StateTree)
 	{
 		UStateTreeEditorData* EditorData = Cast<UStateTreeEditorData>(StateTree->EditorData);
 		// The schema is defined in the EditorData. If we can't find the editor data (probably because the class doesn't exist anymore), then try the compiled schema in the state tree asset.
@@ -219,8 +242,8 @@ void UStateTreeEditingSubsystem::ValidateStateTree(TNotNull<UStateTree*> InState
 			EditorData->Schema = NewObject<UStateTreeSchema>(EditorData, SchemaClass.Get());
 			EditorData->EditorBindings.SetBindingsOwner(EditorData);
 
-			constexpr bool bMarkDirty = false;
-			StateTree->Modify(bMarkDirty);
+			constexpr bool bPubliclyModified = true;
+			ModifyStateTree(bPubliclyModified);
 			StateTree->EditorData = EditorData;
 		}
 		else if (!EditorData->IsA(EditorDataClass.Get()))
@@ -238,8 +261,8 @@ void UStateTreeEditingSubsystem::ValidateStateTree(TNotNull<UStateTree*> InState
 			}
 			EditorData->EditorBindings.SetBindingsOwner(EditorData);
 
-			constexpr bool bMarkDirty = false;
-			StateTree->Modify(bMarkDirty);
+			constexpr bool bPubliclyModified = true;
+			ModifyStateTree(bPubliclyModified);
 			StateTree->EditorData = EditorData;
 
 			// Trash the previous EditorData
@@ -248,7 +271,7 @@ void UStateTreeEditingSubsystem::ValidateStateTree(TNotNull<UStateTree*> InState
 		}
 	};
 
-	auto FixEditorSchema = [](TNotNull<UStateTreeEditorData*> EditorData)
+	auto FixEditorSchema = [&ModifyStateTree](TNotNull<UStateTreeEditorData*> EditorData)
 		{
 			TSubclassOf<const UStateTreeSchema> SchemaClass = EditorData->Schema ? EditorData->Schema->GetClass() : nullptr;
 			if (SchemaClass.Get() == nullptr)
@@ -259,12 +282,16 @@ void UStateTreeEditingSubsystem::ValidateStateTree(TNotNull<UStateTree*> InState
 			TNonNullSubclassOf<UStateTreeEditorSchema> EditorSchemaClass = FStateTreeEditorModule::GetModule().GetEditorSchemaClass(SchemaClass.Get());
 			if (EditorData->EditorSchema == nullptr)
 			{
+				constexpr bool bPubliclyModified = false;
+				ModifyStateTree(bPubliclyModified);
 				EditorData->EditorSchema = NewObject<UStateTreeEditorSchema>(EditorData, EditorSchemaClass.Get(), FName(), RF_Transactional);
 			}
 			else if (!EditorData->EditorSchema->IsA(EditorSchemaClass.Get()))
 			{
 				// The current EditorSchema is not of the correct type. The data needs to be patched by the schema desired editor data subclass.
 				UStateTreeEditorSchema* PreviousEditorSchema = EditorData->EditorSchema;
+				constexpr bool bPubliclyModified = false;
+				ModifyStateTree(bPubliclyModified);
 				EditorData->EditorSchema = CastChecked<UStateTreeEditorSchema>(StaticDuplicateObject(PreviousEditorSchema, EditorData, FName(), RF_Transactional, EditorSchemaClass.Get()));
 				
 				// Trash the previous EditorDataSchema
@@ -287,6 +314,7 @@ void UStateTreeEditingSubsystem::ValidateStateTree(TNotNull<UStateTree*> InState
 				{
 					constexpr bool bMarkDirty = false;
 					State.Modify(bMarkDirty);
+					//@todo detect if the tree is dirty and call ModifyStateTree if modified
 					State.UpdateParametersFromLinkedSubtree();
 				}
 				return EStateTreeVisitor::Continue;
diff --git a/Engine/Plugins/Runtime/StateTree/Source/StateTreeEditorModule/Private/StateTreeEditorData.cpp b/Engine/Plugins/Runtime/StateTree/Source/StateTreeEditorModule/Private/StateTreeEditorData.cpp
--- a/Engine/Plugins/Runtime/StateTree/Source/StateTreeEditorModule/Private/StateTreeEditorData.cpp
+++ b/Engine/Plugins/Runtime/StateTree/Source/StateTreeEditorModule/Private/StateTreeEditorData.cpp
@@ -234,7 +234,7 @@ void UStateTreeEditorData::PostEditChangeChainProperty(FPropertyChangedChainEven
 				Extensions.Reset();
 			}
 			UE::StateTree::Delegates::OnSchemaChanged.Broadcast(*StateTree);
-			UStateTreeEditingSubsystem::MarkAsModified(StateTree);
+			UStateTreeEditingSubsystem::MarkAsPubliclyModified(StateTree);
 		}
 		else if (MemberName == GET_MEMBER_NAME_CHECKED(UStateTreeEditorData, RootParameterPropertyBag))
 		{
@@ -316,6 +316,7 @@ void UStateTreeEditorData::PostDuplicate(EDuplicateMode::Type DuplicateMode)
 
 	EditorBindings.SetBindingsOwner(this);
 	DuplicateIDs();
+	UStateTreeEditingSubsystem::MarkAsPubliclyModified(GetTypedOuter<UStateTree>());
 }
 #endif // WITH_EDITOR
 
@@ -619,13 +620,14 @@ void UStateTreeEditorData::CreateParametersForStruct(const FGuid StructID, TArra
 		return;
 	}
 
-	const UStateTree* StateTree = GetTypedOuter<UStateTree>();
+	UStateTree* StateTree = GetTypedOuter<UStateTree>();
 	checkf(StateTree, TEXT("UStateTreeEditorData should only be allocated within a UStateTree"));
 
 	if (RootParametersGuid == StructID)
 	{
 		CreateRootProperties(InOutCreationDescs);
 		UE::StateTree::Delegates::OnParametersChanged.Broadcast(*StateTree);
+		UStateTreeEditingSubsystem::MarkAsPubliclyModified(StateTree);
 		return;
 	}
 
@@ -635,6 +637,7 @@ void UStateTreeEditorData::CreateParametersForStruct(const FGuid StructID, TArra
 		{
 			UE::PropertyBinding::CreateUniquelyNamedPropertiesInPropertyBag(InOutCreationDescs,State.Parameters.Parameters);
 			UE::StateTree::Delegates::OnStateParametersChanged.Broadcast(*StateTree, State.ID);
+			UStateTreeEditingSubsystem::MarkAsModified(StateTree);
 			return EStateTreeVisitor::Break;
 		}
 		return EStateTreeVisitor::Continue;
@@ -757,8 +760,12 @@ void UStateTreeEditorData::ReparentStates()
 			UE_LOG(LogStateTreeEditor, Log, TEXT("%s: Fixing outer on state %s."), *TreeData->GetFullName(), *GetNameSafe(&State));
 			State.Rename(nullptr, ExpectedOuter, REN_DontCreateRedirectors | REN_DoNotDirty | REN_AllowPackageLinkerMismatch);
 		}
-		
-		State.Parent = ParentState;
+
+		if (State.Parent != ParentState)
+		{
+			UStateTreeEditingSubsystem::MarkAsModified(TreeData->GetTypedOuter<UStateTree>());
+			State.Parent = ParentState;
+		}
 		
 		return EStateTreeVisitor::Continue;
 	});
@@ -767,23 +774,25 @@ void UStateTreeEditorData::ReparentStates()
 void UStateTreeEditorData::CreateRootProperties(TArrayView<UE::PropertyBinding::FPropertyCreationDescriptor> InOutCreationDescs)
 {
 	UE::PropertyBinding::CreateUniquelyNamedPropertiesInPropertyBag(InOutCreationDescs, RootParameterPropertyBag);
-	UStateTree* StateTree = GetTypedOuter<UStateTree>();
-	UStateTreeEditingSubsystem::MarkAsPubliclyModified(StateTree);
+	UStateTreeEditingSubsystem::MarkAsPubliclyModified(GetTypedOuter<UStateTree>());
 }
 
 void UStateTreeEditorData::FixObjectInstances(TSet<UObject*>& SeenObjects, UObject& Outer, FStateTreeEditorNode& Node)
 {
-	auto NewInstances = [this, &Node, &Outer](const UStruct* Struct, const UStruct* ExecutionRuntimeStruct)
+	auto NewInstances = [this, &Node, &Outer](const UStruct* Struct, const UStruct* ExecutionRuntimeStruct) -> bool
 		{
+			bool bDirty = false;
 			if (Struct && !Node.Instance.IsValid() && Node.InstanceObject == nullptr)
 			{
 				if (const UScriptStruct* InstanceType = Cast<const UScriptStruct>(Struct))
 				{
 					Node.Instance.InitializeAs(InstanceType);
+					bDirty = true;
 				}
 				else if (const UClass* InstanceClass = Cast<const UClass>(Struct))
 				{
 					Node.InstanceObject = NewObject<UObject>(&Outer, InstanceClass);
+					bDirty = true;
 				}
 			}
 			if (ExecutionRuntimeStruct && !Node.ExecutionRuntimeData.IsValid() && Node.ExecutionRuntimeDataObject == nullptr)
@@ -791,36 +800,46 @@ void UStateTreeEditorData::FixObjectInstances(TSet<UObject*>& SeenObjects, UObje
 				if (const UScriptStruct* InstanceType = Cast<const UScriptStruct>(ExecutionRuntimeStruct))
 				{
 					Node.ExecutionRuntimeData.InitializeAs(InstanceType);
+					bDirty = true;
 				}
 				else if (const UClass* InstanceClass = Cast<const UClass>(ExecutionRuntimeStruct))
 				{
 					Node.ExecutionRuntimeDataObject = NewObject<UObject>(&Outer, InstanceClass);
+					bDirty = true;
 				}
 			}
+
+			return bDirty;
 		};
-	auto FixInstance = [this, &SeenObjects, &Outer](FInstancedStruct& Instance, const UStruct* Struct)
+	auto FixInstance = [this, &SeenObjects, &Outer](FInstancedStruct& Instance, const UStruct* Struct) -> bool
 		{
+			bool bDirty = false;
 			if (Instance.IsValid())
 			{
 				if (Instance.GetScriptStruct() != Struct)
 				{
 					Instance.Reset();
+					bDirty = true;
 				}
 			}
+			return bDirty;
 		};
-	auto FixInstanceObject = [this, &SeenObjects, &Outer](TObjectPtr<UObject>& InstanceObject, const UStruct* Struct)
+	auto FixInstanceObject = [this, &SeenObjects, &Outer](TObjectPtr<UObject>& InstanceObject, const UStruct* Struct) -> bool
 		{
+			bool bDirty = false;
 			if (InstanceObject)
 			{
 				if (InstanceObject->GetClass() != Struct)
 				{
 					InstanceObject = nullptr;
+					bDirty = true;
 				}
 				// Found a duplicate reference to an object, make unique copy.
 				else if(SeenObjects.Contains(InstanceObject))				
 				{
 					UE_LOG(LogStateTreeEditor, Log, TEXT("%s: Making duplicate node instance %s unique."), *GetFullName(), *GetNameSafe(InstanceObject));
 					InstanceObject = DuplicateObject(InstanceObject, &Outer);
+					bDirty = true;
 				}
 				else
 				{
@@ -829,27 +848,66 @@ void UStateTreeEditorData::FixObjectInstances(TSet<UObject*>& SeenObjects, UObje
 					{
 						UE_LOG(LogStateTreeEditor, Log, TEXT("%s: Fixing outer on node instance %s."), *GetFullName(), *GetNameSafe(InstanceObject));
 						InstanceObject->Rename(nullptr, &Outer, REN_DontCreateRedirectors | REN_DoNotDirty | REN_AllowPackageLinkerMismatch);
+						bDirty = true;
 					}
 				}
 				SeenObjects.Add(InstanceObject);
 			}
+			return bDirty;
 		};
 
 	const FStateTreeNodeBase* NodeBase = Node.Node.GetPtr<FStateTreeNodeBase>();
+	bool bDirty = false;
 	if (NodeBase)
 	{
-		FixInstance(Node.Instance, NodeBase->GetInstanceDataType());
-		FixInstanceObject(Node.InstanceObject, NodeBase->GetInstanceDataType());
-		FixInstance(Node.ExecutionRuntimeData, NodeBase->GetExecutionRuntimeDataType());
-		FixInstanceObject(Node.ExecutionRuntimeDataObject, NodeBase->GetExecutionRuntimeDataType());
-		NewInstances(NodeBase->GetInstanceDataType(), NodeBase->GetExecutionRuntimeDataType());
+		if (FixInstance(Node.Instance, NodeBase->GetInstanceDataType()))
+		{
+			bDirty = true;
+		}
+		if (FixInstanceObject(Node.InstanceObject, NodeBase->GetInstanceDataType()))
+		{
+			bDirty = true;
+		}
+		if (FixInstance(Node.ExecutionRuntimeData, NodeBase->GetExecutionRuntimeDataType()))
+		{
+			bDirty = true;
+		}
+		if (FixInstanceObject(Node.ExecutionRuntimeDataObject, NodeBase->GetExecutionRuntimeDataType()))
+		{
+			bDirty = true;
+		}
+		if (NewInstances(NodeBase->GetInstanceDataType(), NodeBase->GetExecutionRuntimeDataType()))
+		{
+			bDirty = true;
+		}
 	}
 	else
 	{
-		Node.Instance.Reset();
-		Node.InstanceObject = nullptr;
-		Node.ExecutionRuntimeData.Reset();
-		Node.ExecutionRuntimeDataObject = nullptr;
+		if (Node.Instance.IsValid())
+		{
+			Node.Instance.Reset();
+			bDirty = true;
+		}
+		if (Node.InstanceObject != nullptr)
+		{
+			Node.InstanceObject = nullptr;
+			bDirty = true;
+		}
+		if (Node.ExecutionRuntimeData.IsValid())
+		{
+			Node.ExecutionRuntimeData.Reset();
+			bDirty = true;
+		}
+		if (Node.ExecutionRuntimeDataObject != nullptr)
+		{
+			Node.ExecutionRuntimeDataObject = nullptr;
+			bDirty = true;
+		}
+	}
+
+	if (bDirty)
+	{
+		UStateTreeEditingSubsystem::MarkAsModified(GetTypedOuter<UStateTree>());
 	}
 };
 
@@ -870,16 +928,20 @@ void UStateTreeEditorData::FixObjectNodes()
 		{
 			TreeData->FixObjectInstances(SeenObjects, State, Node);
 		}
+
+		// Consideration Utilities
+		for (FStateTreeEditorNode& Node : State.Considerations)
+		{
+			TreeData->FixObjectInstances(SeenObjects, State, Node);
+		}
 		
 		// Tasks
 		for (FStateTreeEditorNode& Node : State.Tasks)
 		{
 			TreeData->FixObjectInstances(SeenObjects, State, Node);
 		}
-
 		TreeData->FixObjectInstances(SeenObjects, State, State.SingleTask);
 
-
 		// Transitions
 		for (FStateTreeTransition& Transition : State.Transitions)
 		{
@@ -940,6 +1002,7 @@ void UStateTreeEditorData::FixDuplicateIDs()
 	// Around version 5.1-5.3 we had issue that copy/paste or some duplication methods could create nodes with duplicate IDs.
 	// This code tries to fix that, it looks for duplicates, makes them unique, and duplicates the bindings when ID changes.
 	TSet<FGuid> FoundNodeIDs;
+	bool bDirty = false;
 
 	// Evaluators
 	for (int32 Index = 0; Index < Evaluators.Num(); Index++)
@@ -955,6 +1018,7 @@ void UStateTreeEditorData::FixDuplicateIDs()
 				UE_LOG(LogStateTreeEditor, Log, TEXT("%s: Found Evaluator '%s' with duplicate ID, changing ID:%s to ID:%s."),
 					*GetFullName(), *Node.GetName().ToString(), *OldID.ToString(), *Node.ID.ToString());
 				EditorBindings.CopyBindings(OldID, Node.ID);
+				bDirty = true;
 			}
 			FoundNodeIDs.Add(Node.ID);
 		}
@@ -974,12 +1038,13 @@ void UStateTreeEditorData::FixDuplicateIDs()
 				UE_LOG(LogStateTreeEditor, Log, TEXT("%s: Found GlobalTask '%s' with duplicate ID, changing ID:%s to ID:%s."),
 					*GetFullName(), *Node.GetName().ToString(), *OldID.ToString(), *Node.ID.ToString());
 				EditorBindings.CopyBindings(OldID, Node.ID);
+				bDirty = true;
 			}
 			FoundNodeIDs.Add(Node.ID);
 		}
 	}
 	
-	VisitHierarchy([&FoundNodeIDs, &EditorBindings = EditorBindings, &Self = *this](UStateTreeState& State, UStateTreeState* ParentState)
+	VisitHierarchy([&bDirty, &FoundNodeIDs, &EditorBindings = EditorBindings, &Self = *this](UStateTreeState& State, UStateTreeState* ParentState)
 	{
 		// Enter conditions
 		for (int32 Index = 0; Index < State.EnterConditions.Num(); Index++)
@@ -998,6 +1063,29 @@ void UStateTreeEditorData::FixDuplicateIDs()
 					UE_LOG(LogStateTreeEditor, Log, TEXT("%s: Found Enter Condition '%s' with duplicate ID on state '%s', changing ID:%s to ID:%s."),
 						*Self.GetFullName(), *Node.GetName().ToString(), *GetNameSafe(&State), *OldID.ToString(), *Node.ID.ToString());
 					EditorBindings.CopyBindings(OldID, Node.ID);
+					bDirty = true;
+				}
+			}
+		}
+
+		// Consideration Utilities
+		for (int32 Index = 0; Index < State.Considerations.Num(); Index++)
+		{
+			FStateTreeEditorNode& Node = State.Considerations[Index];
+			if (const FStateTreeConsiderationBase* Cond = Node.Node.GetPtr<FStateTreeConsiderationBase>())
+			{
+				const FGuid OldID = Node.ID;
+				
+				bool bIsAlreadyInSet = false;
+				FoundNodeIDs.Add(Node.ID, &bIsAlreadyInSet);
+				if (bIsAlreadyInSet)
+				{
+					Node.ID = UE::StateTree::PropertyHelpers::MakeDeterministicID(State, TEXT("Consideration"), Index);
+					
+					UE_LOG(LogStateTreeEditor, Log, TEXT("%s: Found Consideration '%s' with duplicate ID on state '%s', changing ID:%s to ID:%s."),
+						*Self.GetFullName(), *Node.GetName().ToString(), *GetNameSafe(&State), *OldID.ToString(), *Node.ID.ToString());
+					EditorBindings.CopyBindings(OldID, Node.ID);
+					bDirty = true;
 				}
 			}
 		}
@@ -1019,6 +1107,7 @@ void UStateTreeEditorData::FixDuplicateIDs()
 					UE_LOG(LogStateTreeEditor, Log, TEXT("%s: Found Task '%s' with duplicate ID on state '%s', changing ID:%s to ID:%s."),
 						*Self.GetFullName(), *Node.GetName().ToString(), *GetNameSafe(&State), *OldID.ToString(), *Node.ID.ToString());
 					EditorBindings.CopyBindings(OldID, Node.ID);
+					bDirty = true;
 				}
 			}
 		}
@@ -1036,6 +1125,7 @@ void UStateTreeEditorData::FixDuplicateIDs()
 				UE_LOG(LogStateTreeEditor, Log, TEXT("%s: Found enter condition '%s' with duplicate ID on state '%s', changing ID:%s to ID:%s."),
 					*Self.GetFullName(), *State.SingleTask.GetName().ToString(), *GetNameSafe(&State), *OldID.ToString(), *State.SingleTask.ID.ToString());
 				EditorBindings.CopyBindings(OldID, State.SingleTask.ID);
+				bDirty = true;
 			}
 		}
 
@@ -1058,6 +1148,7 @@ void UStateTreeEditorData::FixDuplicateIDs()
 						UE_LOG(LogStateTreeEditor, Log, TEXT("%s: Found transition condition '%s' with duplicate ID on state '%s', changing ID:%s to ID:%s."),
 							*Self.GetFullName(), *Node.GetName().ToString(), *GetNameSafe(&State), *OldID.ToString(), *Node.ID.ToString());
 						EditorBindings.CopyBindings(OldID, Node.ID);
+						bDirty = true;
 					}
 				}
 			}
@@ -1065,6 +1156,11 @@ void UStateTreeEditorData::FixDuplicateIDs()
 		
 		return EStateTreeVisitor::Continue;
 	});
+
+	if (bDirty)
+	{
+		UStateTreeEditingSubsystem::MarkAsModified(GetTypedOuter<UStateTree>());
+	}
 }
 
 void UStateTreeEditorData::DuplicateIDs()
@@ -1104,6 +1200,7 @@ void UStateTreeEditorData::DuplicateIDs()
 
 	// Remap ids properties to the new generated ids
 	{
+		bool bDirty = false;
 		TArray<const UObject*> ObjectToSearch;
 		TSet<const UObject*> ObjectSearched;
 		ObjectToSearch.Add(this);
@@ -1126,7 +1223,11 @@ void UStateTreeEditorData::DuplicateIDs()
 						FGuid* GuidValue = static_cast<FGuid*>(const_cast<void*>(It.Value()));
 						if (FGuid* NewGuidValue = OldToNewIDs.Find(*GuidValue))
 						{
-							*GuidValue = *NewGuidValue;
+							if (*GuidValue != *NewGuidValue)
+							{
+								*GuidValue = *NewGuidValue;
+								bDirty = true;
+							}
 						}
 					}
 				}
@@ -1157,6 +1258,11 @@ void UStateTreeEditorData::DuplicateIDs()
 				}
 			}
 		}
+
+		if (bDirty)
+		{
+			UStateTreeEditingSubsystem::MarkAsModified(GetTypedOuter<UStateTree>());
+		}
 	}
 }
 
