[StateTree] Fix Rewind Debugger track bar extending indefinitely for stopped StateTree instances

OutputInstanceLifetimeEvent(Pop) only emitted TRACE_INSTANCE_LIFETIME_END, which ends the object in the GameplayProvider. It did not emit a StateTreeDebugger InstanceEvent with Pop type, which the StateTree trace provider needs to call Descriptor->Lifetime.SetUpperBound(). Without this, IsInstanceActive() always returned true and the Rewind Debugger track extended the bar to the full recording duration.

This was invisible for UStateTreeComponent because the owning actor lives for the entire recording, so the "forever" bar coincidentally matched. For short-lived owners (e.g. dynamically started/stopped StateTree instances), the bar extended far past the actual execution window.

Fix: emit the StateTreeDebugger InstanceEvent(Pop) before TRACE_INSTANCE_LIFETIME_END so the trace provider receives it and closes the Lifetime range.
#Synced-CL 52073447
#Synced-CL 52073447

#Synced-CL 52073447

diff --git a/Engine/Plugins/Runtime/StateTree/Source/StateTreeModule/Private/Debugger/StateTreeTrace.cpp b/Engine/Plugins/Runtime/StateTree/Source/StateTreeModule/Private/Debugger/StateTreeTrace.cpp
--- a/Engine/Plugins/Runtime/StateTree/Source/StateTreeModule/Private/Debugger/StateTreeTrace.cpp
+++ b/Engine/Plugins/Runtime/StateTree/Source/StateTreeModule/Private/Debugger/StateTreeTrace.cpp
@@ -706,7 +706,17 @@ void OutputInstanceLifetimeEvent(
 		}
 		else if (EventType == EStateTreeTraceEventType::Pop)
 		{
-			TRACE_INSTANCE_LIFETIME_END(Owner, Context->GetInstanceDebugId().ToUint64());
+			const FStateTreeInstanceDebugId InstanceId = Context->GetInstanceDebugId();
+
+			// Emit the StateTree-specific InstanceEvent so the trace provider
+			// sets Lifetime upper bound (making IsInstanceActive() return false).
+			UE_TRACE_LOG(StateTreeDebugger, InstanceEvent, StateTreeDebugChannel)
+				<< InstanceEvent.Cycle(FPlatformTime::Cycles64())
+				<< InstanceEvent.InstanceId(InstanceId.Id)
+				<< InstanceEvent.InstanceSerial(InstanceId.SerialNumber)
+				<< InstanceEvent.EventType(static_cast<std::underlying_type_t<EStateTreeTraceEventType>>(EStateTreeTraceEventType::Pop));
+
+			TRACE_INSTANCE_LIFETIME_END(Owner, InstanceId.ToUint64());
 		}
 	}
 }
