[StateTree] Replaces direct FStateTreeEditorNode field reads (.Node, .Instance, .InstanceObject, .ExecutionRuntimeData, .ExecutionRuntimeDataObject) with the accessor methods.
Fix missing cases where GetExecutionRuntimeData( was not handled properly.

#Synced-CL 52002607

diff --git a/Engine/Plugins/Runtime/StateTree/Source/StateTreeEditorModule/Private/AsyncStateTreeDiff.cpp b/Engine/Plugins/Runtime/StateTree/Source/StateTreeEditorModule/Private/AsyncStateTreeDiff.cpp
--- a/Engine/Plugins/Runtime/StateTree/Source/StateTreeEditorModule/Private/AsyncStateTreeDiff.cpp
+++ b/Engine/Plugins/Runtime/StateTree/Source/StateTreeEditorModule/Private/AsyncStateTreeDiff.cpp
@@ -41,11 +41,35 @@ static bool AreObjectsEqual(const UObject* ObjectA, const UObject* ObjectB)
 
 static bool AreNodesEqual(const FStateTreeEditorNode& NodeA, const FStateTreeEditorNode& NodeB)
 {
-	return AreObjectsEqual(NodeA.InstanceObject.Get(), NodeB.InstanceObject.Get())
-		&& AreObjectsEqual(NodeA.ExecutionRuntimeDataObject.Get(), NodeB.ExecutionRuntimeDataObject.Get())
-		&& NodeA.Node.Identical(&NodeB.Node, PPF_DeepComparison | PPF_ForDiff)
-		&& NodeA.Instance.Identical(&NodeB.Instance, PPF_DeepComparison | PPF_ForDiff)
-		&& NodeA.ExecutionRuntimeData.Identical(&NodeB.ExecutionRuntimeData, PPF_DeepComparison | PPF_ForDiff)
+	auto CompareDataView = [](FStateTreeDataView ViewA, FStateTreeDataView ViewB)
+		{
+			bool bInstanceEqual = false;
+			if (ViewA.GetStruct() == ViewB.GetStruct())
+			{
+				if (Cast<UClass>(ViewA.GetStruct()))
+				{
+					bInstanceEqual = AreObjectsEqual(static_cast<const UObject*>(ViewA.GetMemory()), static_cast<const UObject*>(ViewB.GetMemory()));
+				}
+				else if (const UScriptStruct* ScriptStruct = Cast<const UScriptStruct>(ViewA.GetStruct()))
+				{
+					bInstanceEqual = ScriptStruct->CompareScriptStruct(ViewA.GetMemory(), ViewB.GetMemory(), PPF_DeepComparison | PPF_ForDiff);
+				}
+			}
+			return bInstanceEqual;
+		};
+	auto CompareStructView = [](FStructView ViewA, FStructView ViewB)
+		{
+			bool bEqual = false;
+			if (ViewA.GetScriptStruct() == ViewB.GetScriptStruct())
+			{
+				bEqual = ViewA.GetScriptStruct() == nullptr || ViewA.GetScriptStruct()->CompareScriptStruct(ViewA.GetMemory(), ViewB.GetMemory(), PPF_DeepComparison | PPF_ForDiff);
+			}
+			return bEqual;
+		};
+
+	return CompareDataView(NodeA.GetInstance(), NodeB.GetInstance())
+		&& CompareDataView(NodeA.GetExecutionRuntimeData(), NodeB.GetExecutionRuntimeData())
+		&& CompareStructView(NodeA.GetNode(), NodeB.GetNode())
 		&& NodeA.ExpressionIndent == NodeB.ExpressionIndent
 		&& NodeA.ExpressionOperand == NodeB.ExpressionOperand;
 }
@@ -243,7 +267,8 @@ static FPropertySoftPath GetPropertyPath(const FPropertyBindingPath& StateTreePr
 			{
 				Path.Add(PathSegmentName);
 				Path.Add(FName(FString::FromInt(i)));
-				if (List[i].InstanceObject)
+				const FStateTreeDataView Instance = List[i].GetInstance();
+				if (Cast<UClass>(Instance.GetStruct()))
 				{
 					Path.Add(FName("InstanceObject"));
 				}
diff --git a/Engine/Plugins/Runtime/StateTree/Source/StateTreeEditorModule/Private/Customizations/StateTreeBindingExtension.cpp b/Engine/Plugins/Runtime/StateTree/Source/StateTreeEditorModule/Private/Customizations/StateTreeBindingExtension.cpp
--- a/Engine/Plugins/Runtime/StateTree/Source/StateTreeEditorModule/Private/Customizations/StateTreeBindingExtension.cpp
+++ b/Engine/Plugins/Runtime/StateTree/Source/StateTreeEditorModule/Private/Customizations/StateTreeBindingExtension.cpp
@@ -600,7 +600,7 @@ struct FStateTreeCachedBindingData : UE::PropertyBinding::FCachedBindingData, TS
 	virtual bool GetPropertyFunctionText(FConstStructView InPropertyFunctionStructView, FText& OutText) const override
 	{
 		const FStateTreeEditorNode& EditorNode = InPropertyFunctionStructView.Get<const FStateTreeEditorNode>();
-		if (const FStateTreeNodeBase* Node = EditorNode.Node.GetPtr<const FStateTreeNodeBase>())
+		if (const FStateTreeNodeBase* Node = EditorNode.GetNode().GetPtr<const FStateTreeNodeBase>())
 		{
 			IStateTreeEditorPropertyBindingsOwner* BindingOwner = Cast<IStateTreeEditorPropertyBindingsOwner>(GetOwner());
 			const FText Description = Node->GetDescription(GetSourcePath().GetStructID(), EditorNode.GetInstance(), FStateTreeBindingLookup(BindingOwner), EStateTreeNodeFormatting::Text);
@@ -616,7 +616,7 @@ struct FStateTreeCachedBindingData : UE::PropertyBinding::FCachedBindingData, TS
 	virtual bool GetPropertyFunctionTooltipText(FConstStructView InPropertyFunctionStructView, FText& OutText) const override
 	{
 		const FStateTreeEditorNode& EditorNode = InPropertyFunctionStructView.Get<const FStateTreeEditorNode>();
-		if (const FStateTreeNodeBase* Node = EditorNode.Node.GetPtr<const FStateTreeNodeBase>())
+		if (const FStateTreeNodeBase* Node = EditorNode.GetNode().GetPtr<const FStateTreeNodeBase>())
 		{
 			IStateTreeEditorPropertyBindingsOwner* BindingOwner = Cast<IStateTreeEditorPropertyBindingsOwner>(GetOwner());
 			const FText Description = Node->GetDescription(GetSourcePath().GetStructID(), EditorNode.GetInstance(), FStateTreeBindingLookup(BindingOwner), EStateTreeNodeFormatting::Text);
@@ -632,7 +632,7 @@ struct FStateTreeCachedBindingData : UE::PropertyBinding::FCachedBindingData, TS
 	virtual bool GetPropertyFunctionIconColor(FConstStructView InPropertyFunctionStructView, FLinearColor& OutColor) const override
 	{
 		const FStateTreeEditorNode& EditorNode = InPropertyFunctionStructView.Get<const FStateTreeEditorNode>();
-		if (const FStateTreeNodeBase* Node = EditorNode.Node.GetPtr<const FStateTreeNodeBase>())
+		if (const FStateTreeNodeBase* Node = EditorNode.GetNode().GetPtr<const FStateTreeNodeBase>())
 		{
 			if (GetStructSingleOutputProperty(*Node->GetInstanceDataType()))
 			{
@@ -646,7 +646,7 @@ struct FStateTreeCachedBindingData : UE::PropertyBinding::FCachedBindingData, TS
 	virtual bool GetPropertyFunctionImage(FConstStructView InPropertyFunctionStructView, const FSlateBrush*& OutImage) const override
 	{
 		const FStateTreeEditorNode& EditorNode = InPropertyFunctionStructView.Get<const FStateTreeEditorNode>();
-		if (const FStateTreeNodeBase* Node = EditorNode.Node.GetPtr<const FStateTreeNodeBase>())
+		if (const FStateTreeNodeBase* Node = EditorNode.GetNode().GetPtr<const FStateTreeNodeBase>())
 		{
 			if (GetStructSingleOutputProperty(*Node->GetInstanceDataType()))
 			{
@@ -735,7 +735,7 @@ private:
 				if (EditorNodeView.IsValid())
 				{
 					const FStateTreeEditorNode& EditorNode = EditorNodeView.Get<FStateTreeEditorNode>();
-					if (EditorNode.Node.IsValid() && EditorNode.Instance.IsValid())
+					if (EditorNode.GetNode().IsValid() && EditorNode.GetInstance().IsValid())
 					{
 						return EditorNodeView;
 					}
diff --git a/Engine/Plugins/Runtime/StateTree/Source/StateTreeEditorModule/Private/Customizations/StateTreeEditorNodeDetails.cpp b/Engine/Plugins/Runtime/StateTree/Source/StateTreeEditorModule/Private/Customizations/StateTreeEditorNodeDetails.cpp
--- a/Engine/Plugins/Runtime/StateTree/Source/StateTreeEditorModule/Private/Customizations/StateTreeEditorNodeDetails.cpp
+++ b/Engine/Plugins/Runtime/StateTree/Source/StateTreeEditorModule/Private/Customizations/StateTreeEditorNodeDetails.cpp
@@ -61,7 +61,7 @@ namespace UE::StateTreeEditor::Internal
 				{
 					if (const FStateTreeEditorNode* Node = UE::StateTreeEditor::EditorNodeUtils::GetCommonNode(Property))
 					{
-						if (const UScriptStruct* ScriptStruct = Node->Node.GetScriptStruct())
+						if (const UScriptStruct* ScriptStruct = Node->GetNode().GetScriptStruct())
 						{
 							return ScriptStruct->IsChildOf<FStateTreePropertyFunctionBase>();
 						}
@@ -722,7 +722,7 @@ void FStateTreeEditorNodeDetails::CustomizeHeader(TSharedRef<class IPropertyHand
 	{
 		if (const FStateTreeEditorNode* Node = UE::StateTreeEditor::EditorNodeUtils::GetCommonNode(StructProperty))
 		{
-			if (const FStateTreeTaskBase* TaskBase = Node->Node.GetPtr<FStateTreeTaskBase>())
+			if (const FStateTreeTaskBase* TaskBase = Node->GetNode().GetPtr<FStateTreeTaskBase>())
 			{
 				DescriptionBox->InsertSlot(0)
 				.AutoWidth()
@@ -788,7 +788,7 @@ void FStateTreeEditorNodeDetails::MakeFlagsWidget()
 
 			if (bUseEditorData)
 			{
-				if (const FStateTreeTaskBase* TreeTaskNodePtr = EditorNode->Node.GetPtr<FStateTreeTaskBase>())
+				if (const FStateTreeTaskBase* TreeTaskNodePtr = EditorNode->GetNode().GetPtr<FStateTreeTaskBase>())
 				{
 					if (TreeTaskNodePtr->bConsideredForScheduling)
 					{
@@ -1079,7 +1079,7 @@ bool FStateTreeEditorNodeDetails::ShouldResetToDefault(TSharedPtr<IPropertyHandl
 	{
 		if (const FStateTreeEditorNode* Node = static_cast<const FStateTreeEditorNode*>(Data))
 		{
-			if (Node->Node.IsValid())
+			if (Node->GetNode().IsValid())
 			{
 				bAnyValid = true;
 				break;
@@ -1256,7 +1256,7 @@ void FStateTreeEditorNodeDetails::OnBindingChanged(const FPropertyBindingPath& S
 		UObject* OuterObject = OuterObjects[i]; // Immediate outer, i.e StateTreeState
 		if (EditorNode && OuterObject && EditorNode->ID == TargetPath.GetStructID())
 		{
-			FStateTreeNodeBase* Node = EditorNode->Node.GetMutablePtr<FStateTreeNodeBase>();
+			FStateTreeNodeBase* Node = EditorNode->GetNode().GetPtr<FStateTreeNodeBase>();
 			FStateTreeDataView InstanceView = EditorNode->GetInstance(); 
 
 			if (Node && InstanceView.IsValid())
@@ -1736,7 +1736,7 @@ EVisibility FStateTreeEditorNodeDetails::IsNodeDescriptionVisible() const
 	const UScriptStruct* ScriptStruct = nullptr;
 	if (const FStateTreeEditorNode* Node = UE::StateTreeEditor::EditorNodeUtils::GetCommonNode(StructProperty))
 	{
-		ScriptStruct = Node->Node.GetScriptStruct();
+		ScriptStruct = Node->GetNode().GetScriptStruct();
 	}
 
 	if (ScriptStruct != nullptr && ScriptStruct->IsChildOf(FStateTreeTaskBase::StaticStruct()))
@@ -1776,7 +1776,7 @@ FText FStateTreeEditorNodeDetails::GetNodeTooltip() const
 			const UStruct* Struct = Node->GetInstance().GetStruct();
 			if (Struct == nullptr || !Struct->IsChildOf<UStateTreeNodeBlueprintBase>())
 			{
-				Struct = Node->Node.GetScriptStruct();
+				Struct = Node->GetNode().GetScriptStruct();
 			}
 
 			if (Struct)
@@ -1816,7 +1816,7 @@ FText FStateTreeEditorNodeDetails::GetName() const
 	{
 		if (const FStateTreeEditorNode* Node = static_cast<FStateTreeEditorNode*>(RawNodeData[0]))
 		{
-			if (const FStateTreeNodeBase* BaseNode = Node->Node.GetPtr<FStateTreeNodeBase>())
+			if (const FStateTreeNodeBase* BaseNode = Node->GetNode().GetPtr<FStateTreeNodeBase>())
 			{
 				if (!BaseNode->Name.IsNone())
 				{
@@ -1871,7 +1871,7 @@ void FStateTreeEditorNodeDetails::HandleNameCommitted(const FText& NewText, ETex
 				// Set Name
 				if (FStateTreeEditorNode* Node = static_cast<FStateTreeEditorNode*>(Data))
 				{
-					if (FStateTreeNodeBase* BaseNode = Node->Node.GetMutablePtr<FStateTreeNodeBase>())
+					if (FStateTreeNodeBase* BaseNode = Node->GetNode().GetPtr<FStateTreeNodeBase>())
 					{
 						BaseNode->Name = FName(NewName);
 					}
@@ -2006,17 +2006,17 @@ FText FStateTreeEditorNodeDetails::GetNodePickerTooltip() const
 	// Text describing the type.
 	if (const FStateTreeEditorNode* Node = UE::StateTreeEditor::EditorNodeUtils::GetCommonNode(StructProperty))
 	{
-		if (const UScriptStruct* ScriptStruct = Node->Node.GetScriptStruct())
+		if (const UScriptStruct* ScriptStruct = Node->GetNode().GetScriptStruct())
 		{
 			if (ScriptStruct->IsChildOf(FStateTreeBlueprintEvaluatorWrapper::StaticStruct())
 				|| ScriptStruct->IsChildOf(FStateTreeBlueprintTaskWrapper::StaticStruct())
 				|| ScriptStruct->IsChildOf(FStateTreeBlueprintConditionWrapper::StaticStruct())
 				|| ScriptStruct->IsChildOf(FStateTreeBlueprintConsiderationWrapper::StaticStruct()))
 			{
-				if (Node->InstanceObject != nullptr
-					&& Node->InstanceObject->GetClass() != nullptr)
+				FStateTreeDataView InstanceView = Node->GetInstance();
+				if (InstanceView.GetStruct())
 				{
-					TextBuilder.AppendLine(Node->InstanceObject->GetClass()->GetDisplayNameText());
+					TextBuilder.AppendLine(InstanceView.GetStruct()->GetDisplayNameText());
 				}
 			}
 			else
@@ -2128,16 +2128,16 @@ void FStateTreeEditorNodeDetails::GeneratePickerMenu(class FMenuBuilder& InMenuB
 	const UStruct* CommonStruct  = nullptr;
 	if (const FStateTreeEditorNode* Node = UE::StateTreeEditor::EditorNodeUtils::GetCommonNode(StructProperty))
 	{
-		if (const UScriptStruct* ScriptStruct = Node->Node.GetScriptStruct())
+		if (const UScriptStruct* ScriptStruct = Node->GetNode().GetScriptStruct())
 		{
 			if (ScriptStruct->IsChildOf(FStateTreeBlueprintEvaluatorWrapper::StaticStruct())
 				|| ScriptStruct->IsChildOf(FStateTreeBlueprintTaskWrapper::StaticStruct())
 				|| ScriptStruct->IsChildOf(FStateTreeBlueprintConditionWrapper::StaticStruct())
 				|| ScriptStruct->IsChildOf(FStateTreeBlueprintConsiderationWrapper::StaticStruct()))
 			{
-				if (Node->InstanceObject != nullptr)
+				if (Node->GetInstance().IsValid())
 				{
-					CommonStruct = Node->InstanceObject->GetClass();
+					CommonStruct = Node->GetInstance().GetStruct();
 				}
 			}
 			else
diff --git a/Engine/Plugins/Runtime/StateTree/Source/StateTreeEditorModule/Private/Customizations/StateTreeEditorNodeUtils.cpp b/Engine/Plugins/Runtime/StateTree/Source/StateTreeEditorModule/Private/Customizations/StateTreeEditorNodeUtils.cpp
--- a/Engine/Plugins/Runtime/StateTree/Source/StateTreeEditorModule/Private/Customizations/StateTreeEditorNodeUtils.cpp
+++ b/Engine/Plugins/Runtime/StateTree/Source/StateTreeEditorModule/Private/Customizations/StateTreeEditorNodeUtils.cpp
@@ -33,7 +33,7 @@ EStateTreeConditionEvaluationMode GetConditionEvaluationMode(const TSharedPtr<IP
 {
 	if (const FStateTreeEditorNode* Node = GetCommonNode(StructProperty))
 	{
-		if (const FStateTreeConditionBase* ConditionBase = Node->Node.GetPtr<FStateTreeConditionBase>())
+		if (const FStateTreeConditionBase* ConditionBase = Node->GetNode().GetPtr<FStateTreeConditionBase>())
 		{
 			return ConditionBase->EvaluationMode;
 		}
@@ -54,7 +54,7 @@ bool IsTaskDisabled(const TSharedPtr<IPropertyHandle>& StructProperty)
 
 bool IsTaskEnabled(const FStateTreeEditorNode& EditorNode)
 {
-	if (const FStateTreeTaskBase* TaskBase = EditorNode.Node.GetPtr<FStateTreeTaskBase>())
+	if (const FStateTreeTaskBase* TaskBase = EditorNode.GetNode().GetPtr<FStateTreeTaskBase>())
 	{
 		return TaskBase->bTaskEnabled;
 	}
@@ -64,14 +64,14 @@ bool IsTaskEnabled(const FStateTreeEditorNode& EditorNode)
 bool IsTaskConsideredForCompletion(const FStateTreeEditorNode& EditorNode)
 {
 	// We use the Blueprint flag to have a default value that behaves like the other flags. Sadly, it duplicates the flags.
-	if (const FStateTreeBlueprintTaskWrapper* TaskWrapper = EditorNode.Node.GetPtr<FStateTreeBlueprintTaskWrapper>())
+	if (const FStateTreeBlueprintTaskWrapper* TaskWrapper = EditorNode.GetNode().GetPtr<FStateTreeBlueprintTaskWrapper>())
 	{
-		if (UStateTreeTaskBlueprintBase* BPTaskBase = Cast<UStateTreeTaskBlueprintBase>(EditorNode.InstanceObject))
+		if (const UStateTreeTaskBlueprintBase* BPTaskBase = EditorNode.GetInstance().GetPtr<const UStateTreeTaskBlueprintBase>())
 		{
 			return BPTaskBase->bConsideredForCompletion;
 		}
 	}
-	else if (const FStateTreeTaskBase* TaskBase = EditorNode.Node.GetPtr<FStateTreeTaskBase>())
+	else if (const FStateTreeTaskBase* TaskBase = EditorNode.GetNode().GetPtr<FStateTreeTaskBase>())
 	{
 		return TaskBase->bConsideredForCompletion;
 	}
@@ -80,14 +80,14 @@ bool IsTaskConsideredForCompletion(const FStateTreeEditorNode& EditorNode)
 
 void SetTaskConsideredForCompletion(FStateTreeEditorNode& EditorNode, bool bIsConsidered)
 {
-	if (const FStateTreeBlueprintTaskWrapper* TaskWrapper = EditorNode.Node.GetPtr<FStateTreeBlueprintTaskWrapper>())
+	if (const FStateTreeBlueprintTaskWrapper* TaskWrapper = EditorNode.GetNode().GetPtr<FStateTreeBlueprintTaskWrapper>())
 	{
-		if (UStateTreeTaskBlueprintBase* BPTaskBase = Cast<UStateTreeTaskBlueprintBase>(EditorNode.InstanceObject))
+		if (UStateTreeTaskBlueprintBase* BPTaskBase = EditorNode.GetInstance().GetMutablePtr<UStateTreeTaskBlueprintBase>())
 		{
 			BPTaskBase->bConsideredForCompletion = bIsConsidered;
 		}
 	}
-	else if (FStateTreeTaskBase* TaskBase = EditorNode.Node.GetMutablePtr<FStateTreeTaskBase>())
+	else if (FStateTreeTaskBase* TaskBase = EditorNode.GetNode().GetPtr<FStateTreeTaskBase>())
 	{
 		TaskBase->bConsideredForCompletion = bIsConsidered;
 	}
@@ -95,14 +95,14 @@ void SetTaskConsideredForCompletion(FStateTreeEditorNode& EditorNode, bool bIsCo
 
 bool CanEditTaskConsideredForCompletion(const FStateTreeEditorNode& EditorNode)
 {
-	if (const FStateTreeBlueprintTaskWrapper* TaskWrapper = EditorNode.Node.GetPtr<FStateTreeBlueprintTaskWrapper>())
+	if (const FStateTreeBlueprintTaskWrapper* TaskWrapper = EditorNode.GetNode().GetPtr<FStateTreeBlueprintTaskWrapper>())
 	{
-		if (UStateTreeTaskBlueprintBase* BPTaskBase = Cast<UStateTreeTaskBlueprintBase>(EditorNode.InstanceObject))
+		if (const UStateTreeTaskBlueprintBase* BPTaskBase = EditorNode.GetInstance().GetPtr<UStateTreeTaskBlueprintBase>())
 		{
 			return BPTaskBase->bCanEditConsideredForCompletion;
 		}
 	}
-	else if (const FStateTreeTaskBase* TaskBase = EditorNode.Node.GetPtr<FStateTreeTaskBase>())
+	else if (const FStateTreeTaskBase* TaskBase = EditorNode.GetNode().GetPtr<FStateTreeTaskBase>())
 	{
 		return TaskBase->bCanEditConsideredForCompletion;
 	}
@@ -128,7 +128,7 @@ EVisibility IsConditionVisible(const TSharedPtr<IPropertyHandle>& StructProperty
 	const UScriptStruct* ScriptStruct = nullptr;
 	if (const FStateTreeEditorNode* Node = GetCommonNode(StructProperty))
 	{
-		ScriptStruct = Node->Node.GetScriptStruct();
+		ScriptStruct = Node->GetNode().GetScriptStruct();
 	}
 
 	return ScriptStruct != nullptr && ScriptStruct->IsChildOf(FStateTreeConditionBase::StaticStruct()) ? EVisibility::Visible : EVisibility::Collapsed;
@@ -139,7 +139,7 @@ EVisibility IsConsiderationVisible(const TSharedPtr<IPropertyHandle>& StructProp
 	const UScriptStruct* ScriptStruct = nullptr;
 	if (const FStateTreeEditorNode* Node = GetCommonNode(StructProperty))
 	{
-		ScriptStruct = Node->Node.GetScriptStruct();
+		ScriptStruct = Node->GetNode().GetScriptStruct();
 	}
 
 	return ScriptStruct != nullptr && ScriptStruct->IsChildOf(FStateTreeConsiderationBase::StaticStruct()) ? EVisibility::Visible : EVisibility::Collapsed;
@@ -149,7 +149,7 @@ FName GetNodeIconName(const TSharedPtr<IPropertyHandle>& StructProperty)
 {
 	if (const FStateTreeEditorNode* Node = GetCommonNode(StructProperty))
 	{
-		if (const FStateTreeNodeBase* BaseNode = Node->Node.GetPtr<const FStateTreeNodeBase>())
+		if (const FStateTreeNodeBase* BaseNode = Node->GetNode().GetPtr<const FStateTreeNodeBase>())
 		{
 			return BaseNode->GetIconName();
 		}
@@ -204,7 +204,7 @@ FSlateColor GetIconColor(const TSharedPtr<IPropertyHandle>& StructProperty)
 {
 	if (const FStateTreeEditorNode* Node = GetCommonNode(StructProperty))
 	{
-		if (const FStateTreeNodeBase* BaseNode = Node->Node.GetPtr<const FStateTreeNodeBase>())
+		if (const FStateTreeNodeBase* BaseNode = Node->GetNode().GetPtr<const FStateTreeNodeBase>())
 		{
 			return FLinearColor(BaseNode->GetIconColor());
 		}
@@ -315,36 +315,20 @@ struct FNodeRetainPropertyData
 FNodeRetainPropertyData GetNodeData(FStateTreeEditorNode& EditorNode)
 {
 	FNodeRetainPropertyData Data;
-	if (FStateTreeNodeBase* NodeBase = EditorNode.Node.GetMutablePtr<FStateTreeNodeBase>(); NodeBase != nullptr)
+	if (FStateTreeNodeBase* NodeBase = EditorNode.GetNode().GetPtr<FStateTreeNodeBase>())
 	{
-		Data.Node.Struct = EditorNode.Node.GetScriptStruct();
+		Data.Node.Struct = EditorNode.GetNode().GetScriptStruct();
 		Data.Node.Memory = NodeBase;
 
 		if (const UStruct* InstanceDataType = NodeBase->GetInstanceDataType())
 		{
-			if (InstanceDataType->IsA<UScriptStruct>())
-			{
-				Data.InstanceData.Struct = EditorNode.Instance.GetScriptStruct();
-				Data.InstanceData.Memory = EditorNode.Instance.GetMutableMemory();
-			}
-			else if (InstanceDataType->IsA<UClass>())
-			{
-				Data.InstanceData.Struct = EditorNode.InstanceObject.GetClass();
-				Data.InstanceData.Memory = EditorNode.InstanceObject;
-			}
+			Data.InstanceData.Struct = EditorNode.GetInstance().GetStruct();
+			Data.InstanceData.Memory = EditorNode.GetInstance().GetMutableMemory();
 		}
 		if (const UStruct* ExecutionRuntimeInstanceDataType = NodeBase->GetExecutionRuntimeDataType())
 		{
-			if (ExecutionRuntimeInstanceDataType->IsA<UScriptStruct>())
-			{
-				Data.ExecutionRuntimeData.Struct = EditorNode.ExecutionRuntimeData.GetScriptStruct();
-				Data.ExecutionRuntimeData.Memory = EditorNode.ExecutionRuntimeData.GetMutableMemory();
-			}
-			else if (ExecutionRuntimeInstanceDataType->IsA<UClass>())
-			{
-				Data.ExecutionRuntimeData.Struct = EditorNode.ExecutionRuntimeDataObject.GetClass();
-				Data.ExecutionRuntimeData.Memory = EditorNode.ExecutionRuntimeDataObject;
-			}
+			Data.ExecutionRuntimeData.Struct = EditorNode.GetExecutionRuntimeData().GetStruct();
+			Data.ExecutionRuntimeData.Memory = EditorNode.GetExecutionRuntimeData().GetMutableMemory();
 		}
 	}
 
@@ -661,7 +645,7 @@ void InstantiateStructSubobjects(UObject& OuterObject, FStructView Struct)
 
 void ConditionalUpdateNodeInstanceData(FStateTreeEditorNode& EditorNode, UObject& InstanceOuter)
 {
-	const FStateTreeNodeBase* Node = EditorNode.Node.GetPtr<FStateTreeNodeBase>();
+	const FStateTreeNodeBase* Node = EditorNode.GetNode().GetPtr<FStateTreeNodeBase>();
 	if (!Node)
 	{
 		return;
diff --git a/Engine/Plugins/Runtime/StateTree/Source/StateTreeEditorModule/Private/Debugger/StateTreeDebuggerUIExtensions.cpp b/Engine/Plugins/Runtime/StateTree/Source/StateTreeEditorModule/Private/Debugger/StateTreeDebuggerUIExtensions.cpp
--- a/Engine/Plugins/Runtime/StateTree/Source/StateTreeEditorModule/Private/Debugger/StateTreeDebuggerUIExtensions.cpp
+++ b/Engine/Plugins/Runtime/StateTree/Source/StateTreeEditorModule/Private/Debugger/StateTreeDebuggerUIExtensions.cpp
@@ -30,9 +30,9 @@ void OnConditionEvaluationModeChanged(TSharedPtr<IPropertyHandle> StructProperty
 		StructProperty,
 		[Mode](const TSharedPtr<IPropertyHandle>& StructProperty)
 		{
-			if (FStateTreeEditorNode* Node = EditorNodeUtils::GetMutableCommonNode(StructProperty))
+			if (FStateTreeEditorNode* EditorNode = EditorNodeUtils::GetMutableCommonNode(StructProperty))
 			{
-				if (FStateTreeConditionBase* ConditionBase = Node->Node.GetMutablePtr<FStateTreeConditionBase>())
+				if (FStateTreeConditionBase* ConditionBase = EditorNode->GetNode().GetPtr<FStateTreeConditionBase>())
 				{
 					ConditionBase->EvaluationMode = Mode;
 				}
@@ -46,9 +46,9 @@ void OnTaskEnableToggled(TSharedPtr<IPropertyHandle> StructProperty)
 		StructProperty,
 		[](const TSharedPtr<IPropertyHandle>& StructProperty)
 		{
-			if (FStateTreeEditorNode* Node = EditorNodeUtils::GetMutableCommonNode(StructProperty))
+			if (FStateTreeEditorNode* EditorNode = EditorNodeUtils::GetMutableCommonNode(StructProperty))
 			{
-				if (FStateTreeTaskBase* TaskBase = Node->Node.GetMutablePtr<FStateTreeTaskBase>())
+				if (FStateTreeTaskBase* TaskBase = EditorNode->GetNode().GetPtr<FStateTreeTaskBase>())
 				{
 					TaskBase->bTaskEnabled = !TaskBase->bTaskEnabled;
 				}
@@ -477,9 +477,9 @@ TSharedRef<SWidget> CreateEditorNodeWidget(const TSharedPtr<IPropertyHandle>& St
 		if (TmpStructProperty->Struct->IsChildOf(TBaseStructure<FStateTreeEditorNode>::Get()))
 		{
 			const UScriptStruct* ScriptStruct = nullptr;
-			if (const FStateTreeEditorNode* Node = GetCommonNode(StructPropertyHandle))
+			if (const FStateTreeEditorNode* EditorNode = GetCommonNode(StructPropertyHandle))
 			{
-				ScriptStruct = Node->Node.GetScriptStruct();
+				ScriptStruct = EditorNode->GetNode().GetScriptStruct();
 			}
 
 			if (ScriptStruct != nullptr)
@@ -633,9 +633,9 @@ void AppendEditorNodeMenuItems(FMenuBuilder& InMenuBuilder, const TSharedPtr<IPr
 		if (TmpStructProperty->Struct->IsChildOf(TBaseStructure<FStateTreeEditorNode>::Get()))
 		{
 			const UScriptStruct* ScriptStruct = nullptr;
-			if (const FStateTreeEditorNode* Node = GetCommonNode(StructPropertyHandle))
+			if (const FStateTreeEditorNode* EditorNode = GetCommonNode(StructPropertyHandle))
 			{
-				ScriptStruct = Node->Node.GetScriptStruct();
+				ScriptStruct = EditorNode->GetNode().GetScriptStruct();
 			}
 
 			if (ScriptStruct != nullptr)
@@ -796,7 +796,7 @@ bool IsEditorNodeEnabled(const TSharedPtr<IPropertyHandle>& StructPropertyHandle
 
 	if (const FStateTreeEditorNode* Node = GetCommonNode(StructPropertyHandle))
 	{
-		if (const FStateTreeTaskBase* TaskBase = Node->Node.GetPtr<FStateTreeTaskBase>())
+		if (const FStateTreeTaskBase* TaskBase = Node->GetNode().GetPtr<FStateTreeTaskBase>())
 		{
 			return TaskBase->bTaskEnabled;
 		}
diff --git a/Engine/Plugins/Runtime/StateTree/Source/StateTreeEditorModule/Private/FindTools/SStateTreeFind.cpp b/Engine/Plugins/Runtime/StateTree/Source/StateTreeEditorModule/Private/FindTools/SStateTreeFind.cpp
--- a/Engine/Plugins/Runtime/StateTree/Source/StateTreeEditorModule/Private/FindTools/SStateTreeFind.cpp
+++ b/Engine/Plugins/Runtime/StateTree/Source/StateTreeEditorModule/Private/FindTools/SStateTreeFind.cpp
@@ -375,7 +375,7 @@ void SFindInAsset::MakeQueryDelayed()
 							check(EditorNode);
 							const FSlateBrush* Brush = nullptr;
 							FSlateColor BrushColor = FSlateColor::UseForeground();
-							if (const FStateTreeNodeBase* BaseNode = EditorNode->Node.GetPtr<const FStateTreeNodeBase>())
+							if (const FStateTreeNodeBase* BaseNode = EditorNode->GetNode().GetPtr<const FStateTreeNodeBase>())
 							{
 								Brush = UE::StateTreeEditor::EditorNodeUtils::ParseIcon(BaseNode->GetIconName()).GetIcon();
 								BrushColor = FLinearColor(BaseNode->GetIconColor());
diff --git a/Engine/Plugins/Runtime/StateTree/Source/StateTreeEditorModule/Private/SStateTreeViewRow.cpp b/Engine/Plugins/Runtime/StateTree/Source/StateTreeEditorModule/Private/SStateTreeViewRow.cpp
--- a/Engine/Plugins/Runtime/StateTree/Source/StateTreeEditorModule/Private/SStateTreeViewRow.cpp
+++ b/Engine/Plugins/Runtime/StateTree/Source/StateTreeEditorModule/Private/SStateTreeViewRow.cpp
@@ -413,7 +413,7 @@ TSharedRef<SWidget> SStateTreeViewRow::MakeTasksWidget(const TSharedPtr<SScrollB
 	for (int32 TaskIndex = 0; TaskIndex < NumTasks; TaskIndex++)
 	{
 		const FStateTreeEditorNode& TaskNode = State->Tasks[TaskIndex];
-		if (const FStateTreeTaskBase* Task = TaskNode.Node.GetPtr<FStateTreeTaskBase>())
+		if (const FStateTreeTaskBase* Task = TaskNode.GetNode().GetPtr<FStateTreeTaskBase>())
 		{
 			const FGuid TaskId = State->Tasks[TaskIndex].ID;
 			auto IsTaskEnabledFunc = [WeakState = WeakState, TaskIndex]
@@ -421,7 +421,7 @@ TSharedRef<SWidget> SStateTreeViewRow::MakeTasksWidget(const TSharedPtr<SScrollB
 					const UStateTreeState* State = WeakState.Get();
 					if (State != nullptr && State->Tasks.IsValidIndex(TaskIndex))
 					{
-						if (const FStateTreeTaskBase* Task = State->Tasks[TaskIndex].Node.GetPtr<FStateTreeTaskBase>())
+						if (const FStateTreeTaskBase* Task = State->Tasks[TaskIndex].GetNode().GetPtr<FStateTreeTaskBase>())
 						{
 							return (State->bEnabled && Task->bTaskEnabled);
 						}
@@ -635,7 +635,7 @@ TSharedRef<SWidget> SStateTreeViewRow::MakeConditionsWidget(const TSharedPtr<SSc
 		for (int32 ConditionIndex = 0; ConditionIndex < NumConditions; ConditionIndex++)
 		{
 			const FStateTreeEditorNode& ConditionNode = State->EnterConditions[ConditionIndex];
-			if (const FStateTreeConditionBase* Condition = ConditionNode.Node.GetPtr<FStateTreeConditionBase>())
+			if (const FStateTreeConditionBase* Condition = ConditionNode.GetNode().GetPtr<FStateTreeConditionBase>())
 			{
 				const FGuid ConditionId = ConditionNode.ID;
 
@@ -650,7 +650,7 @@ TSharedRef<SWidget> SStateTreeViewRow::MakeConditionsWidget(const TSharedPtr<SSc
 						const UStateTreeState* State = WeakState.Get();
 						if (State != nullptr && State->EnterConditions.IsValidIndex(ConditionIndex))
 						{
-							if (const FStateTreeConditionBase* Condition = State->EnterConditions[ConditionIndex].Node.GetPtr<FStateTreeConditionBase>())
+							if (const FStateTreeConditionBase* Condition = State->EnterConditions[ConditionIndex].GetNode().GetPtr<FStateTreeConditionBase>())
 							{
 								return Condition->EvaluationMode != EStateTreeConditionEvaluationMode::Evaluated ? EVisibility::Visible : EVisibility::Hidden;
 							}
@@ -663,7 +663,7 @@ TSharedRef<SWidget> SStateTreeViewRow::MakeConditionsWidget(const TSharedPtr<SSc
 						const UStateTreeState* State = WeakState.Get();
 						if (State != nullptr && State->EnterConditions.IsValidIndex(ConditionIndex))
 						{
-							if (const FStateTreeConditionBase* Condition = State->EnterConditions[ConditionIndex].Node.GetPtr<FStateTreeConditionBase>())
+							if (const FStateTreeConditionBase* Condition = State->EnterConditions[ConditionIndex].GetNode().GetPtr<FStateTreeConditionBase>())
 							{
 								if (Condition->EvaluationMode == EStateTreeConditionEvaluationMode::ForcedTrue)
 								{
@@ -683,7 +683,7 @@ TSharedRef<SWidget> SStateTreeViewRow::MakeConditionsWidget(const TSharedPtr<SSc
 						const UStateTreeState* State = WeakState.Get();
 						if (State != nullptr && State->EnterConditions.IsValidIndex(ConditionIndex))
 						{
-							if (const FStateTreeConditionBase* Condition = State->EnterConditions[ConditionIndex].Node.GetPtr<FStateTreeConditionBase>())
+							if (const FStateTreeConditionBase* Condition = State->EnterConditions[ConditionIndex].GetNode().GetPtr<FStateTreeConditionBase>())
 							{
 								if (Condition->EvaluationMode == EStateTreeConditionEvaluationMode::ForcedTrue)
 								{
@@ -1066,11 +1066,11 @@ TSharedRef<SWidget> SStateTreeViewRow::MakeTransitionWidgetInternal(const EState
 		
 		for (const FStateTreeEditorNode& Task : State->Tasks)
 		{
-			AddLinksFromStruct(FStateTreeDataView(Task.Node.GetScriptStruct(), const_cast<uint8*>(Task.Node.GetMemory())), Task.ID);
+			AddLinksFromStruct(FStateTreeDataView(Task.GetNode().GetScriptStruct(), const_cast<uint8*>(Task.GetNode().GetMemory())), Task.ID);
 			AddLinksFromStruct(Task.GetInstance(), Task.ID);
 		}
 
-		AddLinksFromStruct(FStateTreeDataView(State->SingleTask.Node.GetScriptStruct(), const_cast<uint8*>(State->SingleTask.Node.GetMemory())), State->SingleTask.ID);
+		AddLinksFromStruct(FStateTreeDataView(State->SingleTask.GetNode().GetScriptStruct(), const_cast<uint8*>(State->SingleTask.GetNode().GetMemory())), State->SingleTask.ID);
 		AddLinksFromStruct(State->SingleTask.GetInstance(), State->SingleTask.ID);
 	}
 
@@ -1473,7 +1473,7 @@ EVisibility SStateTreeViewRow::GetTaskIconVisibility(FGuid TaskID) const
 	bool bHasIcon = false;
 	if (const FStateTreeEditorNode* TaskNode = GetTaskNodeByID(TaskID))
 	{
-		if (const FStateTreeNodeBase* BaseNode = TaskNode->Node.GetPtr<const FStateTreeNodeBase>())
+		if (const FStateTreeNodeBase* BaseNode = TaskNode->GetNode().GetPtr<const FStateTreeNodeBase>())
 		{
 			bHasIcon = !BaseNode->GetIconName().IsNone();
 		}
@@ -1485,7 +1485,7 @@ const FSlateBrush* SStateTreeViewRow::GetTaskIcon(FGuid TaskID) const
 {
 	if (const FStateTreeEditorNode* TaskNode = GetTaskNodeByID(TaskID))
 	{
-		if (const FStateTreeNodeBase* BaseNode = TaskNode->Node.GetPtr<const FStateTreeNodeBase>())
+		if (const FStateTreeNodeBase* BaseNode = TaskNode->GetNode().GetPtr<const FStateTreeNodeBase>())
 		{
 			return UE::StateTreeEditor::EditorNodeUtils::ParseIcon(BaseNode->GetIconName()).GetIcon();
 		}
@@ -1497,7 +1497,7 @@ FSlateColor SStateTreeViewRow::GetTaskIconColor(FGuid TaskID) const
 {
 	if (const FStateTreeEditorNode* TaskNode = GetTaskNodeByID(TaskID))
 	{
-		if (const FStateTreeNodeBase* BaseNode = TaskNode->Node.GetPtr<const FStateTreeNodeBase>())
+		if (const FStateTreeNodeBase* BaseNode = TaskNode->GetNode().GetPtr<const FStateTreeNodeBase>())
 		{
 			return FLinearColor(BaseNode->GetIconColor());
 		}
@@ -1545,7 +1545,7 @@ EVisibility SStateTreeViewRow::GetConditionIconVisibility(FGuid ConditionID) con
 	bool bHasIcon = false;
 	if (const FStateTreeEditorNode* Node = GetConditionNodeByID(ConditionID))
 	{
-		if (const FStateTreeNodeBase* BaseNode = Node->Node.GetPtr<const FStateTreeNodeBase>())
+		if (const FStateTreeNodeBase* BaseNode = Node->GetNode().GetPtr<const FStateTreeNodeBase>())
 		{
 			bHasIcon = !BaseNode->GetIconName().IsNone();
 		}
@@ -1557,7 +1557,7 @@ const FSlateBrush* SStateTreeViewRow::GetConditionIcon(FGuid ConditionID) const
 {
 	if (const FStateTreeEditorNode* Node = GetConditionNodeByID(ConditionID))
 	{
-		if (const FStateTreeNodeBase* BaseNode = Node->Node.GetPtr<const FStateTreeNodeBase>())
+		if (const FStateTreeNodeBase* BaseNode = Node->GetNode().GetPtr<const FStateTreeNodeBase>())
 		{
 			return UE::StateTreeEditor::EditorNodeUtils::ParseIcon(BaseNode->GetIconName()).GetIcon();
 		}
@@ -1569,7 +1569,7 @@ FSlateColor SStateTreeViewRow::GetConditionIconColor(FGuid ConditionID) const
 {
 	if (const FStateTreeEditorNode* Node = GetConditionNodeByID(ConditionID))
 	{
-		if (const FStateTreeNodeBase* BaseNode = Node->Node.GetPtr<const FStateTreeNodeBase>())
+		if (const FStateTreeNodeBase* BaseNode = Node->GetNode().GetPtr<const FStateTreeNodeBase>())
 		{
 			return FLinearColor(BaseNode->GetIconColor());
 		}
@@ -1891,11 +1891,11 @@ FText SStateTreeViewRow::GetTransitionsDesc(const EStateTreeTransitionTrigger Tr
 		
 		for (const FStateTreeEditorNode& Task : State->Tasks)
 		{
-			AddLinksFromStruct(FStateTreeDataView(Task.Node.GetScriptStruct(), const_cast<uint8*>(Task.Node.GetMemory())));
+			AddLinksFromStruct(FStateTreeDataView(Task.GetNode().GetScriptStruct(), const_cast<uint8*>(Task.GetNode().GetMemory())));
 			AddLinksFromStruct(Task.GetInstance());
 		}
 
-		AddLinksFromStruct(FStateTreeDataView(State->SingleTask.Node.GetScriptStruct(), const_cast<uint8*>(State->SingleTask.Node.GetMemory())));
+		AddLinksFromStruct(FStateTreeDataView(State->SingleTask.GetNode().GetScriptStruct(), const_cast<uint8*>(State->SingleTask.GetNode().GetMemory())));
 		AddLinksFromStruct(State->SingleTask.GetInstance());
 	}
 
@@ -2090,15 +2090,17 @@ EVisibility SStateTreeViewRow::GetTransitionsVisibility(const EStateTreeTransiti
 		
 		for (const FStateTreeEditorNode& Task : State->Tasks)
 		{
-			if (HasAnyLinksInStruct(FStateTreeDataView(Task.Node.GetScriptStruct(), const_cast<uint8*>(Task.Node.GetMemory())))
-				|| HasAnyLinksInStruct(Task.GetInstance()))
+			if (HasAnyLinksInStruct(Task.GetNode())
+				|| HasAnyLinksInStruct(Task.GetInstance())
+				|| HasAnyLinksInStruct(Task.GetExecutionRuntimeData()))
 			{
 				return EVisibility::Visible;
 			}
 		}
 
-		if (HasAnyLinksInStruct(FStateTreeDataView(State->SingleTask.Node.GetScriptStruct(), const_cast<uint8*>(State->SingleTask.Node.GetMemory())))
-			|| HasAnyLinksInStruct(State->SingleTask.GetInstance()))
+		if (HasAnyLinksInStruct(State->SingleTask.GetNode())
+			|| HasAnyLinksInStruct(State->SingleTask.GetInstance())
+			|| HasAnyLinksInStruct(State->SingleTask.GetExecutionRuntimeData()))
 		{
 			return EVisibility::Visible;
 		}
diff --git a/Engine/Plugins/Runtime/StateTree/Source/StateTreeEditorModule/Private/StateTreeCompiler.cpp b/Engine/Plugins/Runtime/StateTree/Source/StateTreeEditorModule/Private/StateTreeCompiler.cpp
--- a/Engine/Plugins/Runtime/StateTree/Source/StateTreeEditorModule/Private/StateTreeCompiler.cpp
+++ b/Engine/Plugins/Runtime/StateTree/Source/StateTreeEditorModule/Private/StateTreeCompiler.cpp
@@ -1082,7 +1082,7 @@ bool FStateTreeCompiler::CreateGlobalTasks()
 	{
 		FStateTreeEditorNode& TaskNode = EditorData->GlobalTasks[TaskIndex];
 		// Silently ignore empty nodes.
-		if (!TaskNode.Node.IsValid())
+		if (!TaskNode.GetNode().IsValid())
 		{
 			continue;
 		}
@@ -1293,7 +1293,7 @@ bool FStateTreeCompiler::CreateStateTasksAndParameters()
 		{
 			Tasks = State->Tasks;
 		}
-		else if (State->SingleTask.Node.IsValid())
+		else if (State->SingleTask.GetNode().IsValid())
 		{
 			Tasks = TArrayView<FStateTreeEditorNode>(&State->SingleTask, 1);
 		}
@@ -1305,12 +1305,12 @@ bool FStateTreeCompiler::CreateStateTasksAndParameters()
 		{
 			FStateTreeEditorNode& TaskNode = Tasks[TaskIndex];
 			// Silently ignore empty nodes.
-			if (!TaskNode.Node.IsValid())
+			if (!TaskNode.GetNode().IsValid())
 			{
 				continue;
 			}
 
-			FStateTreeTaskBase& Task = TaskNode.Node.GetMutable<FStateTreeTaskBase>();
+			FStateTreeTaskBase& Task = TaskNode.GetNode().Get<FStateTreeTaskBase>();
 			if(Task.bTaskEnabled)
 			{
 				EnabledTasksNum += 1;
@@ -1856,7 +1856,7 @@ bool FStateTreeCompiler::CreateBindingsForNodes(TConstArrayView<FStateTreeEditor
 	for (const FStateTreeEditorNode& EditorNode : EditorNodes)
 	{
 		// Node might be an empty line in Editor.
-		if (!EditorNode.Node.IsValid())
+		if (!EditorNode.GetNode().IsValid())
 		{
 			continue;
 		}
@@ -1884,12 +1884,12 @@ bool FStateTreeCompiler::CreateBindingsForNodes(TConstArrayView<FStateTreeEditor
 		FInstancedStruct& Instance = Instances[Node.InstanceTemplateIndex.Get()];
 		if (FStateTreeInstanceObjectWrapper* ObjectWrapper = Instance.GetMutablePtr<FStateTreeInstanceObjectWrapper>())
 		{
-			check(EditorNode.InstanceObject->GetClass() == ObjectWrapper->InstanceObject->GetClass());
+			check(EditorNode.GetInstance().GetStruct() == ObjectWrapper->InstanceObject->GetClass());
 			InstanceView = FStateTreeDataView(ObjectWrapper->InstanceObject);
 		}
 		else
 		{
-			check(EditorNode.Instance.GetScriptStruct() == Instance.GetScriptStruct());
+			check(EditorNode.GetInstance().GetStruct() == Instance.GetScriptStruct());
 			InstanceView = FStateTreeDataView(Instance);
 		}
 
@@ -2060,12 +2060,12 @@ bool FStateTreeCompiler::CreatePropertyFunction(const FStateTreeEditorNode& Func
 	FInstancedStruct& Instance = InstancedStructContainer[Function->InstanceTemplateIndex.Get()];
 	if (FStateTreeInstanceObjectWrapper* ObjectWrapper = Instance.GetMutablePtr<FStateTreeInstanceObjectWrapper>())
 	{
-		check(FuncEditorNode.InstanceObject->GetClass() == ObjectWrapper->InstanceObject->GetClass());
+		check(FuncEditorNode.GetInstance().GetStruct() == ObjectWrapper->InstanceObject->GetClass());
 		InstanceView = FStateTreeDataView(ObjectWrapper->InstanceObject);
 	}
 	else
 	{
-		check(FuncEditorNode.Instance.GetScriptStruct() == Instance.GetScriptStruct());
+		check(FuncEditorNode.GetInstance().GetStruct() == Instance.GetScriptStruct());
 		InstanceView = FStateTreeDataView(Instance);
 	}
 
@@ -2074,13 +2074,13 @@ bool FStateTreeCompiler::CreatePropertyFunction(const FStateTreeEditorNode& Func
 
 FInstancedStruct* FStateTreeCompiler::CreateNode(UStateTreeState* State, const FStateTreeEditorNode& EditorNode, FStateTreeBindableStructDesc& InstanceDesc, const FStateTreeDataHandle DataHandle, TArray<FInstancedStruct>& InstancedStructContainer)
 {
-	if (!EditorNode.Node.IsValid())
+	if (!EditorNode.GetNode().IsValid())
 	{
 		return nullptr;
 	}
 
 	// Check that item has valid instance initialized.
-	if (!EditorNode.Instance.IsValid() && EditorNode.InstanceObject == nullptr)
+	if (!EditorNode.GetInstance().IsValid())
 	{
 		Log.Reportf(EMessageSeverity::Error, InstanceDesc,
 			TEXT("Malformed node, missing instance value."));
@@ -2089,7 +2089,7 @@ FInstancedStruct* FStateTreeCompiler::CreateNode(UStateTreeState* State, const F
 
 	// Copy the node
 	IDToNode.Add(EditorNode.ID, Nodes.Num());
-	FInstancedStruct& RawNode = Nodes.Add_GetRef(EditorNode.Node);
+	FInstancedStruct& RawNode = Nodes.Add_GetRef(FInstancedStruct(EditorNode.GetNode()));
 	InstantiateStructSubobjects(RawNode);
 
 	FStateTreeNodeBase& Node = RawNode.GetMutable<FStateTreeNodeBase>();
@@ -2540,7 +2540,7 @@ bool FStateTreeCompiler::CompileAndValidateNode(const UStateTreeState* SourceSta
 
 bool FStateTreeCompiler::CreateTask(UStateTreeState* State, const FStateTreeEditorNode& TaskNode, const FStateTreeDataHandle TaskDataHandle)
 {
-	if (!TaskNode.Node.IsValid())
+	if (!TaskNode.GetNode().IsValid())
 	{
 		return false;
 	}
@@ -2553,7 +2553,7 @@ bool FStateTreeCompiler::CreateTask(UStateTreeState* State, const FStateTreeEdit
 	InstanceDesc.DataSource = EStateTreeBindableStructSource::Task;
 
 	// Check that node has valid instance initialized.
-	if (!TaskNode.Instance.IsValid() && TaskNode.InstanceObject == nullptr)
+	if (!TaskNode.GetInstance().IsValid())
 	{
 		Log.Reportf(EMessageSeverity::Error, InstanceDesc,
 			TEXT("Malformed task, missing instance value."));
@@ -2590,7 +2590,7 @@ bool FStateTreeCompiler::CreateTask(UStateTreeState* State, const FStateTreeEdit
 bool FStateTreeCompiler::CreateEvaluator(const FStateTreeEditorNode& EvalNode, const FStateTreeDataHandle EvalDataHandle)
 {
 	// Silently ignore empty nodes.
-	if (!EvalNode.Node.IsValid())
+	if (!EvalNode.GetNode().IsValid())
 	{
 		return true;
 	}
@@ -2598,17 +2598,17 @@ bool FStateTreeCompiler::CreateEvaluator(const FStateTreeEditorNode& EvalNode, c
 	// Create binding source descriptor for instance.
 	FStateTreeBindableStructDesc InstanceDesc;
 	InstanceDesc.StatePath = UE::StateTree::Editor::GlobalStateName;
-    InstanceDesc.ID = EvalNode.ID;
+	InstanceDesc.ID = EvalNode.ID;
 	InstanceDesc.Name = EvalNode.GetName();
 	InstanceDesc.DataSource = EStateTreeBindableStructSource::Evaluator;
 
-    // Check that node has valid instance initialized.
-    if (!EvalNode.Instance.IsValid() && EvalNode.InstanceObject == nullptr)
-    {
-        Log.Reportf(EMessageSeverity::Error, InstanceDesc,
-        	TEXT("Malformed evaluator, missing instance value."));
-        return false;
-    }
+	// Check that node has valid instance initialized.
+	if (!EvalNode.GetInstance().IsValid())
+	{
+		Log.Reportf(EMessageSeverity::Error, InstanceDesc,
+			TEXT("Malformed evaluator, missing instance value."));
+		return false;
+	}
 
 	// Copy the evaluator
 	IDToNode.Add(EvalNode.ID, Nodes.Num());
@@ -3206,7 +3206,7 @@ void FStateTreeCompiler::CreateBindingSourceStructsForNode(const FStateTreeEdito
 {
 	using namespace UE::StateTree::Compiler;
 
-	check(EditorNode.Node.IsValid());
+	check(EditorNode.GetNode().IsValid());
 
 	BindingsCompiler.AddSourceStruct(InstanceDesc);
 
@@ -3214,7 +3214,7 @@ void FStateTreeCompiler::CreateBindingSourceStructsForNode(const FStateTreeEdito
 	{
 		FStateTreeBindableStructDesc NodeDesc = InstanceDesc;
 		NodeDesc.ID = EditorNode.GetNodeID();
-		NodeDesc.Struct = EditorNode.Node.GetScriptStruct();
+		NodeDesc.Struct = EditorNode.GetNode().GetScriptStruct();
 		NodeDesc.DataHandle = FStateTreeDataHandle::Invalid;	// DataHandle is only for Instance data
 
 		BindingsCompiler.AddSourceStruct(NodeDesc);
diff --git a/Engine/Plugins/Runtime/StateTree/Source/StateTreeEditorModule/Private/StateTreeEditorData.cpp b/Engine/Plugins/Runtime/StateTree/Source/StateTreeEditorModule/Private/StateTreeEditorData.cpp
--- a/Engine/Plugins/Runtime/StateTree/Source/StateTreeEditorModule/Private/StateTreeEditorData.cpp
+++ b/Engine/Plugins/Runtime/StateTree/Source/StateTreeEditorModule/Private/StateTreeEditorData.cpp
@@ -782,7 +782,7 @@ void UStateTreeEditorData::FixObjectInstances(TSet<UObject*>& SeenObjects, UObje
 	auto NewInstances = [this, &Node, &Outer](const UStruct* Struct, const UStruct* ExecutionRuntimeStruct) -> bool
 		{
 			bool bDirty = false;
-			if (Struct && !Node.Instance.IsValid() && Node.InstanceObject == nullptr)
+			if (Struct && !Node.GetInstance().IsValid())
 			{
 				if (const UScriptStruct* InstanceType = Cast<const UScriptStruct>(Struct))
 				{
@@ -795,7 +795,7 @@ void UStateTreeEditorData::FixObjectInstances(TSet<UObject*>& SeenObjects, UObje
 					bDirty = true;
 				}
 			}
-			if (ExecutionRuntimeStruct && !Node.ExecutionRuntimeData.IsValid() && Node.ExecutionRuntimeDataObject == nullptr)
+			if (ExecutionRuntimeStruct && !Node.GetExecutionRuntimeData().IsValid())
 			{
 				if (const UScriptStruct* InstanceType = Cast<const UScriptStruct>(ExecutionRuntimeStruct))
 				{
@@ -856,7 +856,7 @@ void UStateTreeEditorData::FixObjectInstances(TSet<UObject*>& SeenObjects, UObje
 			return bDirty;
 		};
 
-	const FStateTreeNodeBase* NodeBase = Node.Node.GetPtr<FStateTreeNodeBase>();
+	const FStateTreeNodeBase* NodeBase = Node.GetNode().GetPtr<FStateTreeNodeBase>();
 	bool bDirty = false;
 	if (NodeBase)
 	{
@@ -883,23 +883,15 @@ void UStateTreeEditorData::FixObjectInstances(TSet<UObject*>& SeenObjects, UObje
 	}
 	else
 	{
-		if (Node.Instance.IsValid())
+		if (Node.GetInstance().IsValid())
 		{
 			Node.Instance.Reset();
-			bDirty = true;
-		}
-		if (Node.InstanceObject != nullptr)
-		{
 			Node.InstanceObject = nullptr;
 			bDirty = true;
 		}
-		if (Node.ExecutionRuntimeData.IsValid())
+		if (Node.GetExecutionRuntimeData().IsValid())
 		{
 			Node.ExecutionRuntimeData.Reset();
-			bDirty = true;
-		}
-		if (Node.ExecutionRuntimeDataObject != nullptr)
-		{
 			Node.ExecutionRuntimeDataObject = nullptr;
 			bDirty = true;
 		}
@@ -967,7 +959,7 @@ void UStateTreeEditorData::FixObjectNodes()
 
 FText UStateTreeEditorData::GetNodeDescription(const FStateTreeEditorNode& Node, const EStateTreeNodeFormatting Formatting) const
 {
-	if (const FStateTreeNodeBase* NodePtr = Node.Node.GetPtr<FStateTreeNodeBase>())
+	if (const FStateTreeNodeBase* NodePtr = Node.GetNode().GetPtr<FStateTreeNodeBase>())
 	{
 		// If the node has name override, return it.
 		if (!NodePtr->Name.IsNone())
@@ -989,8 +981,8 @@ FText UStateTreeEditorData::GetNodeDescription(const FStateTreeEditorNode& Node,
 		}
 
 		// As last resort, return node's display name.
-		check(Node.Node.GetScriptStruct());
-		return Node.Node.GetScriptStruct()->GetDisplayNameText();
+		check(Node.GetNode().GetScriptStruct());
+		return Node.GetNode().GetScriptStruct()->GetDisplayNameText();
 	}
 
 	// The node is not initialized.
@@ -1008,7 +1000,7 @@ void UStateTreeEditorData::FixDuplicateIDs()
 	for (int32 Index = 0; Index < Evaluators.Num(); Index++)
 	{
 		FStateTreeEditorNode& Node = Evaluators[Index];
-		if (const FStateTreeEvaluatorBase* Evaluator = Node.Node.GetPtr<FStateTreeEvaluatorBase>())
+		if (const FStateTreeEvaluatorBase* Evaluator = Node.GetNode().GetPtr<FStateTreeEvaluatorBase>())
 		{
 			const FGuid OldID = Node.ID; 
 			if (FoundNodeIDs.Contains(Node.ID))
@@ -1028,9 +1020,9 @@ void UStateTreeEditorData::FixDuplicateIDs()
 	for (int32 Index = 0; Index < GlobalTasks.Num(); Index++)
 	{
 		FStateTreeEditorNode& Node = GlobalTasks[Index];
-		if (const FStateTreeTaskBase* Task = Node.Node.GetPtr<FStateTreeTaskBase>())
+		if (const FStateTreeTaskBase* Task = Node.GetNode().GetPtr<FStateTreeTaskBase>())
 		{
-			const FGuid OldID = Node.ID; 
+			const FGuid OldID = Node.ID;
 			if (FoundNodeIDs.Contains(Node.ID))
 			{
 				Node.ID = UE::StateTree::PropertyHelpers::MakeDeterministicID(*this, TEXT("GlobalTasks"), Index);
@@ -1050,10 +1042,10 @@ void UStateTreeEditorData::FixDuplicateIDs()
 		for (int32 Index = 0; Index < State.EnterConditions.Num(); Index++)
 		{
 			FStateTreeEditorNode& Node = State.EnterConditions[Index];
-			if (const FStateTreeConditionBase* Cond = Node.Node.GetPtr<FStateTreeConditionBase>())
+			if (const FStateTreeConditionBase* Cond = Node.GetNode().GetPtr<FStateTreeConditionBase>())
 			{
 				const FGuid OldID = Node.ID;
-				
+
 				bool bIsAlreadyInSet = false;
 				FoundNodeIDs.Add(Node.ID, &bIsAlreadyInSet);
 				if (bIsAlreadyInSet)
@@ -1072,7 +1064,7 @@ void UStateTreeEditorData::FixDuplicateIDs()
 		for (int32 Index = 0; Index < State.Considerations.Num(); Index++)
 		{
 			FStateTreeEditorNode& Node = State.Considerations[Index];
-			if (const FStateTreeConsiderationBase* Cond = Node.Node.GetPtr<FStateTreeConsiderationBase>())
+			if (const FStateTreeConsiderationBase* Cond = Node.GetNode().GetPtr<FStateTreeConsiderationBase>())
 			{
 				const FGuid OldID = Node.ID;
 				
@@ -1094,10 +1086,10 @@ void UStateTreeEditorData::FixDuplicateIDs()
 		for (int32 Index = 0; Index < State.Tasks.Num(); Index++)
 		{
 			FStateTreeEditorNode& Node = State.Tasks[Index];
-			if (const FStateTreeTaskBase* Task = Node.Node.GetPtr<FStateTreeTaskBase>())
+			if (const FStateTreeTaskBase* Task = Node.GetNode().GetPtr<FStateTreeTaskBase>())
 			{
 				const FGuid OldID = Node.ID;
-				
+
 				bool bIsAlreadyInSet = false;
 				FoundNodeIDs.Add(Node.ID, &bIsAlreadyInSet);
 				if (bIsAlreadyInSet)
@@ -1112,7 +1104,7 @@ void UStateTreeEditorData::FixDuplicateIDs()
 			}
 		}
 
-		if (FStateTreeTaskBase* Task = State.SingleTask.Node.GetMutablePtr<FStateTreeTaskBase>())
+		if (FStateTreeTaskBase* Task = State.SingleTask.GetNode().GetPtr<FStateTreeTaskBase>())
 		{
 			const FGuid OldID = State.SingleTask.ID;
 
@@ -1136,9 +1128,9 @@ void UStateTreeEditorData::FixDuplicateIDs()
 			for (int32 Index = 0; Index < Transition.Conditions.Num(); Index++)
 			{
 				FStateTreeEditorNode& Node = Transition.Conditions[Index];
-				if (const FStateTreeConditionBase* Cond = Node.Node.GetPtr<FStateTreeConditionBase>())
+				if (const FStateTreeConditionBase* Cond = Node.GetNode().GetPtr<FStateTreeConditionBase>())
 				{
-					const FGuid OldID = Node.ID; 
+					const FGuid OldID = Node.ID;
 					bool bIsAlreadyInSet = false;
 					FoundNodeIDs.Add(Node.ID, &bIsAlreadyInSet);
 					if (bIsAlreadyInSet)
@@ -1338,7 +1330,7 @@ void UStateTreeEditorData::CallPostLoadOnNodes()
 {
 	for (FStateTreeEditorNode& EvaluatorEditorNode : Evaluators)
 	{
-		if (FStateTreeNodeBase* EvaluatorNode = EvaluatorEditorNode.Node.GetMutablePtr<FStateTreeNodeBase>())
+		if (FStateTreeNodeBase* EvaluatorNode = EvaluatorEditorNode.GetNode().GetPtr<FStateTreeNodeBase>())
 		{
 			UE::StateTreeEditor::EditorNodeUtils::ConditionalUpdateNodeInstanceData(EvaluatorEditorNode, *this);
 			EvaluatorNode->PostLoad(EvaluatorEditorNode.GetInstance());
@@ -1347,7 +1339,7 @@ void UStateTreeEditorData::CallPostLoadOnNodes()
 
 	for (FStateTreeEditorNode& GlobalTaskEditorNode : GlobalTasks)
 	{
-		if (FStateTreeNodeBase* GlobalTaskNode = GlobalTaskEditorNode.Node.GetMutablePtr<FStateTreeNodeBase>())
+		if (FStateTreeNodeBase* GlobalTaskNode = GlobalTaskEditorNode.GetNode().GetPtr<FStateTreeNodeBase>())
 		{
 			UE::StateTreeEditor::EditorNodeUtils::ConditionalUpdateNodeInstanceData(GlobalTaskEditorNode, *this);
 			GlobalTaskNode->PostLoad(GlobalTaskEditorNode.GetInstance());
@@ -1424,7 +1416,7 @@ EStateTreeVisitor UStateTreeEditorData::VisitStateNodes(const UStateTreeState& S
 				bContinue = false;
 				break;
 			}
-			else if (const FStateTreeConditionBase* Cond = Node.Node.GetPtr<FStateTreeConditionBase>())
+			else if (const FStateTreeConditionBase* Cond = Node.GetNode().GetPtr<FStateTreeConditionBase>())
 			{
 				FStateTreeBindableStructDesc Desc;
 				Desc.StatePath = StatePathWithConditions;
@@ -1441,7 +1433,7 @@ EStateTreeVisitor UStateTreeEditorData::VisitStateNodes(const UStateTreeState& S
 
 				FStateTreeBindableStructDesc NodeDesc;
 				NodeDesc.StatePath = StatePathWithConditions;
-				NodeDesc.Struct = Node.Node.GetScriptStruct();
+				NodeDesc.Struct = Node.GetNode().GetScriptStruct();
 				NodeDesc.Name = Node.GetName();
 				NodeDesc.ID = Node.GetNodeID();
 				NodeDesc.DataSource = EStateTreeBindableStructSource::Condition;
@@ -1465,7 +1457,7 @@ EStateTreeVisitor UStateTreeEditorData::VisitStateNodes(const UStateTreeState& S
 				bContinue = false;
 				break;
 			}
-			else if (const FStateTreeConsiderationBase* Consideration = Node.Node.GetPtr<FStateTreeConsiderationBase>())
+			else if (const FStateTreeConsiderationBase* Consideration = Node.GetNode().GetPtr<FStateTreeConsiderationBase>())
 			{
 				FStateTreeBindableStructDesc Desc;
 				Desc.StatePath = StatePathWithConsiderations;
@@ -1482,7 +1474,7 @@ EStateTreeVisitor UStateTreeEditorData::VisitStateNodes(const UStateTreeState& S
 
 				FStateTreeBindableStructDesc NodeDesc;
 				NodeDesc.StatePath = StatePathWithConsiderations;
-				NodeDesc.Struct = Node.Node.GetScriptStruct();
+				NodeDesc.Struct = Node.GetNode().GetScriptStruct();
 				NodeDesc.Name = Node.GetName();
 				NodeDesc.ID = Node.GetNodeID();
 				NodeDesc.DataSource = EStateTreeBindableStructSource::Consideration;
@@ -1505,7 +1497,7 @@ EStateTreeVisitor UStateTreeEditorData::VisitStateNodes(const UStateTreeState& S
 				bContinue = false;
 				break;
 			}
-			else if (const FStateTreeTaskBase* Task = Node.Node.GetPtr<FStateTreeTaskBase>())
+			else if (const FStateTreeTaskBase* Task = Node.GetNode().GetPtr<FStateTreeTaskBase>())
 			{
 				FStateTreeBindableStructDesc Desc;
 				Desc.StatePath = StatePath;
@@ -1522,7 +1514,7 @@ EStateTreeVisitor UStateTreeEditorData::VisitStateNodes(const UStateTreeState& S
 
 				FStateTreeBindableStructDesc NodeDesc;
 				NodeDesc.StatePath = StatePath;
-				NodeDesc.Struct = Node.Node.GetScriptStruct();
+				NodeDesc.Struct = Node.GetNode().GetScriptStruct();
 				NodeDesc.Name = Node.GetName();
 				NodeDesc.ID = Node.GetNodeID();
 				NodeDesc.DataSource = EStateTreeBindableStructSource::Task;
@@ -1537,7 +1529,7 @@ EStateTreeVisitor UStateTreeEditorData::VisitStateNodes(const UStateTreeState& S
 	}
 	if (bContinue)
 	{
-		if (const FStateTreeTaskBase* Task = State.SingleTask.Node.GetPtr<FStateTreeTaskBase>())
+		if (const FStateTreeTaskBase* Task = State.SingleTask.GetNode().GetPtr<FStateTreeTaskBase>())
 		{
 			if (VisitFuncNodes(State.SingleTask.ID, State.SingleTask.GetName()) == EStateTreeVisitor::Break)
 			{
@@ -1559,7 +1551,7 @@ EStateTreeVisitor UStateTreeEditorData::VisitStateNodes(const UStateTreeState& S
 
 				FStateTreeBindableStructDesc NodeDesc;
 				NodeDesc.StatePath = StatePath;
-				NodeDesc.Struct = State.SingleTask.Node.GetScriptStruct();
+				NodeDesc.Struct = State.SingleTask.GetNode().GetScriptStruct();
 				NodeDesc.Name = State.SingleTask.GetName();
 				NodeDesc.ID = State.SingleTask.GetNodeID();
 				NodeDesc.DataSource = EStateTreeBindableStructSource::Task;
@@ -1618,7 +1610,7 @@ EStateTreeVisitor UStateTreeEditorData::VisitStateNodes(const UStateTreeState& S
 					bContinue = false;
 					break;
 				}
-				else if (const FStateTreeConditionBase* Cond = Node.Node.GetPtr<FStateTreeConditionBase>())
+				else if (const FStateTreeConditionBase* Cond = Node.GetNode().GetPtr<FStateTreeConditionBase>())
 				{
 					FStateTreeBindableStructDesc Desc;
 					Desc.StatePath = StatePathWithTransition;
@@ -1635,7 +1627,7 @@ EStateTreeVisitor UStateTreeEditorData::VisitStateNodes(const UStateTreeState& S
 
 					FStateTreeBindableStructDesc NodeDesc;
 					NodeDesc.StatePath = StatePathWithTransition;
-					NodeDesc.Struct = Node.Node.GetScriptStruct();
+					NodeDesc.Struct = Node.GetNode().GetScriptStruct();
 					NodeDesc.Name = Node.GetName();
 					NodeDesc.ID = Node.GetNodeID();
 					NodeDesc.DataSource = EStateTreeBindableStructSource::Condition;
@@ -1830,7 +1822,7 @@ EStateTreeVisitor UStateTreeEditorData::VisitGlobalNodes(TFunctionRef<EStateTree
 			return EStateTreeVisitor::Break;
 		}
 
-		if (const FStateTreeEvaluatorBase* Evaluator = Node.Node.GetPtr<FStateTreeEvaluatorBase>())
+		if (const FStateTreeEvaluatorBase* Evaluator = Node.GetNode().GetPtr<FStateTreeEvaluatorBase>())
 		{
 			FStateTreeBindableStructDesc Desc;
 			Desc.StatePath = UE::StateTree::Editor::GlobalStateName;
@@ -1846,7 +1838,7 @@ EStateTreeVisitor UStateTreeEditorData::VisitGlobalNodes(TFunctionRef<EStateTree
 
 			FStateTreeBindableStructDesc NodeDesc;
 			NodeDesc.StatePath = UE::StateTree::Editor::GlobalStateName;
-			NodeDesc.Struct = Node.Node.GetScriptStruct();
+			NodeDesc.Struct = Node.GetNode().GetScriptStruct();
 			NodeDesc.Name = Node.GetName();
 			NodeDesc.ID = Node.GetNodeID();
 			NodeDesc.DataSource = EStateTreeBindableStructSource::Evaluator;
@@ -1866,7 +1858,7 @@ EStateTreeVisitor UStateTreeEditorData::VisitGlobalNodes(TFunctionRef<EStateTree
 			return EStateTreeVisitor::Break;
 		}
 
-		if (const FStateTreeTaskBase* Task = Node.Node.GetPtr<FStateTreeTaskBase>())
+		if (const FStateTreeTaskBase* Task = Node.GetNode().GetPtr<FStateTreeTaskBase>())
 		{
 			FStateTreeBindableStructDesc Desc;
 			Desc.StatePath = UE::StateTree::Editor::GlobalStateName;
@@ -1882,7 +1874,7 @@ EStateTreeVisitor UStateTreeEditorData::VisitGlobalNodes(TFunctionRef<EStateTree
 
 			FStateTreeBindableStructDesc NodeDesc;
 			NodeDesc.StatePath = UE::StateTree::Editor::GlobalStateName;
-			NodeDesc.Struct = Node.Node.GetScriptStruct();
+			NodeDesc.Struct = Node.GetNode().GetScriptStruct();
 			NodeDesc.Name = Node.GetName();
 			NodeDesc.ID = Node.GetNodeID();
 			NodeDesc.DataSource = EStateTreeBindableStructSource::GlobalTask;
diff --git a/Engine/Plugins/Runtime/StateTree/Source/StateTreeEditorModule/Private/StateTreeEditorDataClipboardHelpers.cpp b/Engine/Plugins/Runtime/StateTree/Source/StateTreeEditorModule/Private/StateTreeEditorDataClipboardHelpers.cpp
--- a/Engine/Plugins/Runtime/StateTree/Source/StateTreeEditorModule/Private/StateTreeEditorDataClipboardHelpers.cpp
+++ b/Engine/Plugins/Runtime/StateTree/Source/StateTreeEditorModule/Private/StateTreeEditorDataClipboardHelpers.cpp
@@ -119,7 +119,7 @@ struct FScopedEditorDataFixer
 
 		auto ReinstantiateEditorNodeInstanceData = [Self = this](FStateTreeEditorNode& EditorNode)
 			{
-				EditorNodeUtils::InstantiateStructSubobjects(*Self->Owner, EditorNode.Node);
+				EditorNodeUtils::InstantiateStructSubobjects(*Self->Owner, EditorNode.GetNode());
 				if (EditorNode.InstanceObject)
 				{
 					EditorNode.InstanceObject = DuplicateObject(EditorNode.InstanceObject, Self->Owner);
diff --git a/Engine/Plugins/Runtime/StateTree/Source/StateTreeEditorModule/Private/StateTreeEditorPropertyBindings.cpp b/Engine/Plugins/Runtime/StateTree/Source/StateTreeEditorModule/Private/StateTreeEditorPropertyBindings.cpp
--- a/Engine/Plugins/Runtime/StateTree/Source/StateTreeEditorModule/Private/StateTreeEditorPropertyBindings.cpp
+++ b/Engine/Plugins/Runtime/StateTree/Source/StateTreeEditorModule/Private/StateTreeEditorPropertyBindings.cpp
@@ -28,7 +28,7 @@ FPropertyBindingPath FStateTreeEditorPropertyBindings::AddFunctionBinding(const
 	const FGuid NodeID = FGuid::NewGuid();
 	PropertyFunction.ID = NodeID;
 	PropertyFunction.Node.InitializeAs(InPropertyFunctionNodeStruct);
-	const FStateTreePropertyFunctionBase& Function = PropertyFunction.Node.Get<FStateTreePropertyFunctionBase>();
+	const FStateTreePropertyFunctionBase& Function = PropertyFunction.GetNode().Get<FStateTreePropertyFunctionBase>();
 	if (const UScriptStruct* InstanceType = Cast<const UScriptStruct>(Function.GetInstanceDataType()))
 	{
 		PropertyFunction.Instance.InitializeAs(InstanceType);
@@ -126,7 +126,7 @@ void FStateTreeEditorPropertyBindings::CopyBindingsInternal(const FGuid InFromSt
 					const FStateTreeEditorNode* EditorNode = Binding.GetPropertyFunctionNode().GetPtr<const FStateTreeEditorNode>();
 					if (ensure(EditorNode))
 					{
-						FunctionBindingsToCopy.Emplace(EditorNode->Node.GetScriptStruct(), Binding.GetSourcePath(), FPropertyBindingPath(ToStructID, Binding.GetTargetPath().GetSegments()));
+						FunctionBindingsToCopy.Emplace(EditorNode->GetNode().GetScriptStruct(), Binding.GetSourcePath(), FPropertyBindingPath(ToStructID, Binding.GetTargetPath().GetSegments()));
 					}
 				}
 				return !bIsFunctionBinding;
@@ -276,12 +276,12 @@ FText FStateTreeBindingLookup::GetPropertyPathDisplayName(const FPropertyBinding
 			const FConstStructView PropertyFuncEditorNodeView = BindingToPath->GetPropertyFunctionNode();
 			const FStateTreeEditorNode& EditorNode = PropertyFuncEditorNodeView.Get<const FStateTreeEditorNode>();
 
-			if (!EditorNode.Node.IsValid())
+			if (!EditorNode.GetNode().IsValid())
 			{
 				return LOCTEXT("Unlinked", "???");
 			}
 
-			const FStateTreeNodeBase& Node = EditorNode.Node.Get<FStateTreeNodeBase>();
+			const FStateTreeNodeBase& Node = EditorNode.GetNode().Get<FStateTreeNodeBase>();
 
 			// Skipping an output property if there's only one of them.
 			if (UE::StateTree::GetStructSingleOutputProperty(*Node.GetInstanceDataType()))
diff --git a/Engine/Plugins/Runtime/StateTree/Source/StateTreeEditorModule/Private/StateTreeEditorSchema.cpp b/Engine/Plugins/Runtime/StateTree/Source/StateTreeEditorModule/Private/StateTreeEditorSchema.cpp
--- a/Engine/Plugins/Runtime/StateTree/Source/StateTreeEditorModule/Private/StateTreeEditorSchema.cpp
+++ b/Engine/Plugins/Runtime/StateTree/Source/StateTreeEditorModule/Private/StateTreeEditorSchema.cpp
@@ -61,14 +61,14 @@ void UStateTreeEditorSchema::Validate(TNotNull<UStateTree*> StateTree)
 				}
 
 				// Task name is the same as state name.
-				if (FStateTreeTaskBase* Task = State.SingleTask.Node.GetMutablePtr<FStateTreeTaskBase>())
+				if (FStateTreeTaskBase* Task = State.SingleTask.GetNode().GetPtr<FStateTreeTaskBase>())
 				{
 					Task->Name = State.Name;
 				}
 			}
 			else
 			{
-				if (State.SingleTask.Node.IsValid())
+				if (State.SingleTask.GetNode().IsValid())
 				{
 					State.SingleTask.Reset();
 					UE_LOG(LogStateTreeEditor, Warning, TEXT("%s: Resetting Single Task in state %s due to current schema restrictions."), *StateTree->GetName(), *State.GetName());
diff --git a/Engine/Plugins/Runtime/StateTree/Source/StateTreeEditorModule/Private/StateTreePropertyHelpers.cpp b/Engine/Plugins/Runtime/StateTree/Source/StateTreeEditorModule/Private/StateTreePropertyHelpers.cpp
--- a/Engine/Plugins/Runtime/StateTree/Source/StateTreeEditorModule/Private/StateTreePropertyHelpers.cpp
+++ b/Engine/Plugins/Runtime/StateTree/Source/StateTreeEditorModule/Private/StateTreePropertyHelpers.cpp
@@ -17,7 +17,7 @@ namespace Internal
 {
 	void DispatchPostEditToEditorNode(FPropertyChangedChainEvent& InPropertyChangedEvent, const TDoubleLinkedList<FProperty*>::TDoubleLinkedListNode* InEditorNodeInChain, FStateTreeEditorNode& InEditorNode)
 	{
-		if (FStateTreeNodeBase* StateTreeNode = InEditorNode.Node.GetMutablePtr<FStateTreeNodeBase>())
+		if (FStateTreeNodeBase* StateTreeNode = InEditorNode.GetNode().GetPtr<FStateTreeNodeBase>())
 		{
 			// Check that the path contains EditorNode's: Node, Instance or Instance Object
 			if (const TDoubleLinkedList<FProperty*>::TDoubleLinkedListNode* EditorNodeMemberPropNode = InEditorNodeInChain->GetNextNode())
@@ -58,32 +58,20 @@ namespace Internal
 					{
 						StateTreeNode->PostEditNodeChangeChainProperty(PropertyChangedEvent, InEditorNode.GetInstance());
 					}
-					else if (EditorNodeChildMember->GetFName() == GET_MEMBER_NAME_CHECKED(FStateTreeEditorNode, Instance))
+					else if (EditorNodeChildMember->GetFName() == GET_MEMBER_NAME_CHECKED(FStateTreeEditorNode, Instance)
+						|| EditorNodeChildMember->GetFName() == GET_MEMBER_NAME_CHECKED(FStateTreeEditorNode, InstanceObject))
 					{
-						if (InEditorNode.Instance.IsValid())
+						if (InEditorNode.GetInstance().IsValid())
 						{
-							StateTreeNode->PostEditInstanceDataChangeChainProperty(PropertyChangedEvent, FStateTreeDataView(InEditorNode.Instance));
+							StateTreeNode->PostEditInstanceDataChangeChainProperty(PropertyChangedEvent, InEditorNode.GetInstance());
 						}
 					}
-					else if (EditorNodeChildMember->GetFName() == GET_MEMBER_NAME_CHECKED(FStateTreeEditorNode, InstanceObject))
+					else if (EditorNodeChildMember->GetFName() == GET_MEMBER_NAME_CHECKED(FStateTreeEditorNode, ExecutionRuntimeData)
+						|| EditorNodeChildMember->GetFName() == GET_MEMBER_NAME_CHECKED(FStateTreeEditorNode, ExecutionRuntimeDataObject))
 					{
-						if (InEditorNode.InstanceObject)
+						if (InEditorNode.GetExecutionRuntimeData().IsValid())
 						{
-							StateTreeNode->PostEditInstanceDataChangeChainProperty(PropertyChangedEvent, FStateTreeDataView(InEditorNode.InstanceObject));
-						}
-					}
-					else if (EditorNodeChildMember->GetFName() == GET_MEMBER_NAME_CHECKED(FStateTreeEditorNode, ExecutionRuntimeData))
-					{
-						if (InEditorNode.ExecutionRuntimeData.IsValid())
-						{
-							StateTreeNode->PostEditInstanceDataChangeChainProperty(PropertyChangedEvent, FStateTreeDataView(InEditorNode.ExecutionRuntimeData));
-						}
-					}
-					else if (EditorNodeChildMember->GetFName() == GET_MEMBER_NAME_CHECKED(FStateTreeEditorNode, ExecutionRuntimeDataObject))
-					{
-						if (InEditorNode.ExecutionRuntimeDataObject)
-						{
-							StateTreeNode->PostEditInstanceDataChangeChainProperty(PropertyChangedEvent, FStateTreeDataView(InEditorNode.ExecutionRuntimeDataObject));
+							StateTreeNode->PostEditInstanceDataChangeChainProperty(PropertyChangedEvent, InEditorNode.GetExecutionRuntimeData());
 						}
 					}
 				}
diff --git a/Engine/Plugins/Runtime/StateTree/Source/StateTreeEditorModule/Private/StateTreeState.cpp b/Engine/Plugins/Runtime/StateTree/Source/StateTreeEditorModule/Private/StateTreeState.cpp
--- a/Engine/Plugins/Runtime/StateTree/Source/StateTreeEditorModule/Private/StateTreeState.cpp
+++ b/Engine/Plugins/Runtime/StateTree/Source/StateTreeEditorModule/Private/StateTreeState.cpp
@@ -402,7 +402,7 @@ PRAGMA_ENABLE_DEPRECATION_WARNINGS
 #if WITH_EDITOR
 	for (FStateTreeEditorNode& EnterConditionEditorNode : EnterConditions)
 	{
-		if (FStateTreeNodeBase* ConditionNode = EnterConditionEditorNode.Node.GetMutablePtr<FStateTreeNodeBase>())
+		if (FStateTreeNodeBase* ConditionNode = EnterConditionEditorNode.GetNode().GetPtr<FStateTreeNodeBase>())
 		{
 			UE::StateTreeEditor::EditorNodeUtils::ConditionalUpdateNodeInstanceData(EnterConditionEditorNode, *this);
 			ConditionNode->PostLoad(EnterConditionEditorNode.GetInstance());
@@ -411,7 +411,7 @@ PRAGMA_ENABLE_DEPRECATION_WARNINGS
 
 	for (FStateTreeEditorNode& ConsiderationEditorNode : Considerations)
 	{
-		if (FStateTreeNodeBase* ConsiderationNode = ConsiderationEditorNode.Node.GetMutablePtr<FStateTreeNodeBase>())
+		if (FStateTreeNodeBase* ConsiderationNode = ConsiderationEditorNode.GetNode().GetPtr<FStateTreeNodeBase>())
 		{
 			UE::StateTreeEditor::EditorNodeUtils::ConditionalUpdateNodeInstanceData(ConsiderationEditorNode, *this);
 			ConsiderationNode->PostLoad(ConsiderationEditorNode.GetInstance());
@@ -420,14 +420,14 @@ PRAGMA_ENABLE_DEPRECATION_WARNINGS
 
 	for (FStateTreeEditorNode& TaskEditorNode : Tasks)
 	{
-		if (FStateTreeNodeBase* TaskNode = TaskEditorNode.Node.GetMutablePtr<FStateTreeNodeBase>())
+		if (FStateTreeNodeBase* TaskNode = TaskEditorNode.GetNode().GetPtr<FStateTreeNodeBase>())
 		{
 			UE::StateTreeEditor::EditorNodeUtils::ConditionalUpdateNodeInstanceData(TaskEditorNode, *this);
 			TaskNode->PostLoad(TaskEditorNode.GetInstance());
 		}
 	}
 
-	if (FStateTreeNodeBase* SingleTaskNode = SingleTask.Node.GetMutablePtr<FStateTreeNodeBase>())
+	if (FStateTreeNodeBase* SingleTaskNode = SingleTask.GetNode().GetPtr<FStateTreeNodeBase>())
 	{
 		UE::StateTreeEditor::EditorNodeUtils::ConditionalUpdateNodeInstanceData(SingleTask, *this);
 		SingleTaskNode->PostLoad(SingleTask.GetInstance());
@@ -437,7 +437,7 @@ PRAGMA_ENABLE_DEPRECATION_WARNINGS
 	{
 		for (FStateTreeEditorNode& TransitionConditionEditorNode : Transition.Conditions)
 		{
-			if (FStateTreeNodeBase* ConditionNode = TransitionConditionEditorNode.Node.GetMutablePtr<FStateTreeNodeBase>())
+			if (FStateTreeNodeBase* ConditionNode = TransitionConditionEditorNode.GetNode().GetPtr<FStateTreeNodeBase>())
 			{
 				UE::StateTreeEditor::EditorNodeUtils::ConditionalUpdateNodeInstanceData(TransitionConditionEditorNode, *this);
 				ConditionNode->PostLoad(TransitionConditionEditorNode.GetInstance());
@@ -675,7 +675,7 @@ UStateTreeState& UStateTreeState::AddChildState(const FName ChildName, const ESt
 void UStateTreeState::AddNodeImpl(FStateTreeEditorNode& EditorNode)
 {
 	EditorNode.ID = FGuid::NewGuid();
-	const FStateTreeNodeBase& Node = EditorNode.Node.GetMutable<FStateTreeNodeBase>();
+	const FStateTreeNodeBase& Node = EditorNode.GetNode().Get<FStateTreeNodeBase>();
 	if (const UScriptStruct* InstanceType = Cast<const UScriptStruct>(Node.GetInstanceDataType()))
 	{
 		EditorNode.Instance.InitializeAs(InstanceType);
diff --git a/Engine/Plugins/Runtime/StateTree/Source/StateTreeEditorModule/Private/StateTreeViewModel.cpp b/Engine/Plugins/Runtime/StateTree/Source/StateTreeEditorModule/Private/StateTreeViewModel.cpp
--- a/Engine/Plugins/Runtime/StateTree/Source/StateTreeEditorModule/Private/StateTreeViewModel.cpp
+++ b/Engine/Plugins/Runtime/StateTree/Source/StateTreeEditorModule/Private/StateTreeViewModel.cpp
@@ -144,25 +144,17 @@ namespace UE::StateTree::Editor
 			IDsMap.Emplace(Node.ID, NewNodeID);
 			Node.ID = NewNodeID;
 
-			if (Node.Node.IsValid())
+			if (Node.GetNode().IsValid())
 			{
-				CollectStateLinks(Node.Node.GetScriptStruct(), Node.Node.GetMutableMemory(), Links);
+				CollectStateLinks(Node.GetNode().GetScriptStruct(), Node.GetNode().GetMemory(), Links);
 			}
-			if (Node.Instance.IsValid())
+			if (Node.GetInstance().IsValid())
 			{
-				CollectStateLinks(Node.Instance.GetScriptStruct(), Node.Instance.GetMutableMemory(), Links);
+				CollectStateLinks(Node.GetInstance().GetStruct(), Node.GetInstance().GetMutableMemory(), Links);
 			}
-			if (Node.InstanceObject)
+			if (Node.GetExecutionRuntimeData().IsValid())
 			{
-				CollectStateLinks(Node.InstanceObject->GetClass(), Node.InstanceObject, Links);
-			}
-			if (Node.ExecutionRuntimeData.IsValid())
-			{
-				CollectStateLinks(Node.ExecutionRuntimeData.GetScriptStruct(), Node.ExecutionRuntimeData.GetMutableMemory(), Links);
-			}
-			if (Node.ExecutionRuntimeDataObject)
-			{
-				CollectStateLinks(Node.ExecutionRuntimeDataObject->GetClass(), Node.ExecutionRuntimeDataObject, Links);
+				CollectStateLinks(Node.GetExecutionRuntimeData().GetStruct(), Node.GetExecutionRuntimeData().GetMutableMemory(), Links);
 			}
 		}
 	}
diff --git a/Engine/Plugins/Runtime/StateTree/Source/StateTreeEditorModule/Public/StateTreeEditorNode.h b/Engine/Plugins/Runtime/StateTree/Source/StateTreeEditorModule/Public/StateTreeEditorNode.h
--- a/Engine/Plugins/Runtime/StateTree/Source/StateTreeEditorModule/Public/StateTreeEditorNode.h
+++ b/Engine/Plugins/Runtime/StateTree/Source/StateTreeEditorModule/Public/StateTreeEditorNode.h
@@ -19,7 +19,7 @@ enum class EStateTreeNodeType : uint8
 };
 
 /**
- * Base for Evaluator, Task and Condition nodes.
+ * Base for Evaluator, Task, Condition, Consideration Utilities, Property Function nodes.
  */
 USTRUCT()
 struct FStateTreeEditorNode
@@ -120,9 +120,18 @@ template <typename T>
 struct TStateTreeEditorNode : public FStateTreeEditorNode
 {
 	using NodeType = T;
-	inline T& GetNode() { return Node.template GetMutable<T>(); }
-	inline typename T::FInstanceDataType& GetInstanceData() { return Instance.template GetMutable<typename T::FInstanceDataType>(); }
-	inline typename T::FExecutionRuntimeDataType& GetExecutionRuntimeData() { return ExecutionRuntimeData.template GetMutable<typename T::FExecutionRuntimeDataType>(); }
+	inline T& GetNode()
+	{
+		return Node.template GetMutable<T>();
+	}
+	inline typename T::FInstanceDataType& GetInstanceData()
+	{
+		return Instance.template GetMutable<typename T::FInstanceDataType>();
+	}
+	inline typename T::FExecutionRuntimeDataType& GetExecutionRuntimeData()
+	{
+		return ExecutionRuntimeData.template GetMutable<typename T::FExecutionRuntimeDataType>();
+	}
 };
 
 #undef UE_API
