[AI][StateTree] Add StateTreeToolset plugin for MCP inspection

Adds a standalone plugin (StateTreeToolset) that exposes StateTree editor data to AI agents via the MCP ToolsetRegistry. Reads UStateTreeEditorData (source/design-time data) for states, tasks, conditions, transitions, evaluators, and node descriptions.

Changes:
- New plugin: Engine/Plugins/Experimental/StateTreeToolset (9 tools, 12 tests)
- Static UFUNCTION GetEditorData(UStateTree*) on UStateTreeEditorData for typed access without casting
- UFUNCTION on UStateTreeEditorData::GetNodeDescription for node descriptions from Python
- BlueprintReadOnly on UStateTreeEditorData::SubTrees/Evaluators/GlobalTasks
- BlueprintReadOnly on UStateTreeState key fields (Name, Tag, Type, Tasks, EnterConditions, Transitions, Children, LinkedAsset)
- BlueprintType + BlueprintReadOnly on FStateTreeEditorNode (Node, Instance, InstanceObject, ExecutionRuntimeData, ExecutionRuntimeDataObject)
- BlueprintType + BlueprintReadOnly on FStateTreeTransition (Trigger, State, Conditions, bTransitionEnabled)

#Synced-CL 52127619

diff --git a/Engine/Plugins/Experimental/StateTreeToolset/Source/StateTreeToolset/Private/StateTreeToolset.cpp b/Engine/Plugins/Experimental/StateTreeToolset/Source/StateTreeToolset/Private/StateTreeToolset.cpp
new file mode 100644
--- /dev/null
+++ b/Engine/Plugins/Experimental/StateTreeToolset/Source/StateTreeToolset/Private/StateTreeToolset.cpp
@@ -0,0 +1,16 @@
+// Copyright Epic Games, Inc. All Rights Reserved.
+
+#include "StateTreeToolset.h"
+
+#include "Modules/ModuleManager.h"
+
+/*virtual*/ void FStateTreeToolsetModule::StartupModule()
+{
+}
+
+/*virtual*/ void FStateTreeToolsetModule::ShutdownModule()
+{
+}
+
+IMPLEMENT_MODULE(FStateTreeToolsetModule, StateTreeToolset);
+DEFINE_LOG_CATEGORY(LogStateTreeToolset);
diff --git a/Engine/Plugins/Experimental/StateTreeToolset/Source/StateTreeToolset/Public/StateTreeToolset.h b/Engine/Plugins/Experimental/StateTreeToolset/Source/StateTreeToolset/Public/StateTreeToolset.h
new file mode 100644
--- /dev/null
+++ b/Engine/Plugins/Experimental/StateTreeToolset/Source/StateTreeToolset/Public/StateTreeToolset.h
@@ -0,0 +1,15 @@
+// Copyright Epic Games, Inc. All Rights Reserved.
+
+#pragma once
+
+#include "Modules/ModuleInterface.h"
+#include "Logging/LogMacros.h"
+
+DECLARE_LOG_CATEGORY_EXTERN(LogStateTreeToolset, Log, Log);
+
+class FStateTreeToolsetModule : public IModuleInterface
+{
+public:
+	virtual void StartupModule() override;
+	virtual void ShutdownModule() override;
+};
diff --git a/Engine/Plugins/Experimental/StateTreeToolset/Source/StateTreeToolset/StateTreeToolset.Build.cs b/Engine/Plugins/Experimental/StateTreeToolset/Source/StateTreeToolset/StateTreeToolset.Build.cs
new file mode 100644
--- /dev/null
+++ b/Engine/Plugins/Experimental/StateTreeToolset/Source/StateTreeToolset/StateTreeToolset.Build.cs
@@ -0,0 +1,12 @@
+// Copyright Epic Games, Inc. All Rights Reserved.
+
+using UnrealBuildTool;
+
+public class StateTreeToolset : ModuleRules
+{
+	public StateTreeToolset(ReadOnlyTargetRules Target) : base(Target)
+	{
+		PCHUsage = PCHUsageMode.UseExplicitOrSharedPCHs;
+		PublicDependencyModuleNames.Add("Core");
+	}
+}
diff --git a/Engine/Plugins/Experimental/StateTreeToolset/StateTreeToolset.uplugin b/Engine/Plugins/Experimental/StateTreeToolset/StateTreeToolset.uplugin
new file mode 100644
--- /dev/null
+++ b/Engine/Plugins/Experimental/StateTreeToolset/StateTreeToolset.uplugin
@@ -0,0 +1,38 @@
+{
+	"FileVersion": 3,
+	"Version": 1,
+	"VersionName": "1.0",
+	"FriendlyName": "StateTreeToolset",
+	"Description": "Toolset for StateTree Inspection",
+	"Category": "Other",
+	"CreatedBy": "Epic Games",
+	"CreatedByURL": "",
+	"DocsURL": "",
+	"MarketplaceURL": "",
+	"SupportURL": "",
+	"EditorOnly": true,
+	"CanContainContent": true,
+	"IsBetaVersion": false,
+	"IsExperimentalVersion": true,
+	"NoRedist": true,
+	"Installed": false,
+	"EnabledByDefault": true,
+	"Modules": [
+		{
+			"Name": "StateTreeToolset",
+			"Type": "Editor",
+			"LoadingPhase": "Default",
+			"TargetAllowList": ["Editor"]
+		}
+	],
+	"Plugins": [
+		{
+			"Name": "StateTree",
+			"Enabled": true
+		},
+		{
+			"Name": "ToolsetRegistry",
+			"Enabled": true
+		}
+	]
+}
diff --git a/Engine/Plugins/Runtime/StateTree/Source/StateTreeEditorModule/Private/StateTreeEditorData.cpp b/Engine/Plugins/Runtime/StateTree/Source/StateTreeEditorModule/Private/StateTreeEditorData.cpp
--- a/Engine/Plugins/Runtime/StateTree/Source/StateTreeEditorModule/Private/StateTreeEditorData.cpp
+++ b/Engine/Plugins/Runtime/StateTree/Source/StateTreeEditorModule/Private/StateTreeEditorData.cpp
@@ -95,6 +95,15 @@ UStateTreeEditorData::UStateTreeEditorData()
 	EditorBindings.SetBindingsOwner(this);
 }
 
+/*static*/ UStateTreeEditorData* UStateTreeEditorData::GetEditorData(UStateTree* StateTree)
+{
+	if (StateTree == nullptr)
+	{
+		return nullptr;
+	}
+	return Cast<UStateTreeEditorData>(StateTree->EditorData);
+}
+
 void UStateTreeEditorData::PostInitProperties()
 {
 	Super::PostInitProperties();
diff --git a/Engine/Plugins/Runtime/StateTree/Source/StateTreeEditorModule/Public/StateTreeEditorData.h b/Engine/Plugins/Runtime/StateTree/Source/StateTreeEditorModule/Public/StateTreeEditorData.h
--- a/Engine/Plugins/Runtime/StateTree/Source/StateTreeEditorModule/Public/StateTreeEditorData.h
+++ b/Engine/Plugins/Runtime/StateTree/Source/StateTreeEditorModule/Public/StateTreeEditorData.h
@@ -69,6 +69,10 @@ class UStateTreeEditorData : public UObject, public IStateTreeEditorPropertyBind
 public:
 	UE_API UStateTreeEditorData();
 
+	/** @return The editor data for a StateTree asset, or nullptr if unavailable. */
+	UFUNCTION(BlueprintCallable, Category = "StateTree")
+	static UStateTreeEditorData* GetEditorData(UStateTree* StateTree);
+
 	UE_API virtual void PostInitProperties() override;
 	UE_API virtual void Serialize(FArchive& Ar) override;
 	
@@ -112,6 +116,7 @@ public:
 	 * Handles the name override logic, figures out required data for the GetDescription() call, and handles the fallbacks.
 	 * @return description for the node.
 	 */
+	UFUNCTION(BlueprintCallable, Category = "StateTree")
 	UE_API FText GetNodeDescription(const FStateTreeEditorNode& Node, const EStateTreeNodeFormatting Formatting = EStateTreeNodeFormatting::Text) const;
 
 #if WITH_EDITOR
@@ -429,10 +434,10 @@ private:
 	FInstancedPropertyBag RootParameterPropertyBag;
 
 public:
-	UPROPERTY(EditDefaultsOnly, Category = "Evaluators", meta = (BaseStruct = "/Script/StateTreeModule.StateTreeEvaluatorBase", BaseClass = "/Script/StateTreeModule.StateTreeEvaluatorBlueprintBase"))
+	UPROPERTY(EditDefaultsOnly, BlueprintReadOnly, Category = "Evaluators", meta = (BaseStruct = "/Script/StateTreeModule.StateTreeEvaluatorBase", BaseClass = "/Script/StateTreeModule.StateTreeEvaluatorBlueprintBase"))
 	TArray<FStateTreeEditorNode> Evaluators;
 
-	UPROPERTY(EditDefaultsOnly, Category = "Global Tasks", meta = (BaseStruct = "/Script/StateTreeModule.StateTreeTaskBase", BaseClass = "/Script/StateTreeModule.StateTreeTaskBlueprintBase"))
+	UPROPERTY(EditDefaultsOnly, BlueprintReadOnly, Category = "Global Tasks", meta = (BaseStruct = "/Script/StateTreeModule.StateTreeTaskBase", BaseClass = "/Script/StateTreeModule.StateTreeTaskBlueprintBase"))
 	TArray<FStateTreeEditorNode> GlobalTasks;
 
 	UPROPERTY(EditDefaultsOnly, Category = "Global Tasks")
@@ -446,7 +451,7 @@ public:
 	TSet<FStateTreeEditorColor> Colors;
 
 	/** Top level States. */
-	UPROPERTY(Instanced)
+	UPROPERTY(Instanced, BlueprintReadOnly, Category = "StateTree")
 	TArray<TObjectPtr<UStateTreeState>> SubTrees;
 
 	/**
diff --git a/Engine/Plugins/Runtime/StateTree/Source/StateTreeEditorModule/Public/StateTreeEditorNode.h b/Engine/Plugins/Runtime/StateTree/Source/StateTreeEditorModule/Public/StateTreeEditorNode.h
--- a/Engine/Plugins/Runtime/StateTree/Source/StateTreeEditorModule/Public/StateTreeEditorNode.h
+++ b/Engine/Plugins/Runtime/StateTree/Source/StateTreeEditorModule/Public/StateTreeEditorNode.h
@@ -21,7 +21,7 @@ enum class EStateTreeNodeType : uint8
 /**
  * Base for Evaluator, Task, Condition, Consideration Utilities, Property Function nodes.
  */
-USTRUCT()
+USTRUCT(BlueprintType)
 struct FStateTreeEditorNode
 {
 	GENERATED_BODY()
@@ -90,19 +90,19 @@ struct FStateTreeEditorNode
 		return ExecutionRuntimeDataObject ? FStateTreeDataView(ExecutionRuntimeDataObject) : FStateTreeDataView(ExecutionRuntimeData);
 	}
 
-	UPROPERTY(EditDefaultsOnly, Category = Node)
+	UPROPERTY(EditDefaultsOnly, BlueprintReadOnly, Category = Node)
 	FInstancedStruct Node;
 
-	UPROPERTY(EditDefaultsOnly, Category = Node)
+	UPROPERTY(EditDefaultsOnly, BlueprintReadOnly, Category = Node)
 	FInstancedStruct Instance;
 
-	UPROPERTY(EditDefaultsOnly, Instanced, Category = Node)
+	UPROPERTY(EditDefaultsOnly, Instanced, BlueprintReadOnly, Category = Node)
 	TObjectPtr<UObject> InstanceObject = nullptr;
 
-	UPROPERTY(EditDefaultsOnly, Category = Node)
+	UPROPERTY(EditDefaultsOnly, BlueprintReadOnly, Category = Node)
 	FInstancedStruct ExecutionRuntimeData;
 
-	UPROPERTY(EditDefaultsOnly, Instanced, Category = Node)
+	UPROPERTY(EditDefaultsOnly, Instanced, BlueprintReadOnly, Category = Node)
 	TObjectPtr<UObject> ExecutionRuntimeDataObject = nullptr;
 
 	/** ID for the node instance. */
diff --git a/Engine/Plugins/Runtime/StateTree/Source/StateTreeEditorModule/Public/StateTreeState.h b/Engine/Plugins/Runtime/StateTree/Source/StateTreeEditorModule/Public/StateTreeState.h
--- a/Engine/Plugins/Runtime/StateTree/Source/StateTreeEditorModule/Public/StateTreeState.h
+++ b/Engine/Plugins/Runtime/StateTree/Source/StateTreeEditorModule/Public/StateTreeState.h
@@ -76,7 +76,7 @@ struct FStateTreeTransitionDelegateListener
 /**
  * Editor representation of a transition in StateTree
  */
-USTRUCT()
+USTRUCT(BlueprintType)
 struct FStateTreeTransition
 {
 	GENERATED_BODY()
@@ -123,7 +123,7 @@ PRAGMA_ENABLE_DEPRECATION_WARNINGS
 	UE_API void PostSerialize(const FArchive& Ar);
 
 	/** When to try trigger the transition. */
-	UPROPERTY(EditDefaultsOnly, Category = "Transition")
+	UPROPERTY(EditDefaultsOnly, BlueprintReadOnly, Category = "Transition")
 	EStateTreeTransitionTrigger Trigger = EStateTreeTransitionTrigger::OnStateCompleted;
 
 	/** Defines the event required to be present during state selection for the transition to trigger. */
@@ -131,7 +131,7 @@ PRAGMA_ENABLE_DEPRECATION_WARNINGS
 	FStateTreeEventDesc RequiredEvent; 
 
 	/** Transition target state. */
-	UPROPERTY(EditDefaultsOnly, Category = "Transition", meta=(DisplayName="Transition To"))
+	UPROPERTY(EditDefaultsOnly, BlueprintReadOnly, Category = "Transition", meta=(DisplayName="Transition To"))
 	FStateTreeStateLink State;
 
 	UPROPERTY(EditDefaultsOnly, Category = "Transition")
@@ -162,7 +162,7 @@ PRAGMA_ENABLE_DEPRECATION_WARNINGS
 	float DelayRandomVariance = 0.0f;
 
 	/** Expression of conditions that need to evaluate to true to allow transition to be triggered. */
-	UPROPERTY(EditDefaultsOnly, Category = "Transition", meta = (BaseStruct = "/Script/StateTreeModule.StateTreeConditionBase", BaseClass = "/Script/StateTreeModule.StateTreeConditionBlueprintBase"))
+	UPROPERTY(EditDefaultsOnly, BlueprintReadOnly, Category = "Transition", meta = (BaseStruct = "/Script/StateTreeModule.StateTreeConditionBase", BaseClass = "/Script/StateTreeModule.StateTreeConditionBlueprintBase"))
 	TArray<FStateTreeEditorNode> Conditions;
 
 	/** 
@@ -173,7 +173,7 @@ PRAGMA_ENABLE_DEPRECATION_WARNINGS
 	EStateTreeTransitionChangeTypeRules ReactivateTargetState = EStateTreeTransitionChangeTypeRules::Default;
 
 	/** True if the Transition is Enabled (i.e. not explicitly disabled in the asset). */
-	UPROPERTY(EditDefaultsOnly, Category = "Debug")
+	UPROPERTY(EditDefaultsOnly, BlueprintReadOnly, Category = "Debug")
 	bool bTransitionEnabled = true;
 
 #if WITH_EDITORONLY_DATA
@@ -386,7 +386,7 @@ public:
 	//~ Note: these properties are customized out in FStateTreeStateDetails, adding a new property might require to adjust the customization.
 	
 	/** Display name of the State */
-	UPROPERTY(EditDefaultsOnly, Category = "State")
+	UPROPERTY(EditDefaultsOnly, BlueprintReadOnly, Category = "State")
 	FName Name;
 
 	/** Description of the State */
@@ -394,7 +394,7 @@ public:
 	FString Description;
 
 	/** GameplayTag describing the State */
-	UPROPERTY(EditDefaultsOnly, Category = "State")
+	UPROPERTY(EditDefaultsOnly, BlueprintReadOnly, Category = "State")
 	FGameplayTag Tag;
 
 	/** Display color of the State */
@@ -402,7 +402,7 @@ public:
 	FStateTreeEditorColorRef ColorRef;
 
 	/** Type the State, allows e.g. states to be linked to other States. */
-	UPROPERTY(EditDefaultsOnly, Category = "State")
+	UPROPERTY(EditDefaultsOnly, BlueprintReadOnly, Category = "State")
 	EStateTreeStateType Type = EStateTreeStateType::State;
 
 	/** How to treat child states when this State is selected.  */
@@ -418,7 +418,7 @@ public:
 	FStateTreeStateLink LinkedSubtree;
 
 	/** Another State Tree asset to run as extension of this State. */
-	UPROPERTY(EditDefaultsOnly, Category = "State")
+	UPROPERTY(EditDefaultsOnly, BlueprintReadOnly, Category = "State")
 	TObjectPtr<UStateTree> LinkedAsset = nullptr;
 
 	/** 
@@ -455,10 +455,10 @@ public:
 	float Weight = 1.f;
 
 	/** Expression of enter conditions that needs to evaluate true to allow the state to be selected. */
-	UPROPERTY(EditDefaultsOnly, Category = "Enter Conditions", meta = (BaseStruct = "/Script/StateTreeModule.StateTreeConditionBase", BaseClass = "/Script/StateTreeModule.StateTreeConditionBlueprintBase"))
+	UPROPERTY(EditDefaultsOnly, BlueprintReadOnly, Category = "Enter Conditions", meta = (BaseStruct = "/Script/StateTreeModule.StateTreeConditionBase", BaseClass = "/Script/StateTreeModule.StateTreeConditionBlueprintBase"))
 	TArray<FStateTreeEditorNode> EnterConditions;
 
-	UPROPERTY(EditDefaultsOnly, Category = "Tasks", meta = (BaseStruct = "/Script/StateTreeModule.StateTreeTaskBase", BaseClass = "/Script/StateTreeModule.StateTreeTaskBlueprintBase"))
+	UPROPERTY(EditDefaultsOnly, BlueprintReadOnly, Category = "Tasks", meta = (BaseStruct = "/Script/StateTreeModule.StateTreeTaskBase", BaseClass = "/Script/StateTreeModule.StateTreeTaskBlueprintBase"))
 	TArray<FStateTreeEditorNode> Tasks;
 
 	/** Expression used to calculate a selection utility score for state weighing. */
@@ -469,10 +469,10 @@ public:
 	UPROPERTY(EditDefaultsOnly, Category = "Task", meta = (BaseStruct = "/Script/StateTreeModule.StateTreeTaskBase", BaseClass = "/Script/StateTreeModule.StateTreeTaskBlueprintBase"))
 	FStateTreeEditorNode SingleTask;
 
-	UPROPERTY(EditDefaultsOnly, Category = "Transitions")
+	UPROPERTY(EditDefaultsOnly, BlueprintReadOnly, Category = "Transitions")
 	TArray<FStateTreeTransition> Transitions;
 
-	UPROPERTY(Instanced)
+	UPROPERTY(Instanced, BlueprintReadOnly, Category = "State")
 	TArray<TObjectPtr<UStateTreeState>> Children;
 
 	UPROPERTY(EditDefaultsOnly, Category = "State", meta = (IgnoreForMemberInitializationTest))
