[StateTree] As a compiler optimizations, add a cvar to skip CheckCompiledStateTreeOuters() to detect invalid outer references in compiled StateTree data. Always called from the editor compile button. Enabled via the cvar StateTree.Compiler.bEnableCheckOutersOnCompilationSucceeded.
#Synced-CL 52128263

diff --git a/Engine/Plugins/Runtime/StateTree/Source/StateTreeEditorModule/Private/StateTreeCompiler.cpp b/Engine/Plugins/Runtime/StateTree/Source/StateTreeEditorModule/Private/StateTreeCompiler.cpp
--- a/Engine/Plugins/Runtime/StateTree/Source/StateTreeEditorModule/Private/StateTreeCompiler.cpp
+++ b/Engine/Plugins/Runtime/StateTree/Source/StateTreeEditorModule/Private/StateTreeCompiler.cpp
@@ -95,9 +95,17 @@ namespace UE::StateTree::Compiler
 		TEXT("Use EvaluationScope data for utility considerations instead of SharedInstance data.")
 	);
 
-	FAutoConsoleVariable CVarLogCompiledStateTree(
+	bool bEnableCheckOutersOnCompilationSucceeded = false;
+	FAutoConsoleVariableRef CVarEnableCheckOutersOnCompilationSucceeded(
+		TEXT("StateTree.Compiler.bEnableCheckOutersOnCompilationSucceeded"),
+		bEnableCheckOutersOnCompilationSucceeded,
+		TEXT("After a StateTree successfully compiles, test if there are invalid references in the compiled data.")
+	);
+
+	bool bLogResultOnCompilationCompleted = false;
+	FAutoConsoleVariableRef CVarLogCompiledStateTree(
 		TEXT("StateTree.Compiler.LogResultOnCompilationCompleted"),
-		false,
+		bLogResultOnCompilationCompleted,
 		TEXT("After a StateTree compiles, log the internal content of the StateTree.")
 	);
 
@@ -658,13 +666,12 @@ bool FStateTreeCompiler::CompileInternalImpl()
 		return FailCompilation(UStateTree::ECompileStatus::Internal);
 	}
 
+	if (UE::StateTree::Compiler::bEnableCheckOutersOnCompilationSucceeded)
 	{
-		// Test if there are invalid reference
-		UE::StateTree::Compiler::FCheckOutersArchive CheckOuters(StateTree, EditorData, StateTree->GlobalParametersStructType, Log);
-		StateTree->Serialize(CheckOuters);
+		CheckCompiledStateTreeOuters(StateTree, Log);
 	}
 
-	if (UE::StateTree::Compiler::CVarLogCompiledStateTree->GetBool())
+	if (UE::StateTree::Compiler::bLogResultOnCompilationCompleted)
 	{
 		UE_LOG(LogStateTreeEditor, Log, TEXT("%s"), *StateTree->DebugInternalLayoutAsString());
 	}
@@ -692,6 +699,20 @@ bool FStateTreeCompiler::Compile(UStateTree& InStateTree)
 	return Compile(&InStateTree);
 }
 
+void FStateTreeCompiler::CheckCompiledStateTreeOuters(TNotNull<const UStateTree*> InStateTree, FStateTreeCompilerLog& InLog)
+{
+	const UStateTreeEditorData* EditorData = Cast<UStateTreeEditorData>(InStateTree->EditorData);
+	if (!EditorData)
+	{
+		return;
+	}
+
+	// Test if there are invalid references in the compiled data.
+	// Serialize is non-const but does not modify the StateTree when used with a saving archive.
+	UE::StateTree::Compiler::FCheckOutersArchive CheckOuters(InStateTree, EditorData, InStateTree->GlobalParametersStructType, InLog);
+	const_cast<UStateTree*>(static_cast<const UStateTree*>(InStateTree))->Serialize(CheckOuters);
+}
+
 FStateTreeStateHandle FStateTreeCompiler::GetStateHandle(const FGuid& StateID) const
 {
 	const int32* Idx = IDToState.Find(StateID);
diff --git a/Engine/Plugins/Runtime/StateTree/Source/StateTreeEditorModule/Private/StateTreeEditorMode.cpp b/Engine/Plugins/Runtime/StateTree/Source/StateTreeEditorModule/Private/StateTreeEditorMode.cpp
--- a/Engine/Plugins/Runtime/StateTree/Source/StateTreeEditorModule/Private/StateTreeEditorMode.cpp
+++ b/Engine/Plugins/Runtime/StateTree/Source/StateTreeEditorModule/Private/StateTreeEditorMode.cpp
@@ -22,6 +22,7 @@
 
 #include "IStateTreeEditorHost.h"
 #include "MessageLogModule.h"
+#include "StateTreeCompiler.h"
 #include "StateTreeEditingSubsystem.h"
 #include "Customizations/StateTreeBindingExtension.h"
 
@@ -436,6 +437,19 @@ void UStateTreeEditorMode::Compile()
 	FStateTreeCompilerLog Log;
 	const bool bCompileSucceeded = UStateTreeEditingSubsystem::CompileStateTree(StateTree, Log);
 
+	if (bCompileSucceeded)
+	{
+		// We always want to check outers when compiling from the editor to catch invalid references early.
+		// The cvar bEnableCheckOutersOnCompilationSucceeded is an optimization that runs the same check
+		// inside the compiler pipeline itself. If it is already enabled, skip the check here to avoid
+		// running it twice.
+		const IConsoleVariable* CVarCheckOuters = IConsoleManager::Get().FindConsoleVariable(TEXT("StateTree.Compiler.bEnableCheckOutersOnCompilationSucceeded"));
+		if (!CVarCheckOuters || !CVarCheckOuters->GetBool())
+		{
+			FStateTreeCompiler::CheckCompiledStateTreeOuters(StateTree, Log);
+		}
+	}
+
 	if (TSharedPtr<IMessageLogListing> Listing = GetMessageLogListing())
 	{
 		Log.AppendToLog(Listing.Get());
diff --git a/Engine/Plugins/Runtime/StateTree/Source/StateTreeEditorModule/Public/StateTreeCompiler.h b/Engine/Plugins/Runtime/StateTree/Source/StateTreeEditorModule/Public/StateTreeCompiler.h
--- a/Engine/Plugins/Runtime/StateTree/Source/StateTreeEditorModule/Public/StateTreeCompiler.h
+++ b/Engine/Plugins/Runtime/StateTree/Source/StateTreeEditorModule/Public/StateTreeCompiler.h
@@ -66,6 +66,9 @@ public:
 	/** Compile the public and internal steps of a state tree asset. */
 	UE_API bool Compile(UStateTree& InStateTree);
 
+	/** Tests that all instanced sub-objects in the compiled StateTree have correct outers. Reports errors to the compiler log. */
+	UE_API static void CheckCompiledStateTreeOuters(TNotNull<const UStateTree*> InStateTree, FStateTreeCompilerLog& InLog);
+
 	/**
 	 * Append config values that can change how state tree asset compiles when cooked.
 	 * When the context is different from the previous result, then a cook is needed.
