[StateTree] Centralize FStateTreeEditorNode initialization via InitializeAs()
Added FStateTreeEditorNode::InitializeAs() that consolidate the previously duplicated initialization logic spread across multiple call sites.
Refactor AddConditionWithOuter() on FStateTreeTransition and deprecated the outer-less AddCondition(). AddCondition had variadic template arguments, and the 2 implementations conflicted.
Added GetBindingsOwnerObject() accessor to FPropertyBindingBindingCollection.

#Synced-CL 52130857

diff --git a/Engine/Plugins/Runtime/PropertyBindingUtils/Source/PropertyBindingUtils/Public/PropertyBindingBindingCollection.h b/Engine/Plugins/Runtime/PropertyBindingUtils/Source/PropertyBindingUtils/Public/PropertyBindingBindingCollection.h
--- a/Engine/Plugins/Runtime/PropertyBindingUtils/Source/PropertyBindingUtils/Public/PropertyBindingBindingCollection.h
+++ b/Engine/Plugins/Runtime/PropertyBindingUtils/Source/PropertyBindingUtils/Public/PropertyBindingBindingCollection.h
@@ -41,6 +41,12 @@ struct FPropertyBindingBindingCollection
 		return BindingsOwner.GetInterface();
 	}
 
+	/** Get the associated bindings owner as a UObject. */
+	UObject* GetBindingsOwnerObject() const
+	{
+		return BindingsOwner.GetObject();
+	}
+
 	/** Sets associated bindings owner, used to validate added property paths. */
 	void SetBindingsOwner(TScriptInterface<IPropertyBindingBindingCollectionOwner> InBindingsOwner)
 	{
diff --git a/Engine/Plugins/Runtime/StateTree/Source/StateTreeEditorModule/Private/Customizations/StateTreeEditorNodeUtils.cpp b/Engine/Plugins/Runtime/StateTree/Source/StateTreeEditorModule/Private/Customizations/StateTreeEditorNodeUtils.cpp
--- a/Engine/Plugins/Runtime/StateTree/Source/StateTreeEditorModule/Private/Customizations/StateTreeEditorNodeUtils.cpp
+++ b/Engine/Plugins/Runtime/StateTree/Source/StateTreeEditorModule/Private/Customizations/StateTreeEditorNodeUtils.cpp
@@ -470,44 +470,19 @@ void SetNodeTypeStruct(const TSharedPtr<IPropertyHandle>& StructProperty, const
 				const bool bRetainProperties = InStruct && UStateTreeEditorSettings::Get().bRetainNodePropertyValues;
 				FStateTreeEditorNode OldNode = bRetainProperties ? *Node : FStateTreeEditorNode();
 
-				Node->Reset();
-				
 				if (InStruct)
 				{
-					// Generate new ID.
-					Node->ID = FGuid::NewGuid();
-
-					// Initialize node
-					Node->Node.InitializeAs(InStruct);
-					
-					// Generate new name and instantiate instance data.
-					if (InStruct->IsChildOf(FStateTreeNodeBase::StaticStruct()))
-					{
-						FStateTreeNodeBase& NodeBase = Node->Node.GetMutable<FStateTreeNodeBase>();
-						if (const UScriptStruct* InstanceType = Cast<const UScriptStruct>(NodeBase.GetInstanceDataType()))
-						{
-							Node->Instance.InitializeAs(InstanceType);
-						}
-						else if (const UClass* InstanceClass = Cast<const UClass>(NodeBase.GetInstanceDataType()))
-						{
-							Node->InstanceObject = NewObject<UObject>(Outer, InstanceClass);
-						}
-
-						if (const UScriptStruct* InstanceType = Cast<const UScriptStruct>(NodeBase.GetExecutionRuntimeDataType()))
-						{
-							Node->ExecutionRuntimeData.InitializeAs(InstanceType);
-						}
-						else if (const UClass* InstanceClass = Cast<const UClass>(NodeBase.GetExecutionRuntimeDataType()))
-						{
-							Node->ExecutionRuntimeDataObject = NewObject<UObject>(Outer, InstanceClass);
-						}
-					}
+					Node->InitializeAs(Outer, InStruct);
 
 					if (bRetainProperties)
 					{
 						RetainProperties(OldNode, *Node);
 					}
 				}
+				else
+				{
+					Node->Reset();
+				}
 			}
 		}
 	}
@@ -534,52 +509,43 @@ void SetNodeTypeClass(const TSharedPtr<IPropertyHandle>& StructProperty, const U
 				bool bRetainProperties = InClass && UStateTreeEditorSettings::Get().bRetainNodePropertyValues;
 				FStateTreeEditorNode OldNode = bRetainProperties ? *Node : FStateTreeEditorNode();
 
-				Node->Reset();
-
-				if (InClass && InClass->IsChildOf(UStateTreeTaskBlueprintBase::StaticClass()))
+				if (InClass)
 				{
-					Node->Node.InitializeAs(FStateTreeBlueprintTaskWrapper::StaticStruct());
-					FStateTreeBlueprintTaskWrapper& Task = Node->Node.GetMutable<FStateTreeBlueprintTaskWrapper>();
-					Task.TaskClass = const_cast<UClass*>(InClass);
-					
-					Node->InstanceObject = NewObject<UObject>(Outer, InClass);
-
-					Node->ID = FGuid::NewGuid();
-				}
-				else if (InClass && InClass->IsChildOf(UStateTreeEvaluatorBlueprintBase::StaticClass()))
-				{
-					Node->Node.InitializeAs(FStateTreeBlueprintEvaluatorWrapper::StaticStruct());
-					FStateTreeBlueprintEvaluatorWrapper& Eval = Node->Node.GetMutable<FStateTreeBlueprintEvaluatorWrapper>();
-					Eval.EvaluatorClass = const_cast<UClass*>(InClass);
-					
-					Node->InstanceObject = NewObject<UObject>(Outer, InClass);
-
-					Node->ID = FGuid::NewGuid();
-				}
-				else if (InClass && InClass->IsChildOf(UStateTreeConditionBlueprintBase::StaticClass()))
-				{
-					Node->Node.InitializeAs(FStateTreeBlueprintConditionWrapper::StaticStruct());
-					FStateTreeBlueprintConditionWrapper& Cond = Node->Node.GetMutable<FStateTreeBlueprintConditionWrapper>();
-					Cond.ConditionClass = const_cast<UClass*>(InClass);
-
-					Node->InstanceObject = NewObject<UObject>(Outer, InClass);
-
-					Node->ID = FGuid::NewGuid();
-				}
-				else if (InClass && InClass->IsChildOf(UStateTreeConsiderationBlueprintBase::StaticClass()))
-				{
-					Node->Node.InitializeAs(FStateTreeBlueprintConsiderationWrapper::StaticStruct());
-					FStateTreeBlueprintConsiderationWrapper& Consideration = Node->Node.GetMutable<FStateTreeBlueprintConsiderationWrapper>();
-					Consideration.ConsiderationClass = const_cast<UClass*>(InClass);
-
-					Node->InstanceObject = NewObject<UObject>(Outer, InClass);
-
-					Node->ID = FGuid::NewGuid();
+					if (InClass->IsChildOf(UStateTreeTaskBlueprintBase::StaticClass()))
+					{
+						FStateTreeBlueprintTaskWrapper TaskWrapper;
+						TaskWrapper.TaskClass = const_cast<UClass*>(InClass);
+						Node->InitializeAs<FStateTreeBlueprintTaskWrapper>(Outer, MoveTemp(TaskWrapper));
+					}
+					else if (InClass->IsChildOf(UStateTreeEvaluatorBlueprintBase::StaticClass()))
+					{
+						FStateTreeBlueprintEvaluatorWrapper EvaluatorWrapper;
+						EvaluatorWrapper.EvaluatorClass = const_cast<UClass*>(InClass);
+						Node->InitializeAs<FStateTreeBlueprintEvaluatorWrapper>(Outer, MoveTemp(EvaluatorWrapper));
+					}
+					else if (InClass->IsChildOf(UStateTreeConditionBlueprintBase::StaticClass()))
+					{
+						FStateTreeBlueprintConditionWrapper ConditionWrapper;
+						ConditionWrapper.ConditionClass = const_cast<UClass*>(InClass);
+						Node->InitializeAs<FStateTreeBlueprintConditionWrapper>(Outer, MoveTemp(ConditionWrapper));
+					}
+					else if (InClass->IsChildOf(UStateTreeConsiderationBlueprintBase::StaticClass()))
+					{
+						FStateTreeBlueprintConsiderationWrapper ConsiderationWrapper;
+						ConsiderationWrapper.ConsiderationClass = const_cast<UClass*>(InClass);
+						Node->InitializeAs<FStateTreeBlueprintConsiderationWrapper>(Outer, MoveTemp(ConsiderationWrapper));
+					}
+					else
+					{
+						// Not retaining properties if we haven't initialized a new node
+						bRetainProperties = false;
+						Node->Reset();
+					}
 				}
 				else
 				{
-					// Not retaining properties if we haven't initialized a new node
 					bRetainProperties = false;
+					Node->Reset();
 				}
 
 				if (bRetainProperties)
diff --git a/Engine/Plugins/Runtime/StateTree/Source/StateTreeEditorModule/Private/StateTreeEditorData.cpp b/Engine/Plugins/Runtime/StateTree/Source/StateTreeEditorModule/Private/StateTreeEditorData.cpp
--- a/Engine/Plugins/Runtime/StateTree/Source/StateTreeEditorModule/Private/StateTreeEditorData.cpp
+++ b/Engine/Plugins/Runtime/StateTree/Source/StateTreeEditorModule/Private/StateTreeEditorData.cpp
@@ -104,6 +104,14 @@ UStateTreeEditorData::UStateTreeEditorData()
 	return Cast<UStateTreeEditorData>(StateTree->EditorData);
 }
 
+void UStateTreeEditorData::OnNodeAdded(FStateTreeEditorNode& EditorNode)
+{
+	if (UStateTree* StateTree = GetTypedOuter<UStateTree>())
+	{
+		UStateTreeEditingSubsystem::MarkAsModified(StateTree);
+	}
+}
+
 void UStateTreeEditorData::PostInitProperties()
 {
 	Super::PostInitProperties();
diff --git a/Engine/Plugins/Runtime/StateTree/Source/StateTreeEditorModule/Private/StateTreeEditorNode.cpp b/Engine/Plugins/Runtime/StateTree/Source/StateTreeEditorModule/Private/StateTreeEditorNode.cpp
--- a/Engine/Plugins/Runtime/StateTree/Source/StateTreeEditorModule/Private/StateTreeEditorNode.cpp
+++ b/Engine/Plugins/Runtime/StateTree/Source/StateTreeEditorModule/Private/StateTreeEditorNode.cpp
@@ -2,6 +2,7 @@
 
 #include "StateTreeEditorNode.h"
 
+#include "StateTreeEditorModule.h"
 #include "Blueprint/StateTreeConditionBlueprintBase.h"
 #include "Blueprint/StateTreeConsiderationBlueprintBase.h"
 #include "Blueprint/StateTreeEvaluatorBlueprintBase.h"
@@ -9,6 +10,35 @@
 
 #include UE_INLINE_GENERATED_CPP_BY_NAME(StateTreeEditorNode)
 
+void FStateTreeEditorNode::InitializeAs(TNotNull<UObject*> InOuter, TNotNull<const UScriptStruct*> InNodeStruct)
+{
+	Reset();
+	Node.InitializeAs(InNodeStruct);
+	InitializeInstanceData(InOuter);
+}
+
+void FStateTreeEditorNode::InitializeInstanceData(TNotNull<UObject*> InOuter)
+{
+	ID = FGuid::NewGuid();
+	const FStateTreeNodeBase& NodeBase = GetNode().Get<FStateTreeNodeBase>();
+	if (const UScriptStruct* InstanceType = Cast<const UScriptStruct>(NodeBase.GetInstanceDataType()))
+	{
+		Instance.InitializeAs(InstanceType);
+	}
+	else if (const UClass* InstanceClass = Cast<const UClass>(NodeBase.GetInstanceDataType()))
+	{
+		InstanceObject = NewObject<UObject>(InOuter, InstanceClass);
+	}
+	if (const UScriptStruct* RuntimeType = Cast<const UScriptStruct>(NodeBase.GetExecutionRuntimeDataType()))
+	{
+		ExecutionRuntimeData.InitializeAs(RuntimeType);
+	}
+	else if (const UClass* RuntimeClass = Cast<const UClass>(NodeBase.GetExecutionRuntimeDataType()))
+	{
+		ExecutionRuntimeDataObject = NewObject<UObject>(InOuter, RuntimeClass);
+	}
+}
+
 FName FStateTreeEditorNode::GetName() const
 {
 	const UScriptStruct* NodeType = Node.GetScriptStruct();
diff --git a/Engine/Plugins/Runtime/StateTree/Source/StateTreeEditorModule/Private/StateTreeEditorPropertyBindings.cpp b/Engine/Plugins/Runtime/StateTree/Source/StateTreeEditorModule/Private/StateTreeEditorPropertyBindings.cpp
--- a/Engine/Plugins/Runtime/StateTree/Source/StateTreeEditorModule/Private/StateTreeEditorPropertyBindings.cpp
+++ b/Engine/Plugins/Runtime/StateTree/Source/StateTreeEditorModule/Private/StateTreeEditorPropertyBindings.cpp
@@ -3,10 +3,11 @@
 #include "StateTreeEditorPropertyBindings.h"
 
 #include "Algo/Copy.h"
+#include "StateTreeEditingSubsystem.h"
+#include "StateTreeEditorNode.h"
 #include "StateTreePropertyBindingCompiler.h"
 #include "StateTreeNodeBase.h"
 #include "StateTreePropertyFunctionBase.h"
-#include "StateTreeEditorNode.h"
 
 #include UE_INLINE_GENERATED_CPP_BY_NAME(StateTreeEditorPropertyBindings)
 
@@ -22,28 +23,25 @@ UStateTreeEditorPropertyBindingsOwner::UStateTreeEditorPropertyBindingsOwner(con
 FPropertyBindingPath FStateTreeEditorPropertyBindings::AddFunctionBinding(const UScriptStruct* InPropertyFunctionNodeStruct, TConstArrayView<FPropertyBindingPathSegment> InSourcePathSegments, const FPropertyBindingPath& TargetPath)
 {
 	check(InPropertyFunctionNodeStruct->IsChildOf<FStateTreePropertyFunctionBase>());
+	UObject* OwnerObject = GetBindingsOwnerObject();
+	check(OwnerObject);
 
-	FInstancedStruct PropertyFunctionNode(FStateTreeEditorNode::StaticStruct());
-	FStateTreeEditorNode& PropertyFunction = PropertyFunctionNode.GetMutable<FStateTreeEditorNode>();
-	const FGuid NodeID = FGuid::NewGuid();
-	PropertyFunction.ID = NodeID;
-	PropertyFunction.Node.InitializeAs(InPropertyFunctionNodeStruct);
-	const FStateTreePropertyFunctionBase& Function = PropertyFunction.GetNode().Get<FStateTreePropertyFunctionBase>();
-	if (const UScriptStruct* InstanceType = Cast<const UScriptStruct>(Function.GetInstanceDataType()))
-	{
-		PropertyFunction.Instance.InitializeAs(InstanceType);
-	}
-	if (const UScriptStruct* InstanceType = Cast<const UScriptStruct>(Function.GetExecutionRuntimeDataType()))
-	{
-		PropertyFunction.ExecutionRuntimeData.InitializeAs(InstanceType);
-	}
+	FStateTreeEditorNode PropertyFunction;
+	PropertyFunction.InitializeAs(OwnerObject, InPropertyFunctionNodeStruct);
 
 	Super::RemoveBindings(TargetPath, ESearchMode::Exact);
-	FPropertyBindingPath SourcePath = FPropertyBindingPath(NodeID, InSourcePathSegments);
-	PropertyBindings.Emplace(MoveTemp(PropertyFunctionNode), SourcePath, TargetPath);
+
+	FConstStructView PropertyFunctionView = FConstStructView::Make(PropertyFunction);
+	FPropertyBindingPath SourcePath = FPropertyBindingPath(PropertyFunction.GetInstanceDataID(), InSourcePathSegments);
+	PropertyBindings.Emplace(PropertyFunctionView, SourcePath, TargetPath);
 
 	UpdateSegmentsForNewlyAddedBinding(PropertyBindings.Last());
 
+	if (UStateTree* StateTree = OwnerObject->GetTypedOuter<UStateTree>())
+	{
+		UStateTreeEditingSubsystem::MarkAsModified(StateTree);
+	}
+
 	return SourcePath;
 }
 
diff --git a/Engine/Plugins/Runtime/StateTree/Source/StateTreeEditorModule/Private/StateTreeState.cpp b/Engine/Plugins/Runtime/StateTree/Source/StateTreeEditorModule/Private/StateTreeState.cpp
--- a/Engine/Plugins/Runtime/StateTree/Source/StateTreeEditorModule/Private/StateTreeState.cpp
+++ b/Engine/Plugins/Runtime/StateTree/Source/StateTreeEditorModule/Private/StateTreeState.cpp
@@ -66,6 +66,14 @@ PRAGMA_ENABLE_DEPRECATION_WARNINGS
 #endif // WITH_EDITORONLY_DATA
 }
 
+void FStateTreeTransition::OnNodeAdded(TNotNull<UStateTreeState*> InState, FStateTreeEditorNode& EditorNode)
+{
+	if (UStateTree* StateTree = InState->GetTypedOuter<UStateTree>())
+	{
+		UStateTreeEditingSubsystem::MarkAsModified(StateTree);
+	}
+}
+
 //////////////////////////////////////////////////////////////////////////
 // UStateTreeState
 
@@ -672,18 +680,8 @@ UStateTreeState& UStateTreeState::AddChildState(const FName ChildName, const ESt
 	return *ChildState;
 }
 
-void UStateTreeState::AddNodeImpl(FStateTreeEditorNode& EditorNode)
+void UStateTreeState::OnNodeAdded(FStateTreeEditorNode& EditorNode)
 {
-	EditorNode.ID = FGuid::NewGuid();
-	const FStateTreeNodeBase& Node = EditorNode.GetNode().Get<FStateTreeNodeBase>();
-	if (const UScriptStruct* InstanceType = Cast<const UScriptStruct>(Node.GetInstanceDataType()))
-	{
-		EditorNode.Instance.InitializeAs(InstanceType);
-	}
-	if (const UScriptStruct* InstanceType = Cast<const UScriptStruct>(Node.GetExecutionRuntimeDataType()))
-	{
-		EditorNode.ExecutionRuntimeData.InitializeAs(InstanceType);
-	}
 	if (UStateTree* StateTree = GetTypedOuter<UStateTree>())
 	{
 		UStateTreeEditingSubsystem::MarkAsModified(StateTree);
diff --git a/Engine/Plugins/Runtime/StateTree/Source/StateTreeEditorModule/Public/StateTreeEditorData.h b/Engine/Plugins/Runtime/StateTree/Source/StateTreeEditorModule/Public/StateTreeEditorData.h
--- a/Engine/Plugins/Runtime/StateTree/Source/StateTreeEditorModule/Public/StateTreeEditorData.h
+++ b/Engine/Plugins/Runtime/StateTree/Source/StateTreeEditorModule/Public/StateTreeEditorData.h
@@ -296,17 +296,8 @@ public:
 	TStateTreeEditorNode<T>& AddEvaluator(TArgs&&... InArgs)
 	{
 		FStateTreeEditorNode& EditorNode = Evaluators.AddDefaulted_GetRef();
-		EditorNode.ID = FGuid::NewGuid();
-		EditorNode.Node.InitializeAs<T>(Forward<TArgs>(InArgs)...);
-		T& Eval = EditorNode.Node.GetMutable<T>();
-		if (const UScriptStruct* InstanceType = Cast<const UScriptStruct>(Eval.GetInstanceDataType()))
-		{
-			EditorNode.Instance.InitializeAs(InstanceType);
-		}
-		if (const UScriptStruct* InstanceType = Cast<const UScriptStruct>(Eval.GetExecutionRuntimeDataType()))
-		{
-			EditorNode.ExecutionRuntimeData.InitializeAs(InstanceType);
-		}
+		EditorNode.InitializeAs<T>(this, Forward<TArgs>(InArgs)...);
+		OnNodeAdded(EditorNode);
 		return static_cast<TStateTreeEditorNode<T>&>(EditorNode);
 	}
 
@@ -318,17 +309,8 @@ public:
 	TStateTreeEditorNode<T>& AddGlobalTask(TArgs&&... InArgs)
 	{
 		FStateTreeEditorNode& EditorNode = GlobalTasks.AddDefaulted_GetRef();
-		EditorNode.ID = FGuid::NewGuid();
-		EditorNode.Node.InitializeAs<T>(Forward<TArgs>(InArgs)...);
-		const FStateTreeNodeBase& Task = EditorNode.Node.Get<FStateTreeNodeBase>();
-		if (const UScriptStruct* InstanceType = Cast<const UScriptStruct>(Task.GetInstanceDataType()))
-		{
-			EditorNode.Instance.InitializeAs(InstanceType);
-		}
-		if (const UScriptStruct* InstanceType = Cast<const UScriptStruct>(Task.GetExecutionRuntimeDataType()))
-		{
-			EditorNode.ExecutionRuntimeData.InitializeAs(InstanceType);
-		}
+		EditorNode.InitializeAs<T>(this, Forward<TArgs>(InArgs)...);
+		OnNodeAdded(EditorNode);
 		return static_cast<TStateTreeEditorNode<T>&>(EditorNode);
 	}
 
@@ -400,6 +382,7 @@ private:
 	UE_API void FixObjectNodes();
 	UE_API void DuplicateIDs();
 	UE_API void CallPostLoadOnNodes();
+	UE_API void OnNodeAdded(FStateTreeEditorNode& EditorNode);
 
 #if WITH_EDITORONLY_DATA
 	FDelegateHandle OnParametersChangedHandle;
diff --git a/Engine/Plugins/Runtime/StateTree/Source/StateTreeEditorModule/Public/StateTreeEditorNode.h b/Engine/Plugins/Runtime/StateTree/Source/StateTreeEditorModule/Public/StateTreeEditorNode.h
--- a/Engine/Plugins/Runtime/StateTree/Source/StateTreeEditorModule/Public/StateTreeEditorNode.h
+++ b/Engine/Plugins/Runtime/StateTree/Source/StateTreeEditorModule/Public/StateTreeEditorNode.h
@@ -2,6 +2,7 @@
 
 #pragma once
 
+#include "Containers/Set.h"
 #include "StateTreeNodeBase.h"
 #include "StateTreeEditorNode.generated.h"
 
@@ -36,6 +37,32 @@ struct FStateTreeEditorNode
 		ID = FGuid();
 	}
 
+	/**
+	 * Initializes the node as the given type, resetting any previous state.
+	 * Handles both instance data and execution runtime data types.
+	 * @param InOuter Outer object used when creating UObject instance or execution runtime data.
+	 */
+	template<typename T, typename... TArgs>
+	void InitializeAs(TNotNull<UObject*> InOuter, TArgs&&... InArgs)
+	{
+		Reset();
+		Node.InitializeAs<T>(Forward<TArgs>(InArgs)...);
+		InitializeInstanceData(InOuter);
+	}
+
+	/**
+	 * Initializes the node with the given struct type, resetting any previous state.
+	 * Non-templated counterpart to the templated InitializeAs, for use when the node type is only known at runtime.
+	 * Handles both instance data and execution runtime data types.
+	 * @param InOuter Outer object used when creating UObject instance or execution runtime data.
+	 * @param InNodeStruct Node struct type to initialize with.
+	 */
+	UE_API void InitializeAs(TNotNull<UObject*> InOuter, TNotNull<const UScriptStruct*> InNodeStruct);
+
+private:
+	UE_API void InitializeInstanceData(TNotNull<UObject*> InOuter);
+
+public:
 	/**
 	 * This is used to name nodes for runtime, as well as for error reporting.
 	 * If the node has a specified name, used that, or else of return the display name of the node.
diff --git a/Engine/Plugins/Runtime/StateTree/Source/StateTreeEditorModule/Public/StateTreeState.h b/Engine/Plugins/Runtime/StateTree/Source/StateTreeEditorModule/Public/StateTreeState.h
--- a/Engine/Plugins/Runtime/StateTree/Source/StateTreeEditorModule/Public/StateTreeState.h
+++ b/Engine/Plugins/Runtime/StateTree/Source/StateTreeEditorModule/Public/StateTreeState.h
@@ -93,6 +93,16 @@ PRAGMA_DISABLE_DEPRECATION_WARNINGS
 PRAGMA_ENABLE_DEPRECATION_WARNINGS
 
 	template<typename T, typename... TArgs>
+	TStateTreeEditorNode<T>& AddConditionWithOuter(TNotNull<UStateTreeState*> InState, TArgs&&... InArgs)
+	{
+		FStateTreeEditorNode& CondNode = Conditions.AddDefaulted_GetRef();
+		CondNode.InitializeAs<T>(InState, Forward<TArgs>(InArgs)...);
+		OnNodeAdded(InState, CondNode);
+		return static_cast<TStateTreeEditorNode<T>&>(CondNode);
+	}
+
+	template<typename T, typename... TArgs>
+	UE_DEPRECATED(5.8, "Use AddConditionWithOuter that takes a TNotNull<UStateTreeState*> Outer.")
 	TStateTreeEditorNode<T>& AddCondition(TArgs&&... InArgs)
 	{
 		FStateTreeEditorNode& CondNode = Conditions.AddDefaulted_GetRef();
@@ -122,6 +132,11 @@ PRAGMA_ENABLE_DEPRECATION_WARNINGS
 
 	UE_API void PostSerialize(const FArchive& Ar);
 
+private:
+	UE_API void OnNodeAdded(TNotNull<UStateTreeState*> InState, FStateTreeEditorNode& EditorNode);
+
+public:
+
 	/** When to try trigger the transition. */
 	UPROPERTY(EditDefaultsOnly, BlueprintReadOnly, Category = "Transition")
 	EStateTreeTransitionTrigger Trigger = EStateTreeTransitionTrigger::OnStateCompleted;
@@ -275,8 +290,8 @@ public:
 	TStateTreeEditorNode<T>& AddEnterCondition(TArgs&&... InArgs)
 	{
 		FStateTreeEditorNode& CondNode = EnterConditions.AddDefaulted_GetRef();
-		CondNode.Node.InitializeAs<T>(Forward<TArgs>(InArgs)...);
-		AddNodeImpl(CondNode);
+		CondNode.InitializeAs<T>(this, Forward<TArgs>(InArgs)...);
+		OnNodeAdded(CondNode);
 		return static_cast<TStateTreeEditorNode<T>&>(CondNode);
 	}
 
@@ -288,8 +303,8 @@ public:
 	TStateTreeEditorNode<T>& AddConsideration(TArgs&&... InArgs)
 	{
 		FStateTreeEditorNode& ConsiderationNode = Considerations.AddDefaulted_GetRef();
-		ConsiderationNode.Node.InitializeAs<T>(Forward<TArgs>(InArgs)...);
-		AddNodeImpl(ConsiderationNode);
+		ConsiderationNode.InitializeAs<T>(this, Forward<TArgs>(InArgs)...);
+		OnNodeAdded(ConsiderationNode);
 		return static_cast<TStateTreeEditorNode<T>&>(ConsiderationNode);
 	}
 	
@@ -301,13 +316,13 @@ public:
 	TStateTreeEditorNode<T>& AddTask(TArgs&&... InArgs)
 	{
 		FStateTreeEditorNode& TaskItem = Tasks.AddDefaulted_GetRef();
-		TaskItem.Node.InitializeAs<T>(Forward<TArgs>(InArgs)...);
-		AddNodeImpl(TaskItem);
+		TaskItem.InitializeAs<T>(this, Forward<TArgs>(InArgs)...);
+		OnNodeAdded(TaskItem);
 		return static_cast<TStateTreeEditorNode<T>&>(TaskItem);
 	}
 
 private:
-	UE_API void AddNodeImpl(FStateTreeEditorNode& EditorNode);
+	UE_API void OnNodeAdded(FStateTreeEditorNode& EditorNode);
 
 public:
 	/** Sets linked state and updates parameters to match the linked state. */
diff --git a/Engine/Plugins/Runtime/StateTree/Source/StateTreeTestSuite/Private/StateTreeTest.cpp b/Engine/Plugins/Runtime/StateTree/Source/StateTreeTestSuite/Private/StateTreeTest.cpp
--- a/Engine/Plugins/Runtime/StateTree/Source/StateTreeTestSuite/Private/StateTreeTest.cpp
+++ b/Engine/Plugins/Runtime/StateTree/Source/StateTreeTestSuite/Private/StateTreeTest.cpp
@@ -58,7 +58,7 @@ struct FStateTreeTest_MakeAndBakeStateTree : FStateTreeTestBase
 		EditorData.AddPropertyBinding(EvalA, TEXT("bBoolA"), TaskB2, TEXT("bBoolB"));
 
 		FStateTreeTransition& Trans = StateB.AddTransition({}, EStateTreeTransitionType::GotoState, &Root);
-		auto& TransFloatCond = Trans.AddCondition<FStateTreeCompareFloatCondition>(UE::StateTree::EComparisonOperator::Less);
+		auto& TransFloatCond = Trans.AddConditionWithOuter<FStateTreeCompareFloatCondition>(&StateB, UE::StateTree::EComparisonOperator::Less);
 		TransFloatCond.GetInstanceData().Right = 13.0f;
 		EditorData.AddPropertyBinding(EvalA, TEXT("FloatA"), TransFloatCond, TEXT("Left"));
 
diff --git a/Engine/Plugins/Runtime/StateTree/Source/StateTreeTestSuite/Private/StateTreeTransitionTest.cpp b/Engine/Plugins/Runtime/StateTree/Source/StateTreeTestSuite/Private/StateTreeTransitionTest.cpp
--- a/Engine/Plugins/Runtime/StateTree/Source/StateTreeTestSuite/Private/StateTreeTransitionTest.cpp
+++ b/Engine/Plugins/Runtime/StateTree/Source/StateTreeTestSuite/Private/StateTreeTransitionTest.cpp
@@ -476,7 +476,7 @@ struct FStateTreeTest_TransitionReactivateTargetState : FStateTreeTestBase
 
 		// Reactivate transtion back to root, by default will not fire due to condition. But should reactivate / re-trigger enter state when activated.
 		FStateTreeTransition& ReactivateRootTransition = State0.AddTransition(EStateTreeTransitionTrigger::OnTick, EStateTreeTransitionType::GotoState, &Root);
-		TStateTreeEditorNode<FStateTreeCompareIntCondition>& TransIntCond = ReactivateRootTransition.AddCondition<FStateTreeCompareIntCondition>(UE::StateTree::EComparisonOperator::Equal);
+		TStateTreeEditorNode<FStateTreeCompareIntCondition>& TransIntCond = ReactivateRootTransition.AddConditionWithOuter<FStateTreeCompareIntCondition>(&State0, UE::StateTree::EComparisonOperator::Equal);
 		TransIntCond.GetInstanceData().Right = 0;
 		EditorData.AddPropertyBinding(
 			FPropertyBindingPath(EditorData.GetRootParametersGuid(), TEXT("Int")),
@@ -670,7 +670,7 @@ struct FStateTreeTest_PassingTransitionEventToStateSelection : FStateTreeTestBas
 		// Transition from Initial -> StateA
 		FStateTreeTransition& TransA = StateInitial.AddTransition(EStateTreeTransitionTrigger::OnEvent, FGameplayTag(), EStateTreeTransitionType::GotoState, &StateA);
 		TransA.RequiredEvent.PayloadStruct = FStateTreeTest_PropertyStructA::StaticStruct();
-		TStateTreeEditorNode<FStateTreeCompareIntCondition>& TransAIntCond = TransA.AddCondition<FStateTreeCompareIntCondition>(UE::StateTree::EComparisonOperator::Equal);
+		TStateTreeEditorNode<FStateTreeCompareIntCondition>& TransAIntCond = TransA.AddConditionWithOuter<FStateTreeCompareIntCondition>(&StateInitial, UE::StateTree::EComparisonOperator::Equal);
 		TransAIntCond.GetInstanceData().Right = 1;
 		EditorData.AddPropertyBinding(
 			FPropertyBindingPath(TransA.GetEventID(), PathToPayloadMember.GetSegments()),
@@ -678,7 +678,7 @@ struct FStateTreeTest_PassingTransitionEventToStateSelection : FStateTreeTestBas
 		// Transition from Initial -> StateB
 		FStateTreeTransition& TransB = StateInitial.AddTransition(EStateTreeTransitionTrigger::OnEvent, FGameplayTag(), EStateTreeTransitionType::GotoState, &StateB);
 		TransB.RequiredEvent.PayloadStruct = FStateTreeTest_PropertyStructA::StaticStruct();
-		TStateTreeEditorNode<FStateTreeCompareIntCondition>& TransBIntCond = TransB.AddCondition<FStateTreeCompareIntCondition>(UE::StateTree::EComparisonOperator::Equal);
+		TStateTreeEditorNode<FStateTreeCompareIntCondition>& TransBIntCond = TransB.AddConditionWithOuter<FStateTreeCompareIntCondition>(&StateInitial, UE::StateTree::EComparisonOperator::Equal);
 		TransBIntCond.GetInstanceData().Right = 1;
 		EditorData.AddPropertyBinding(
 			FPropertyBindingPath(TransB.GetEventID(), PathToPayloadMember.GetSegments()),
@@ -748,7 +748,7 @@ struct FStateTreeTest_FollowTransitions : FStateTreeTestBase
 			{
 				// This transition should be skipped due to the condition
 				FStateTreeTransition& TransA = StateTrans.AddTransition(EStateTreeTransitionTrigger::OnTick, EStateTreeTransitionType::GotoState, &StateA);
-				TStateTreeEditorNode<FStateTreeCompareIntCondition>& TransIntCond = TransA.AddCondition<FStateTreeCompareIntCondition>(UE::StateTree::EComparisonOperator::Equal);
+				TStateTreeEditorNode<FStateTreeCompareIntCondition>& TransIntCond = TransA.AddConditionWithOuter<FStateTreeCompareIntCondition>(&StateTrans, UE::StateTree::EComparisonOperator::Equal);
 				TransIntCond.GetInstanceData().Right = 0;
 				EditorData.AddPropertyBinding(
 					FPropertyBindingPath(EditorData.GetRootParametersGuid(), TEXT("Int")),
@@ -759,7 +759,7 @@ struct FStateTreeTest_FollowTransitions : FStateTreeTestBase
 				// This transition leads to selection, but will be overridden.
 				FStateTreeTransition& TransB = StateTrans.AddTransition(EStateTreeTransitionTrigger::OnTick, EStateTreeTransitionType::GotoState, &StateB);
 				TransB.Priority = EStateTreeTransitionPriority::Normal;
-				TStateTreeEditorNode<FStateTreeCompareIntCondition>& TransIntCond = TransB.AddCondition<FStateTreeCompareIntCondition>(UE::StateTree::EComparisonOperator::Equal);
+				TStateTreeEditorNode<FStateTreeCompareIntCondition>& TransIntCond = TransB.AddConditionWithOuter<FStateTreeCompareIntCondition>(&StateTrans, UE::StateTree::EComparisonOperator::Equal);
 				TransIntCond.GetInstanceData().Right = 1;
 				EditorData.AddPropertyBinding(
 					FPropertyBindingPath(EditorData.GetRootParametersGuid(), TEXT("Int")),
@@ -770,7 +770,7 @@ struct FStateTreeTest_FollowTransitions : FStateTreeTestBase
 				// This transition is selected, should override previous one due to priority.
 				FStateTreeTransition& TransC = StateTrans.AddTransition(EStateTreeTransitionTrigger::OnTick, EStateTreeTransitionType::GotoState, &StateC);
 				TransC.Priority = EStateTreeTransitionPriority::High;
-				TStateTreeEditorNode<FStateTreeCompareIntCondition>& TransIntCond = TransC.AddCondition<FStateTreeCompareIntCondition>(UE::StateTree::EComparisonOperator::Equal);
+				TStateTreeEditorNode<FStateTreeCompareIntCondition>& TransIntCond = TransC.AddConditionWithOuter<FStateTreeCompareIntCondition>(&StateTrans, UE::StateTree::EComparisonOperator::Equal);
 				TransIntCond.GetInstanceData().Right = 1;
 				EditorData.AddPropertyBinding(
 					FPropertyBindingPath(EditorData.GetRootParametersGuid(), TEXT("Int")),
@@ -833,7 +833,7 @@ struct FStateTreeTest_InfiniteLoop : FStateTreeTestBase
 			{
 				// A -> B
 				FStateTreeTransition& Trans = StateA.AddTransition(EStateTreeTransitionTrigger::OnTick, EStateTreeTransitionType::GotoState, &StateB);
-				TStateTreeEditorNode<FStateTreeCompareIntCondition>& TransIntCond = Trans.AddCondition<FStateTreeCompareIntCondition>(UE::StateTree::EComparisonOperator::Equal);
+				TStateTreeEditorNode<FStateTreeCompareIntCondition>& TransIntCond = Trans.AddConditionWithOuter<FStateTreeCompareIntCondition>(&StateA, UE::StateTree::EComparisonOperator::Equal);
 				TransIntCond.GetInstanceData().Right = 1;
 				EditorData.AddPropertyBinding(
 					FPropertyBindingPath(EditorData.GetRootParametersGuid(), TEXT("Int")),
@@ -847,7 +847,7 @@ struct FStateTreeTest_InfiniteLoop : FStateTreeTestBase
 			{
 				// B -> A
 				FStateTreeTransition& Trans = StateB.AddTransition(EStateTreeTransitionTrigger::OnTick, EStateTreeTransitionType::GotoState, &StateA);
-				TStateTreeEditorNode<FStateTreeCompareIntCondition>& TransIntCond = Trans.AddCondition<FStateTreeCompareIntCondition>(UE::StateTree::EComparisonOperator::Equal);
+				TStateTreeEditorNode<FStateTreeCompareIntCondition>& TransIntCond = Trans.AddConditionWithOuter<FStateTreeCompareIntCondition>(&StateB, UE::StateTree::EComparisonOperator::Equal);
 				TransIntCond.GetInstanceData().Right = 1;
 				EditorData.AddPropertyBinding(
 					FPropertyBindingPath(EditorData.GetRootParametersGuid(), TEXT("Int")),
