[StateTree] Refactor FStateTreeCompiler::CreateNodeInstanceData to use the proper getter.

#rnx
#Synced-CL 52180191

diff --git a/Engine/Plugins/Runtime/StateTree/Source/StateTreeEditorModule/Private/StateTreeCompiler.cpp b/Engine/Plugins/Runtime/StateTree/Source/StateTreeEditorModule/Private/StateTreeCompiler.cpp
--- a/Engine/Plugins/Runtime/StateTree/Source/StateTreeEditorModule/Private/StateTreeCompiler.cpp
+++ b/Engine/Plugins/Runtime/StateTree/Source/StateTreeEditorModule/Private/StateTreeCompiler.cpp
@@ -2144,108 +2144,117 @@ TOptional<FStateTreeDataView> FStateTreeCompiler::CreateNodeInstanceData(const F
 	// Update Node name as description for runtime.
 	Node.Name = EditorNode.GetName();
 
-	if (EditorNode.Instance.IsValid())
+	// Update the instance data
 	{
-		if (ensure(EditorNode.Instance.GetScriptStruct() == Node.GetInstanceDataType()))
+		const FStateTreeDataView InstanceView = EditorNode.GetInstance();
+		if (InstanceView.IsValid())
 		{
-			// Struct Instance
-			const int32 InstanceIndex = InstancedStructContainer.Add(EditorNode.Instance);
-			InstantiateStructSubobjects(InstancedStructContainer[InstanceIndex]);
+			if (InstanceView.GetStruct() == Node.GetInstanceDataType())
+			{
+				if (const UScriptStruct* InstanceStruct = Cast<const UScriptStruct>(InstanceView.GetStruct()))
+				{
+					// Struct Instance
+					FInstancedStruct InstancedStruct;
+					InstancedStruct.InitializeAs(InstanceStruct, static_cast<const uint8*>(InstanceView.GetMemory()));
+					const int32 InstanceIndex = InstancedStructContainer.Add(MoveTemp(InstancedStruct));
+					InstantiateStructSubobjects(InstancedStructContainer[InstanceIndex]);
 
-			// Create binding source struct descriptor.
-			StructDesc.Struct = EditorNode.Instance.GetScriptStruct();
+					// Create binding source struct FInstancedStruct descriptor.
+					StructDesc.Struct = InstanceStruct;
 
-			if (const auto Validation = UE::StateTree::Compiler::IsValidIndex16(InstanceIndex); Validation.DidFail())
-			{
-				Validation.Log(Log, TEXT("InstanceIndex"), StructDesc);
-				return {};
-			}
-			Node.InstanceTemplateIndex = FStateTreeIndex16(InstanceIndex);
-			Node.InstanceDataHandle = DataHandle;
-			InstanceDataView = FStateTreeDataView(InstancedStructContainer[InstanceIndex]);
-		}
-		else
-		{
-			Log.Reportf(EMessageSeverity::Error, StructDesc, TEXT("The instance data type does not match."));
-			return {};
-		}
-	}
-	else if (EditorNode.InstanceObject != nullptr)
-	{
-		if (ensure(EditorNode.InstanceObject->GetClass() == Node.GetInstanceDataType()))
-		{
-			UObject* Instance = UE::StateTree::Compiler::DuplicateInstanceObject(Log, StructDesc, EditorNode.ID, EditorNode.InstanceObject, StateTree);
+					if (const auto Validation = UE::StateTree::Compiler::IsValidIndex16(InstanceIndex); Validation.DidFail())
+					{
+						Validation.Log(Log, TEXT("InstanceIndex"), StructDesc);
+						return {};
+					}
+					Node.InstanceTemplateIndex = FStateTreeIndex16(InstanceIndex);
+					Node.InstanceDataHandle = DataHandle;
+					InstanceDataView = FStateTreeDataView(InstancedStructContainer[InstanceIndex]);
+				}
+				else if (const UClass* InstanceClass = Cast<const UClass>(InstanceView.GetStruct()))
+				{
+					const UObject* InstanceObject = InstanceView.GetPtr<const UObject>();
+					UObject* Instance = UE::StateTree::Compiler::DuplicateInstanceObject(Log, StructDesc, EditorNode.ID, InstanceObject, StateTree);
 
-			FInstancedStruct Wrapper;
-			Wrapper.InitializeAs<FStateTreeInstanceObjectWrapper>(Instance);
-			const int32 InstanceIndex = InstancedStructContainer.Add(MoveTemp(Wrapper));
+					FInstancedStruct Wrapper;
+					Wrapper.InitializeAs<FStateTreeInstanceObjectWrapper>(Instance);
+					const int32 InstanceIndex = InstancedStructContainer.Add(MoveTemp(Wrapper));
 
-			// Create binding source struct descriptor.
-			StructDesc.Struct = Instance->GetClass();
+					// Create binding source struct descriptor.
+					StructDesc.Struct = InstanceClass;
 
-			if (const auto Validation = UE::StateTree::Compiler::IsValidIndex16(InstanceIndex); Validation.DidFail())
+					if (const auto Validation = UE::StateTree::Compiler::IsValidIndex16(InstanceIndex); Validation.DidFail())
+					{
+						Validation.Log(Log, TEXT("InstanceIndex"), StructDesc);
+						return {};
+					}
+					Node.InstanceTemplateIndex = FStateTreeIndex16(InstanceIndex);
+					Node.InstanceDataHandle = DataHandle.ToObjectSource();
+					InstanceDataView = FStateTreeDataView(Instance);
+				}
+				else
+				{
+					check(false); // Unsupported instance data type. It should be either UScriptStruct or UClass.
+					return {};
+				}
+			}
+			else
 			{
-				Validation.Log(Log, TEXT("InstanceIndex"), StructDesc);
+				Log.Reportf(EMessageSeverity::Error, StructDesc, TEXT("The instance data type does not match."));
 				return {};
 			}
-			Node.InstanceTemplateIndex = FStateTreeIndex16(InstanceIndex);
-			Node.InstanceDataHandle = DataHandle.ToObjectSource();
-			InstanceDataView = FStateTreeDataView(Instance);
 		}
-		else
+		else if (Node.GetInstanceDataType() != nullptr)
 		{
-			Log.Reportf(EMessageSeverity::Error, StructDesc, TEXT("The instance data type does not match."));
+			Log.Reportf(EMessageSeverity::Error, StructDesc, TEXT("The instance data is missing."));
 			return {};
 		}
-	}
-	else if (Node.GetInstanceDataType() != nullptr)
-	{
-		Log.Reportf(EMessageSeverity::Error, StructDesc, TEXT("The instance data is missing."));
-		return {};
+		StructDesc.DataHandle = Node.InstanceDataHandle;
 	}
 
-	StructDesc.DataHandle = Node.InstanceDataHandle;
-
-	if (EditorNode.ExecutionRuntimeData.IsValid() || EditorNode.ExecutionRuntimeDataObject != nullptr)
 	{
-		if (EditorNode.ExecutionRuntimeData.IsValid())
+		const FStateTreeDataView ExecutionRuntimeView = EditorNode.GetExecutionRuntimeData();
+		if (ExecutionRuntimeView.IsValid())
 		{
-			if (ensure(EditorNode.ExecutionRuntimeData.GetScriptStruct() == Node.GetExecutionRuntimeDataType()))
+			if (ExecutionRuntimeView.GetStruct() == Node.GetExecutionRuntimeDataType())
 			{
-				// Struct Instance
-				const int32 InstanceIndex = ExecutionRuntimeStructs.Add(EditorNode.ExecutionRuntimeData);
-				InstantiateStructSubobjects(ExecutionRuntimeStructs[InstanceIndex]);
-
-				if (const auto Validation = UE::StateTree::Compiler::IsValidIndex16(InstanceIndex); Validation.DidFail())
+				if (const UScriptStruct* ExecutionRuntimeStruct = Cast<const UScriptStruct>(ExecutionRuntimeView.GetStruct()))
 				{
-					Validation.Log(Log, TEXT("ExecutionRuntime Index"), StructDesc);
-					return {};
+					// Struct Instance
+					FInstancedStruct InstancedStruct;
+					InstancedStruct.InitializeAs(ExecutionRuntimeStruct, static_cast<const uint8*>(ExecutionRuntimeView.GetMemory()));
+					const int32 InstanceIndex = ExecutionRuntimeStructs.Add(MoveTemp(InstancedStruct));
+					InstantiateStructSubobjects(ExecutionRuntimeStructs[InstanceIndex]);
+
+					if (const auto Validation = UE::StateTree::Compiler::IsValidIndex16(InstanceIndex); Validation.DidFail())
+					{
+						Validation.Log(Log, TEXT("ExecutionRuntime Index"), StructDesc);
+						return {};
+					}
+					Node.ExecutionRuntimeTemplateIndex = FStateTreeIndex16(InstanceIndex);
 				}
-				Node.ExecutionRuntimeTemplateIndex = FStateTreeIndex16(InstanceIndex);
-			}
-			else
-			{
-				Log.Reportf(EMessageSeverity::Error, StructDesc, TEXT("The execution runtime data type does not match."));
-				return {};
-			}
-		}
-		else if (EditorNode.ExecutionRuntimeDataObject != nullptr)
-		{
-			if (ensure(EditorNode.ExecutionRuntimeDataObject->GetClass() == Node.GetExecutionRuntimeDataType()))
-			{
-				// Object Instance
-				UObject* Instance = UE::StateTree::Compiler::DuplicateInstanceObject(Log, StructDesc, EditorNode.ID, EditorNode.ExecutionRuntimeDataObject, StateTree);
+				else if (const UClass* ExecutionRuntimeClass = Cast<const UClass>(ExecutionRuntimeView.GetStruct()))
+				{
+					// Object Instance
+					const UObject* ExecutionRuntimeObject = ExecutionRuntimeView.GetPtr<const UObject>();
+					UObject* Instance = UE::StateTree::Compiler::DuplicateInstanceObject(Log, StructDesc, EditorNode.ID, ExecutionRuntimeObject, StateTree);
 
-				FInstancedStruct Wrapper;
-				Wrapper.InitializeAs<FStateTreeInstanceObjectWrapper>(Instance);
-				const int32 InstanceIndex = ExecutionRuntimeStructs.Add(MoveTemp(Wrapper));
+					FInstancedStruct Wrapper;
+					Wrapper.InitializeAs<FStateTreeInstanceObjectWrapper>(Instance);
+					const int32 InstanceIndex = ExecutionRuntimeStructs.Add(MoveTemp(Wrapper));
 
-				if (const auto Validation = UE::StateTree::Compiler::IsValidIndex16(InstanceIndex); Validation.DidFail())
+					if (const auto Validation = UE::StateTree::Compiler::IsValidIndex16(InstanceIndex); Validation.DidFail())
+					{
+						Validation.Log(Log, TEXT("ExecutionRuntime Index"), StructDesc);
+						return {};
+					}
+					Node.ExecutionRuntimeTemplateIndex = FStateTreeIndex16(InstanceIndex);
+				}
+				else
 				{
-					Validation.Log(Log, TEXT("ExecutionRuntime Index"), StructDesc);
+					check(false); // Unsupported instance data type. It should be either UScriptStruct or UClass.
 					return {};
 				}
-				Node.ExecutionRuntimeTemplateIndex = FStateTreeIndex16(InstanceIndex);
 			}
 			else
 			{
@@ -2253,11 +2262,11 @@ TOptional<FStateTreeDataView> FStateTreeCompiler::CreateNodeInstanceData(const F
 				return {};
 			}
 		}
-	}
-	else if (Node.GetExecutionRuntimeDataType() != nullptr)
-	{
-		Log.Reportf(EMessageSeverity::Error, StructDesc, TEXT("The execution runtime data is missing."));
-		return {};
+		else if (Node.GetExecutionRuntimeDataType() != nullptr)
+		{
+			Log.Reportf(EMessageSeverity::Error, StructDesc, TEXT("The execution runtime data is missing."));
+			return {};
+		}
 	}
 
 	return TOptional<FStateTreeDataView>(InstanceDataView);
