Refactor the FixIntances

#rnx
#Synced-CL 52180191

diff --git a/Engine/Plugins/Runtime/StateTree/Source/StateTreeEditorModule/Private/StateTreeEditorData.cpp b/Engine/Plugins/Runtime/StateTree/Source/StateTreeEditorModule/Private/StateTreeEditorData.cpp
--- a/Engine/Plugins/Runtime/StateTree/Source/StateTreeEditorModule/Private/StateTreeEditorData.cpp
+++ b/Engine/Plugins/Runtime/StateTree/Source/StateTreeEditorModule/Private/StateTreeEditorData.cpp
@@ -794,132 +794,6 @@ void UStateTreeEditorData::CreateRootProperties(TArrayView<UE::PropertyBinding::
 	UStateTreeEditingSubsystem::MarkAsPubliclyModified(GetTypedOuter<UStateTree>());
 }
 
-void UStateTreeEditorData::FixObjectInstances(TSet<UObject*>& SeenObjects, UObject& Outer, FStateTreeEditorNode& Node)
-{
-	auto NewInstances = [this, &Node, &Outer](const UStruct* Struct, const UStruct* ExecutionRuntimeStruct) -> bool
-		{
-			bool bDirty = false;
-			if (Struct && !Node.GetInstance().IsValid())
-			{
-				if (const UScriptStruct* InstanceType = Cast<const UScriptStruct>(Struct))
-				{
-					Node.Instance.InitializeAs(InstanceType);
-					bDirty = true;
-				}
-				else if (const UClass* InstanceClass = Cast<const UClass>(Struct))
-				{
-					Node.InstanceObject = NewObject<UObject>(&Outer, InstanceClass);
-					bDirty = true;
-				}
-			}
-			if (ExecutionRuntimeStruct && !Node.GetExecutionRuntimeData().IsValid())
-			{
-				if (const UScriptStruct* InstanceType = Cast<const UScriptStruct>(ExecutionRuntimeStruct))
-				{
-					Node.ExecutionRuntimeData.InitializeAs(InstanceType);
-					bDirty = true;
-				}
-				else if (const UClass* InstanceClass = Cast<const UClass>(ExecutionRuntimeStruct))
-				{
-					Node.ExecutionRuntimeDataObject = NewObject<UObject>(&Outer, InstanceClass);
-					bDirty = true;
-				}
-			}
-
-			return bDirty;
-		};
-	auto FixInstance = [this, &SeenObjects, &Outer](FInstancedStruct& Instance, const UStruct* Struct) -> bool
-		{
-			bool bDirty = false;
-			if (Instance.IsValid())
-			{
-				if (Instance.GetScriptStruct() != Struct)
-				{
-					Instance.Reset();
-					bDirty = true;
-				}
-			}
-			return bDirty;
-		};
-	auto FixInstanceObject = [this, &SeenObjects, &Outer](TObjectPtr<UObject>& InstanceObject, const UStruct* Struct) -> bool
-		{
-			bool bDirty = false;
-			if (InstanceObject)
-			{
-				if (InstanceObject->GetClass() != Struct)
-				{
-					InstanceObject = nullptr;
-					bDirty = true;
-				}
-				// Found a duplicate reference to an object, make unique copy.
-				else if(SeenObjects.Contains(InstanceObject))				
-				{
-					UE_LOG(LogStateTreeEditor, Log, TEXT("%s: Making duplicate node instance %s unique."), *GetFullName(), *GetNameSafe(InstanceObject));
-					InstanceObject = DuplicateObject(InstanceObject, &Outer);
-					bDirty = true;
-				}
-				else
-				{
-					// Make sure the instance object is property outered.
-					if (InstanceObject->GetOuter() != &Outer)
-					{
-						UE_LOG(LogStateTreeEditor, Log, TEXT("%s: Fixing outer on node instance %s."), *GetFullName(), *GetNameSafe(InstanceObject));
-						InstanceObject->Rename(nullptr, &Outer, REN_DontCreateRedirectors | REN_DoNotDirty | REN_AllowPackageLinkerMismatch);
-						bDirty = true;
-					}
-				}
-				SeenObjects.Add(InstanceObject);
-			}
-			return bDirty;
-		};
-
-	const FStateTreeNodeBase* NodeBase = Node.GetNode().GetPtr<FStateTreeNodeBase>();
-	bool bDirty = false;
-	if (NodeBase)
-	{
-		if (FixInstance(Node.Instance, NodeBase->GetInstanceDataType()))
-		{
-			bDirty = true;
-		}
-		if (FixInstanceObject(Node.InstanceObject, NodeBase->GetInstanceDataType()))
-		{
-			bDirty = true;
-		}
-		if (FixInstance(Node.ExecutionRuntimeData, NodeBase->GetExecutionRuntimeDataType()))
-		{
-			bDirty = true;
-		}
-		if (FixInstanceObject(Node.ExecutionRuntimeDataObject, NodeBase->GetExecutionRuntimeDataType()))
-		{
-			bDirty = true;
-		}
-		if (NewInstances(NodeBase->GetInstanceDataType(), NodeBase->GetExecutionRuntimeDataType()))
-		{
-			bDirty = true;
-		}
-	}
-	else
-	{
-		if (Node.GetInstance().IsValid())
-		{
-			Node.Instance.Reset();
-			Node.InstanceObject = nullptr;
-			bDirty = true;
-		}
-		if (Node.GetExecutionRuntimeData().IsValid())
-		{
-			Node.ExecutionRuntimeData.Reset();
-			Node.ExecutionRuntimeDataObject = nullptr;
-			bDirty = true;
-		}
-	}
-
-	if (bDirty)
-	{
-		UStateTreeEditingSubsystem::MarkAsModified(GetTypedOuter<UStateTree>());
-	}
-};
-
 void UStateTreeEditorData::FixObjectNodes()
 {
 	// Older version of State Trees had all instances outered to the editor data. This causes issues with State copy/paste.
@@ -928,49 +802,76 @@ void UStateTreeEditorData::FixObjectNodes()
 	// Ensure that all node objects are parented to their states, and make duplicated instances unique.
 
 	TSet<UObject*> SeenObjects;
-	
-	VisitHierarchy([&SeenObjects, TreeData = this](UStateTreeState& State, UStateTreeState* ParentState) mutable 
+	bool bDirty = false;
+	VisitHierarchy([&SeenObjects, &bDirty](UStateTreeState& State, UStateTreeState* ParentState) mutable
 	{
 
 		// Enter conditions
 		for (FStateTreeEditorNode& Node : State.EnterConditions)
 		{
-			TreeData->FixObjectInstances(SeenObjects, State, Node);
+			if (Node.FixObjectInstances(SeenObjects, &State))
+			{
+				bDirty = true;
+			}
 		}
 
 		// Consideration Utilities
 		for (FStateTreeEditorNode& Node : State.Considerations)
 		{
-			TreeData->FixObjectInstances(SeenObjects, State, Node);
+			if (Node.FixObjectInstances(SeenObjects, &State))
+			{
+				bDirty = true;
+			}
 		}
-		
+
 		// Tasks
 		for (FStateTreeEditorNode& Node : State.Tasks)
 		{
-			TreeData->FixObjectInstances(SeenObjects, State, Node);
+			if (Node.FixObjectInstances(SeenObjects, &State))
+			{
+				bDirty = true;
+			}
+		}
+
+		if (State.SingleTask.FixObjectInstances(SeenObjects, &State))
+		{
+			bDirty = true;
 		}
-		TreeData->FixObjectInstances(SeenObjects, State, State.SingleTask);
 
 		// Transitions
 		for (FStateTreeTransition& Transition : State.Transitions)
 		{
 			for (FStateTreeEditorNode& Node : Transition.Conditions)
 			{
-				TreeData->FixObjectInstances(SeenObjects, State, Node);
+				if (Node.FixObjectInstances(SeenObjects, &State))
+				{
+					bDirty = true;
+				}
 			}
 		}
-		
+
 		return EStateTreeVisitor::Continue;
 	});
 
 	for (FStateTreeEditorNode& Node : Evaluators)
 	{
-		FixObjectInstances(SeenObjects, *this, Node);
+		if (Node.FixObjectInstances(SeenObjects, this))
+		{
+			bDirty = true;
+		}
 	}
 
 	for (FStateTreeEditorNode& Node : GlobalTasks)
 	{
-		FixObjectInstances(SeenObjects, *this, Node);
+		if (Node.FixObjectInstances(SeenObjects, this))
+		{
+			bDirty = true;
+		}
+	}
+
+	if (bDirty)
+	{
+		UStateTreeEditingSubsystem::MarkAsModified(GetTypedOuter<UStateTree>());
 	}
 }
 
diff --git a/Engine/Plugins/Runtime/StateTree/Source/StateTreeEditorModule/Private/StateTreeEditorNode.cpp b/Engine/Plugins/Runtime/StateTree/Source/StateTreeEditorModule/Private/StateTreeEditorNode.cpp
--- a/Engine/Plugins/Runtime/StateTree/Source/StateTreeEditorModule/Private/StateTreeEditorNode.cpp
+++ b/Engine/Plugins/Runtime/StateTree/Source/StateTreeEditorModule/Private/StateTreeEditorNode.cpp
@@ -39,6 +39,128 @@ void FStateTreeEditorNode::InitializeInstanceData(TNotNull<UObject*> InOuter)
 	}
 }
 
+bool FStateTreeEditorNode::FixObjectInstances(TSet<UObject*>& SeenObjects, TNotNull<UObject*> Outer)
+{
+	auto NewInstances = [this, Outer](const UStruct* Struct, const UStruct* ExecutionRuntimeStruct) -> bool
+		{
+			bool bDirty = false;
+			if (Struct && !GetInstance().IsValid())
+			{
+				if (const UScriptStruct* InstanceType = Cast<const UScriptStruct>(Struct))
+				{
+					Instance.InitializeAs(InstanceType);
+					bDirty = true;
+				}
+				else if (const UClass* InstanceClass = Cast<const UClass>(Struct))
+				{
+					InstanceObject = NewObject<UObject>(Outer, InstanceClass);
+					bDirty = true;
+				}
+			}
+			if (ExecutionRuntimeStruct && !GetExecutionRuntimeData().IsValid())
+			{
+				if (const UScriptStruct* InstanceType = Cast<const UScriptStruct>(ExecutionRuntimeStruct))
+				{
+					ExecutionRuntimeData.InitializeAs(InstanceType);
+					bDirty = true;
+				}
+				else if (const UClass* InstanceClass = Cast<const UClass>(ExecutionRuntimeStruct))
+				{
+					ExecutionRuntimeDataObject = NewObject<UObject>(Outer, InstanceClass);
+					bDirty = true;
+				}
+			}
+			return bDirty;
+		};
+	auto FixInstance = [](FInstancedStruct& InInstance, const UStruct* Struct) -> bool
+		{
+			bool bDirty = false;
+			if (InInstance.IsValid())
+			{
+				if (InInstance.GetScriptStruct() != Struct)
+				{
+					InInstance.Reset();
+					bDirty = true;
+				}
+			}
+			return bDirty;
+		};
+	auto FixInstanceObject = [&SeenObjects, Outer](TObjectPtr<UObject>& InInstanceObject, const UStruct* Struct) -> bool
+		{
+			bool bDirty = false;
+			if (InInstanceObject)
+			{
+				if (InInstanceObject->GetClass() != Struct)
+				{
+					InInstanceObject = nullptr;
+					bDirty = true;
+				}
+				// Found a duplicate reference to an object, make unique copy.
+				else if (SeenObjects.Contains(InInstanceObject))
+				{
+					UE_LOG(LogStateTreeEditor, Log, TEXT("%s: Making duplicate node instance %s unique."), *Outer->GetFullName(), *GetNameSafe(InInstanceObject));
+					InInstanceObject = DuplicateObject(InInstanceObject, Outer);
+					bDirty = true;
+				}
+				else
+				{
+					// Make sure the instance object is properly outered.
+					if (InInstanceObject->GetOuter() != Outer)
+					{
+						UE_LOG(LogStateTreeEditor, Log, TEXT("%s: Fixing outer on node instance %s."), *Outer->GetFullName(), *GetNameSafe(InInstanceObject));
+						InInstanceObject->Rename(nullptr, Outer, REN_DontCreateRedirectors | REN_DoNotDirty | REN_AllowPackageLinkerMismatch);
+						bDirty = true;
+					}
+				}
+				SeenObjects.Add(InInstanceObject);
+			}
+			return bDirty;
+		};
+
+	const FStateTreeNodeBase* NodeBase = GetNode().GetPtr<FStateTreeNodeBase>();
+	bool bDirty = false;
+	if (NodeBase)
+	{
+		if (FixInstance(Instance, NodeBase->GetInstanceDataType()))
+		{
+			bDirty = true;
+		}
+		if (FixInstanceObject(InstanceObject, NodeBase->GetInstanceDataType()))
+		{
+			bDirty = true;
+		}
+		if (FixInstance(ExecutionRuntimeData, NodeBase->GetExecutionRuntimeDataType()))
+		{
+			bDirty = true;
+		}
+		if (FixInstanceObject(ExecutionRuntimeDataObject, NodeBase->GetExecutionRuntimeDataType()))
+		{
+			bDirty = true;
+		}
+		if (NewInstances(NodeBase->GetInstanceDataType(), NodeBase->GetExecutionRuntimeDataType()))
+		{
+			bDirty = true;
+		}
+	}
+	else
+	{
+		if (GetInstance().IsValid())
+		{
+			Instance.Reset();
+			InstanceObject = nullptr;
+			bDirty = true;
+		}
+		if (GetExecutionRuntimeData().IsValid())
+		{
+			ExecutionRuntimeData.Reset();
+			ExecutionRuntimeDataObject = nullptr;
+			bDirty = true;
+		}
+	}
+
+	return bDirty;
+}
+
 FName FStateTreeEditorNode::GetName() const
 {
 	const UScriptStruct* NodeType = Node.GetScriptStruct();
diff --git a/Engine/Plugins/Runtime/StateTree/Source/StateTreeEditorModule/Public/StateTreeEditorData.h b/Engine/Plugins/Runtime/StateTree/Source/StateTreeEditorModule/Public/StateTreeEditorData.h
--- a/Engine/Plugins/Runtime/StateTree/Source/StateTreeEditorModule/Public/StateTreeEditorData.h
+++ b/Engine/Plugins/Runtime/StateTree/Source/StateTreeEditorModule/Public/StateTreeEditorData.h
@@ -378,7 +378,6 @@ public:
 	UE_API virtual void CreateRootProperties(TArrayView<UE::PropertyBinding::FPropertyCreationDescriptor> InOutCreationDescs);
 
 private:
-	UE_API void FixObjectInstances(TSet<UObject*>& SeenObjects, UObject& Outer, FStateTreeEditorNode& Node);
 	UE_API void FixObjectNodes();
 	UE_API void DuplicateIDs();
 	UE_API void CallPostLoadOnNodes();
diff --git a/Engine/Plugins/Runtime/StateTree/Source/StateTreeEditorModule/Public/StateTreeEditorNode.h b/Engine/Plugins/Runtime/StateTree/Source/StateTreeEditorModule/Public/StateTreeEditorNode.h
--- a/Engine/Plugins/Runtime/StateTree/Source/StateTreeEditorModule/Public/StateTreeEditorNode.h
+++ b/Engine/Plugins/Runtime/StateTree/Source/StateTreeEditorModule/Public/StateTreeEditorNode.h
@@ -59,6 +59,21 @@ struct FStateTreeEditorNode
 	 */
 	UE_API void InitializeAs(TNotNull<UObject*> InOuter, TNotNull<const UScriptStruct*> InNodeStruct);
 
+	/**
+	 * Fixes UObject instance data of the node, ensuring proper outers and making duplicates unique.
+	 * Creates any missing instance or execution runtime data.
+	 *
+	 * SeenObjects is threaded through all nodes in a single traversal so that shared UObject references
+	 * (e.g. from copy-paste) are detected and made unique on the spot, without a separate pre-pass to
+	 * collect reference counts. Each node adds its live instance objects to the set after processing,
+	 * so any later node holding the same pointer can immediately DuplicateObject it.
+	 *
+	 * @param SeenObjects Set of UObject instances already assigned to earlier nodes. Updated in place.
+	 * @param Outer Outer used for re-parenting and creating new instances.
+	 * @return true if any changes were made.
+	 */
+	UE_API bool FixObjectInstances(TSet<UObject*>& SeenObjects, TNotNull<UObject*> Outer);
+
 private:
 	UE_API void InitializeInstanceData(TNotNull<UObject*> InOuter);
 
