[StateTree] Fix use-after-free crash in FStateTreeReferenceDetails::SyncParameters

* FStateTreeReferenceDetails::CustomizeHeader captures this by raw pointer in a lambda passed to SetOnPropertyValueChanged. When the details panel rebuilds (e.g. due to a FPropertyBindingExtension on the same view), the customization instance can be destroyed while the lambda remains registered on the property node delegate. The next property change dereferences the freed pointer.
* Fix: use CreateSPLambda(AsShared(), ...) so the delegate prevents the customization from being destroyed while it is still bound, matching the pattern already used for OnPostCompile on the next line.

#Synced-CL 52196679

diff --git a/Engine/Plugins/Runtime/StateTree/Source/StateTreeEditorModule/Private/Customizations/StateTreeReferenceDetails.cpp b/Engine/Plugins/Runtime/StateTree/Source/StateTreeEditorModule/Private/Customizations/StateTreeReferenceDetails.cpp
--- a/Engine/Plugins/Runtime/StateTree/Source/StateTreeEditorModule/Private/Customizations/StateTreeReferenceDetails.cpp
+++ b/Engine/Plugins/Runtime/StateTree/Source/StateTreeEditorModule/Private/Customizations/StateTreeReferenceDetails.cpp
@@ -137,7 +137,7 @@ void FStateTreeReferenceDetails::CustomizeHeader(const TSharedRef<IPropertyHandl
 	]
 	.ShouldAutoExpand(true);
 	
-	StateTreeProperty->SetOnPropertyValueChanged(FSimpleDelegate::CreateLambda([this]()
+	StateTreeProperty->SetOnPropertyValueChanged(FSimpleDelegate::CreateSPLambda(AsShared(), [this]()
 		{
 			SyncParameters();
 		}));
