[StateTree] Cache the node name for better UI and compiler performance. The best solution would be to direclty use the FText system but that involve a huge refactor.

#Synced-CL 52180191

diff --git a/Engine/Plugins/Runtime/StateTree/Source/StateTreeEditorModule/Private/Customizations/StateTreeEditorNodeUtils.cpp b/Engine/Plugins/Runtime/StateTree/Source/StateTreeEditorModule/Private/Customizations/StateTreeEditorNodeUtils.cpp
--- a/Engine/Plugins/Runtime/StateTree/Source/StateTreeEditorModule/Private/Customizations/StateTreeEditorNodeUtils.cpp
+++ b/Engine/Plugins/Runtime/StateTree/Source/StateTreeEditorModule/Private/Customizations/StateTreeEditorNodeUtils.cpp
@@ -630,27 +630,7 @@ void ConditionalUpdateNodeInstanceData(FStateTreeEditorNode& EditorNode, UObject
 
 	FStateTreeEditorNode OldEditorNode = EditorNode;
 
-	EditorNode.Instance.Reset();
-	EditorNode.InstanceObject = nullptr;
-	EditorNode.ExecutionRuntimeData.Reset();
-	EditorNode.ExecutionRuntimeDataObject = nullptr;
-
-	if (const UScriptStruct* InstanceType = Cast<UScriptStruct>(DesiredType))
-	{
-		EditorNode.Instance.InitializeAs(InstanceType);
-	}
-	else if (const UClass* InstanceClass = Cast<UClass>(DesiredType))
-	{
-		EditorNode.InstanceObject = NewObject<UObject>(&InstanceOuter, InstanceClass);
-	}
-	if (const UScriptStruct* InstanceType = Cast<UScriptStruct>(ExecutionRuntimeDesiredType))
-	{
-		EditorNode.ExecutionRuntimeData.InitializeAs(InstanceType);
-	}
-	else if (const UClass* InstanceClass = Cast<UClass>(ExecutionRuntimeDesiredType))
-	{
-		EditorNode.ExecutionRuntimeDataObject = NewObject<UObject>(&InstanceOuter, InstanceClass);
-	}
+	EditorNode.ReallocInstanceData(&InstanceOuter);
 
 	RetainProperties(OldEditorNode, EditorNode);
 
diff --git a/Engine/Plugins/Runtime/StateTree/Source/StateTreeEditorModule/Private/StateTreeCompiler.cpp b/Engine/Plugins/Runtime/StateTree/Source/StateTreeEditorModule/Private/StateTreeCompiler.cpp
--- a/Engine/Plugins/Runtime/StateTree/Source/StateTreeEditorModule/Private/StateTreeCompiler.cpp
+++ b/Engine/Plugins/Runtime/StateTree/Source/StateTreeEditorModule/Private/StateTreeCompiler.cpp
@@ -2592,7 +2592,7 @@ bool FStateTreeCompiler::CreateTask(UStateTreeState* State, const FStateTreeEdit
 
 	// Copy the task
 	IDToNode.Add(TaskNode.ID, Nodes.Num());
-	FInstancedStruct& Node = Nodes.Add_GetRef(TaskNode.Node);
+	FInstancedStruct& Node = Nodes.Add_GetRef(FInstancedStruct(TaskNode.GetNode()));
 	InstantiateStructSubobjects(Node);
 	
 	FStateTreeTaskBase& Task = Node.GetMutable<FStateTreeTaskBase>();
@@ -2642,7 +2642,7 @@ bool FStateTreeCompiler::CreateEvaluator(const FStateTreeEditorNode& EvalNode, c
 
 	// Copy the evaluator
 	IDToNode.Add(EvalNode.ID, Nodes.Num());
-	FInstancedStruct& Node = Nodes.Add_GetRef(EvalNode.Node);
+	FInstancedStruct& Node = Nodes.Add_GetRef(FInstancedStruct(EvalNode.GetNode()));
 	InstantiateStructSubobjects(Node);
 	
 	FStateTreeEvaluatorBase& Eval = Node.GetMutable<FStateTreeEvaluatorBase>();
diff --git a/Engine/Plugins/Runtime/StateTree/Source/StateTreeEditorModule/Private/StateTreeEditorData.cpp b/Engine/Plugins/Runtime/StateTree/Source/StateTreeEditorModule/Private/StateTreeEditorData.cpp
--- a/Engine/Plugins/Runtime/StateTree/Source/StateTreeEditorModule/Private/StateTreeEditorData.cpp
+++ b/Engine/Plugins/Runtime/StateTree/Source/StateTreeEditorModule/Private/StateTreeEditorData.cpp
@@ -1578,7 +1578,7 @@ EStateTreeVisitor UStateTreeEditorData::VisitStructBoundPropertyFunctions(FGuid
 				return EStateTreeVisitor::Break;
 			}
 
-			if (const UStruct* NodeStruct = FunctionNode->Node.GetScriptStruct())
+			if (const UStruct* NodeStruct = FunctionNode->GetNode().GetScriptStruct())
 			{
 				FStateTreeBindableStructDesc Desc;
 				Desc.Struct = FunctionNode->GetInstance().GetStruct();
diff --git a/Engine/Plugins/Runtime/StateTree/Source/StateTreeEditorModule/Private/StateTreeEditorNode.cpp b/Engine/Plugins/Runtime/StateTree/Source/StateTreeEditorModule/Private/StateTreeEditorNode.cpp
--- a/Engine/Plugins/Runtime/StateTree/Source/StateTreeEditorModule/Private/StateTreeEditorNode.cpp
+++ b/Engine/Plugins/Runtime/StateTree/Source/StateTreeEditorModule/Private/StateTreeEditorNode.cpp
@@ -10,16 +10,38 @@
 
 #include UE_INLINE_GENERATED_CPP_BY_NAME(StateTreeEditorNode)
 
+void FStateTreeEditorNode::Reset()
+{
+	Node.Reset();
+	Instance.Reset();
+	InstanceObject = nullptr;
+	ExecutionRuntimeData.Reset();
+	ExecutionRuntimeDataObject = nullptr;
+	ID = FGuid();
+	CachedNodeName.Reset();
+	// don't init the ExpressionIndent and ExpressionOperand
+}
+
 void FStateTreeEditorNode::InitializeAs(TNotNull<UObject*> InOuter, TNotNull<const UScriptStruct*> InNodeStruct)
 {
 	Reset();
 	Node.InitializeAs(InNodeStruct);
+	ID = FGuid::NewGuid();
+	InitializeInstanceData(InOuter);
+}
+
+void FStateTreeEditorNode::ReallocInstanceData(TNotNull<UObject*> InOuter)
+{
+	Instance.Reset();
+	InstanceObject = nullptr;
+	ExecutionRuntimeData.Reset();
+	ExecutionRuntimeDataObject = nullptr;
+	CachedNodeName.Reset();
 	InitializeInstanceData(InOuter);
 }
 
 void FStateTreeEditorNode::InitializeInstanceData(TNotNull<UObject*> InOuter)
 {
-	ID = FGuid::NewGuid();
 	const FStateTreeNodeBase& NodeBase = GetNode().Get<FStateTreeNodeBase>();
 	if (const UScriptStruct* InstanceType = Cast<const UScriptStruct>(NodeBase.GetInstanceDataType()))
 	{
@@ -29,6 +51,7 @@ void FStateTreeEditorNode::InitializeInstanceData(TNotNull<UObject*> InOuter)
 	{
 		InstanceObject = NewObject<UObject>(InOuter, InstanceClass);
 	}
+
 	if (const UScriptStruct* RuntimeType = Cast<const UScriptStruct>(NodeBase.GetExecutionRuntimeDataType()))
 	{
 		ExecutionRuntimeData.InitializeAs(RuntimeType);
@@ -162,29 +185,37 @@ bool FStateTreeEditorNode::FixObjectInstances(TSet<UObject*>& SeenObjects, TNotN
 }
 
 FName FStateTreeEditorNode::GetName() const
-{
-	const UScriptStruct* NodeType = Node.GetScriptStruct();
-	if (!NodeType)
+{	
+	if (!CachedNodeName.IsSet())
 	{
-		return FName();
-	}
-
-	if (const FStateTreeNodeBase* NodePtr = Node.GetPtr<FStateTreeNodeBase>())
-	{
-		if (NodePtr->Name.IsNone())
-		{
-			if (InstanceObject &&
-					(NodeType->IsChildOf(TBaseStructure<FStateTreeBlueprintTaskWrapper>::Get())
-					|| NodeType->IsChildOf(TBaseStructure<FStateTreeBlueprintEvaluatorWrapper>::Get())
-					|| NodeType->IsChildOf(TBaseStructure<FStateTreeBlueprintConditionWrapper>::Get())
-					|| NodeType->IsChildOf(TBaseStructure<FStateTreeBlueprintConsiderationWrapper>::Get())))
+		auto MakeName = [this]() -> FName
 			{
-				return FName(InstanceObject->GetClass()->GetDisplayNameText().ToString());
-			}
-			return FName(NodeType->GetDisplayNameText().ToString());
-		}
-		return NodePtr->Name;
+				const UScriptStruct* NodeType = Node.GetScriptStruct();
+				if (!NodeType)
+				{
+					return FName();
+				}
+
+				if (const FStateTreeNodeBase* NodePtr = Node.GetPtr<FStateTreeNodeBase>())
+				{
+					if (NodePtr->Name.IsNone())
+					{
+						if (InstanceObject &&
+							(NodeType->IsChildOf(TBaseStructure<FStateTreeBlueprintTaskWrapper>::Get())
+								|| NodeType->IsChildOf(TBaseStructure<FStateTreeBlueprintEvaluatorWrapper>::Get())
+								|| NodeType->IsChildOf(TBaseStructure<FStateTreeBlueprintConditionWrapper>::Get())
+								|| NodeType->IsChildOf(TBaseStructure<FStateTreeBlueprintConsiderationWrapper>::Get())))
+						{
+							return FName(InstanceObject->GetClass()->GetDisplayNameText().ToString());
+						}
+						return FName(NodeType->GetDisplayNameText().ToString());
+					}
+					return NodePtr->Name;
+				}
+				return FName();
+			};
+
+		CachedNodeName = MakeName();
 	}
-	
-	return FName();
+	return CachedNodeName.Get(FName());
 }
diff --git a/Engine/Plugins/Runtime/StateTree/Source/StateTreeEditorModule/Public/StateTreeEditorNode.h b/Engine/Plugins/Runtime/StateTree/Source/StateTreeEditorModule/Public/StateTreeEditorNode.h
--- a/Engine/Plugins/Runtime/StateTree/Source/StateTreeEditorModule/Public/StateTreeEditorNode.h
+++ b/Engine/Plugins/Runtime/StateTree/Source/StateTreeEditorModule/Public/StateTreeEditorNode.h
@@ -27,15 +27,8 @@ struct FStateTreeEditorNode
 {
 	GENERATED_BODY()
 
-	void Reset()
-	{
-		Node.Reset();
-		Instance.Reset();
-		InstanceObject = nullptr;
-		ExecutionRuntimeData.Reset();
-		ExecutionRuntimeDataObject = nullptr;
-		ID = FGuid();
-	}
+	/** Resets node and instance fields to their default values. */
+	UE_API void Reset();
 
 	/**
 	 * Initializes the node as the given type, resetting any previous state.
@@ -47,6 +40,7 @@ struct FStateTreeEditorNode
 	{
 		Reset();
 		Node.InitializeAs<T>(Forward<TArgs>(InArgs)...);
+		ID = FGuid::NewGuid();
 		InitializeInstanceData(InOuter);
 	}
 
@@ -59,6 +53,9 @@ struct FStateTreeEditorNode
 	 */
 	UE_API void InitializeAs(TNotNull<UObject*> InOuter, TNotNull<const UScriptStruct*> InNodeStruct);
 
+	/** Resets and reallocates instance and execution runtime data to match the current node type, preserving the node ID. */
+	UE_API void ReallocInstanceData(TNotNull<UObject*> InOuter);
+
 	/**
 	 * Fixes UObject instance data of the node, ensuring proper outers and making duplicates unique.
 	 * Creates any missing instance or execution runtime data.
@@ -156,6 +153,9 @@ public:
 
 	UPROPERTY(EditDefaultsOnly, Category = Node)
 	EStateTreeExpressionOperand ExpressionOperand = EStateTreeExpressionOperand::And;
+
+private:
+	mutable TOptional<FName> CachedNodeName;
 };
 
 template <typename T>
