[StateTree] Backout CL 52256259. Renamed node are not properly cached.

#Synced-CL 52180191

diff --git a/Engine/Plugins/Runtime/StateTree/Source/StateTreeEditorModule/Private/StateTreeEditorNode.cpp b/Engine/Plugins/Runtime/StateTree/Source/StateTreeEditorModule/Private/StateTreeEditorNode.cpp
--- a/Engine/Plugins/Runtime/StateTree/Source/StateTreeEditorModule/Private/StateTreeEditorNode.cpp
+++ b/Engine/Plugins/Runtime/StateTree/Source/StateTreeEditorModule/Private/StateTreeEditorNode.cpp
@@ -185,9 +185,7 @@ bool FStateTreeEditorNode::FixObjectInstances(TSet<UObject*>& SeenObjects, TNotN
 }
 
 FName FStateTreeEditorNode::GetName() const
-{	
-	if (!CachedNodeName.IsSet())
-	{
+{
 		auto MakeName = [this]() -> FName
 			{
 				const UScriptStruct* NodeType = Node.GetScriptStruct();
@@ -214,8 +212,5 @@ FName FStateTreeEditorNode::GetName() const
 				}
 				return FName();
 			};
-
-		CachedNodeName = MakeName();
-	}
-	return CachedNodeName.Get(FName());
+	return MakeName();
 }
