TLazyObjectPtr deprecation preparation:
Wrap existing deprecated usage with the deprecation warning macro

#Synced-CL 52252558

diff --git a/Engine/Plugins/Compositing/Composure/Source/Composure/Classes/CompositingElement.h b/Engine/Plugins/Compositing/Composure/Source/Composure/Classes/CompositingElement.h
--- a/Engine/Plugins/Compositing/Composure/Source/Composure/Classes/CompositingElement.h
+++ b/Engine/Plugins/Compositing/Composure/Source/Composure/Classes/CompositingElement.h
@@ -106,8 +106,10 @@ public:
 	UPROPERTY(EditAnywhere, Category = "Composure|Input")
 	ESceneCameraLinkType CameraSource = ESceneCameraLinkType::Unused;
 
-	UPROPERTY()	
+	PRAGMA_DISABLE_DEPRECATION_WARNINGS
+	UPROPERTY()
 	TLazyObjectPtr<ACameraActor> TargetCameraActor_DEPRECATED;
+	PRAGMA_ENABLE_DEPRECATION_WARNINGS
 
 	UPROPERTY(EditInstanceOnly, Category = "Composure|Input", DisplayName = "Target Camera Actor")
 	TWeakObjectPtr<ACameraActor> TargetCameraActorPtr;
diff --git a/Engine/Plugins/Compositing/Composure/Source/Composure/Private/CompositingElement.cpp b/Engine/Plugins/Compositing/Composure/Source/Composure/Private/CompositingElement.cpp
--- a/Engine/Plugins/Compositing/Composure/Source/Composure/Private/CompositingElement.cpp
+++ b/Engine/Plugins/Compositing/Composure/Source/Composure/Private/CompositingElement.cpp
@@ -902,11 +902,13 @@ void ACompositingElement::PostLoad()
 #if WITH_EDITORONLY_DATA
 	// Re-assign the value of deprecated pointer to the new one.
 	// This will continue until user re-saves the level.
+	PRAGMA_DISABLE_DEPRECATION_WARNINGS
 	if (TargetCameraActor_DEPRECATED.IsValid())
 	{
 		TargetCameraActorPtr = TargetCameraActor_DEPRECATED.Get();
 		TargetCameraActor_DEPRECATED = nullptr;
 	}
+	PRAGMA_ENABLE_DEPRECATION_WARNINGS
 #endif
 
 	// Note that RefreshAllInternalPassLists is not called here because the refresh can call Reset, which is not allowed in PostLoad.
diff --git a/Engine/Plugins/Experimental/Water/Source/Runtime/Public/WaterBodyActor.h b/Engine/Plugins/Experimental/Water/Source/Runtime/Public/WaterBodyActor.h
--- a/Engine/Plugins/Experimental/Water/Source/Runtime/Public/WaterBodyActor.h
+++ b/Engine/Plugins/Experimental/Water/Source/Runtime/Public/WaterBodyActor.h
@@ -277,11 +277,13 @@ protected:
 	UPROPERTY()
 	TObjectPtr<UMaterialInstanceDynamic> UnderwaterPostProcessMID_DEPRECATED;
 
+	PRAGMA_DISABLE_DEPRECATION_WARNINGS
 	UPROPERTY()
 	TArray<TLazyObjectPtr<AWaterBodyIsland>> Islands_DEPRECATED;
 
 	UPROPERTY()
 	TArray<TLazyObjectPtr<AWaterBodyExclusionVolume>> ExclusionVolumes_DEPRECATED;
+	PRAGMA_ENABLE_DEPRECATION_WARNINGS
 
 	UPROPERTY()
 	bool bCanAffectNavigation_DEPRECATED;
diff --git a/Engine/Plugins/Experimental/Water/Source/Runtime/Public/WaterBodyComponent.h b/Engine/Plugins/Experimental/Water/Source/Runtime/Public/WaterBodyComponent.h
--- a/Engine/Plugins/Experimental/Water/Source/Runtime/Public/WaterBodyComponent.h
+++ b/Engine/Plugins/Experimental/Water/Source/Runtime/Public/WaterBodyComponent.h
@@ -703,10 +703,12 @@ protected:
 	bool bUseBakedSimForQueriesAndPhysics = true;
 
 #if WITH_EDITORONLY_DATA
+	PRAGMA_DISABLE_DEPRECATION_WARNINGS
 	UPROPERTY()
 	TArray<TLazyObjectPtr<AWaterBodyIsland>> Islands_DEPRECATED;
 	UPROPERTY()
 	TArray<TLazyObjectPtr<AWaterBodyExclusionVolume>> ExclusionVolumes_DEPRECATED;
+	PRAGMA_ENABLE_DEPRECATION_WARNINGS
 
 	UPROPERTY()
 	TObjectPtr<UMaterialInterface> WaterLODMaterial_DEPRECATED;
diff --git a/Engine/Plugins/Runtime/StateTree/Source/StateTreeTestSuite/Private/StateTreeBindingTest.cpp b/Engine/Plugins/Runtime/StateTree/Source/StateTreeTestSuite/Private/StateTreeBindingTest.cpp
--- a/Engine/Plugins/Runtime/StateTree/Source/StateTreeTestSuite/Private/StateTreeBindingTest.cpp
+++ b/Engine/Plugins/Runtime/StateTree/Source/StateTreeTestSuite/Private/StateTreeBindingTest.cpp
@@ -445,38 +445,48 @@ struct FStateTreeTest_BindingsCompiler_WithObjectPtr : FStateTreeTestBase
 		SourceA.A = TBaseStructure<FVector>::Get();
 		SourceA.B = TBaseStructure<FRotator>::Get();
 		SourceA.C = TBaseStructure<FMatrix>::Get();
+PRAGMA_DISABLE_DEPRECATION_WARNINGS
 		SourceA.D = TBaseStructure<FTransform>::Get();
+PRAGMA_ENABLE_DEPRECATION_WARNINGS
 		constexpr int32 NumberOfElements = 4;
 		SourceA.ArrayA.SetNum(NumberOfElements, EAllowShrinking::No);
 		SourceA.ArrayB.SetNum(NumberOfElements, EAllowShrinking::No);
 		SourceA.ArrayC.SetNum(NumberOfElements, EAllowShrinking::No);
+PRAGMA_DISABLE_DEPRECATION_WARNINGS
 		SourceA.ArrayD.SetNum(NumberOfElements, EAllowShrinking::No);
-		
+PRAGMA_ENABLE_DEPRECATION_WARNINGS
 		SourceA.ArrayA[2] = TBaseStructure<FVector>::Get();
 		SourceA.ArrayB[2] = TBaseStructure<FRotator>::Get();
 		SourceA.ArrayC[2] = TBaseStructure<FMatrix>::Get();
+PRAGMA_DISABLE_DEPRECATION_WARNINGS
 		SourceA.ArrayD[2] = TBaseStructure<FTransform>::Get();
-
+PRAGMA_ENABLE_DEPRECATION_WARNINGS
 		FStateTreeTest_PropertyStruct SourceB;
 
 		FStateTreeTest_PropertyStructObjectPtr Target;
 		Target.ArrayA.SetNum(2, EAllowShrinking::No);
 		Target.ArrayB.SetNum(2, EAllowShrinking::No);
 		Target.ArrayC.SetNum(2, EAllowShrinking::No);
+PRAGMA_DISABLE_DEPRECATION_WARNINGS
 		Target.ArrayD.SetNum(2, EAllowShrinking::No);
+PRAGMA_ENABLE_DEPRECATION_WARNINGS
 
 		auto ResetTargetData = [&Target]()
 			{
 				Target.A = nullptr;
 				Target.B = nullptr;
 				Target.C = nullptr;
+PRAGMA_DISABLE_DEPRECATION_WARNINGS
 				Target.D = nullptr;
+PRAGMA_ENABLE_DEPRECATION_WARNINGS
 				for (int32 Index = 0; Index < Target.ArrayA.Num(); ++Index)
 				{
 					Target.ArrayA[Index] = nullptr;
 					Target.ArrayB[Index] = nullptr;
 					Target.ArrayC[Index] = nullptr;
+PRAGMA_DISABLE_DEPRECATION_WARNINGS
 					Target.ArrayD[Index] = nullptr;
+PRAGMA_ENABLE_DEPRECATION_WARNINGS
 				}
 				return true;
 			};
@@ -484,14 +494,19 @@ struct FStateTreeTest_BindingsCompiler_WithObjectPtr : FStateTreeTestBase
 
 		auto TestCopy = [this, &Target](UObject* Object, const FStringView SourceStr)
 			{
+
 				AITEST_EQUAL(FString::Printf(TEXT("Expect A to be copied from %s"), SourceStr.GetData()), Target.A.Get(), Object);
 				AITEST_EQUAL(FString::Printf(TEXT("Expect B to be copied from %s"), SourceStr.GetData()), Target.B.Get(), Object);
 				AITEST_EQUAL(FString::Printf(TEXT("Expect C to be copied from %s"), SourceStr.GetData()), Target.C.Get(), Object);
+PRAGMA_DISABLE_DEPRECATION_WARNINGS
 				AITEST_EQUAL(FString::Printf(TEXT("Expect D to be copied from %s"), SourceStr.GetData()), Target.D.Get(), Object);
+PRAGMA_ENABLE_DEPRECATION_WARNINGS
 				AITEST_EQUAL(FString::Printf(TEXT("Expect A[] to be copied from %s"), SourceStr.GetData()), Target.ArrayA[1].Get(), Object);
 				AITEST_EQUAL(FString::Printf(TEXT("Expect B[] to be copied from %s"), SourceStr.GetData()), Target.ArrayB[1].Get(), Object);
 				AITEST_EQUAL(FString::Printf(TEXT("Expect C[] to be copied from %s"), SourceStr.GetData()), Target.ArrayC[1].Get(), Object);
+PRAGMA_DISABLE_DEPRECATION_WARNINGS
 				AITEST_EQUAL(FString::Printf(TEXT("Expect D[] to be copied from %s"), SourceStr.GetData()), Target.ArrayD[1].Get(), Object);
+PRAGMA_ENABLE_DEPRECATION_WARNINGS
 				return true;
 			};
 
@@ -538,11 +553,13 @@ struct FStateTreeTest_BindingsCompiler_WithObjectPtr : FStateTreeTestBase
 		{
 			return false;
 		}
+PRAGMA_DISABLE_DEPRECATION_WARNINGS
 		if (!CopyAndTestBatch(CopyBatchIndex7, SourceA.D.Get(), TEXT("SourceA.D")))
 		{
 			return false;
 		}
 		if (!CopyAndTestBatch(CopyBatchIndex8, SourceA.D.Get(), TEXT("SourceA.ArrayD")))
+PRAGMA_ENABLE_DEPRECATION_WARNINGS
 		{
 			return false;
 		}
@@ -553,7 +570,9 @@ struct FStateTreeTest_BindingsCompiler_WithObjectPtr : FStateTreeTestBase
 			SourceA.ArrayA[Index] = TBaseStructure<FVector>::Get();
 			SourceA.ArrayB[Index] = TBaseStructure<FRotator>::Get();
 			SourceA.ArrayC[Index] = TBaseStructure<FMatrix>::Get();
+PRAGMA_DISABLE_DEPRECATION_WARNINGS
 			SourceA.ArrayD[Index] = TBaseStructure<FTransform>::Get();
+PRAGMA_ENABLE_DEPRECATION_WARNINGS
 		}
 
 		auto TestCopyArray = [this, &Target, NumberOfElements](UObject* Object, const FStringView SourceStr)
@@ -563,7 +582,9 @@ struct FStateTreeTest_BindingsCompiler_WithObjectPtr : FStateTreeTestBase
 					AITEST_EQUAL(FString::Printf(TEXT("Expect A[] to be copied from %s"), SourceStr.GetData()), Target.ArrayA[Index].Get(), Object);
 					AITEST_EQUAL(FString::Printf(TEXT("Expect B[] to be copied from %s"), SourceStr.GetData()), Target.ArrayB[Index].Get(), Object);
 					AITEST_EQUAL(FString::Printf(TEXT("Expect C[] to be copied from %s"), SourceStr.GetData()), Target.ArrayC[Index].Get(), Object);
+PRAGMA_DISABLE_DEPRECATION_WARNINGS
 					AITEST_EQUAL(FString::Printf(TEXT("Expect D[] to be copied from %s"), SourceStr.GetData()), Target.ArrayD[Index].Get(), Object);
+PRAGMA_ENABLE_DEPRECATION_WARNINGS
 				}
 				return true;
 			};
@@ -591,7 +612,9 @@ struct FStateTreeTest_BindingsCompiler_WithObjectPtr : FStateTreeTestBase
 		{
 			return false;
 		}
+PRAGMA_DISABLE_DEPRECATION_WARNINGS
 		if (!CopyAndTestBatchArray(CopyBatchIndex12, SourceA.D.Get(), TEXT("SourceA.ArrayD")))
+PRAGMA_ENABLE_DEPRECATION_WARNINGS
 		{
 			return false;
 		}
diff --git a/Engine/Plugins/Runtime/StateTree/Source/StateTreeTestSuite/Private/StateTreeTestTypes.h b/Engine/Plugins/Runtime/StateTree/Source/StateTreeTestSuite/Private/StateTreeTestTypes.h
--- a/Engine/Plugins/Runtime/StateTree/Source/StateTreeTestSuite/Private/StateTreeTestTypes.h
+++ b/Engine/Plugins/Runtime/StateTree/Source/StateTreeTestSuite/Private/StateTreeTestTypes.h
@@ -988,8 +988,10 @@ struct FStateTreeTest_PropertyStructObjectPtr
 	UPROPERTY(EditAnywhere, Category = "")
 	TSoftObjectPtr<UObject> C;
 
-	UPROPERTY(EditAnywhere, Category = "")
+PRAGMA_DISABLE_DEPRECATION_WARNINGS
+	UPROPERTY()
 	TLazyObjectPtr<UObject> D;
+PRAGMA_ENABLE_DEPRECATION_WARNINGS
 
 	UPROPERTY(EditAnywhere, Category = "")
 	TArray<TObjectPtr<UObject>> ArrayA;
@@ -1000,8 +1002,10 @@ struct FStateTreeTest_PropertyStructObjectPtr
 	UPROPERTY(EditAnywhere, Category = "")
 	TArray<TSoftObjectPtr<UObject>> ArrayC;
 
-	UPROPERTY(EditAnywhere, Category = "")
+PRAGMA_DISABLE_DEPRECATION_WARNINGS
+	UPROPERTY()
 	TArray<TLazyObjectPtr<UObject>> ArrayD;
+PRAGMA_ENABLE_DEPRECATION_WARNINGS
 };
 
 USTRUCT()
diff --git a/Engine/Source/Runtime/Foliage/Private/InstancedFoliage.cpp b/Engine/Source/Runtime/Foliage/Private/InstancedFoliage.cpp
--- a/Engine/Source/Runtime/Foliage/Private/InstancedFoliage.cpp
+++ b/Engine/Source/Runtime/Foliage/Private/InstancedFoliage.cpp
@@ -4745,8 +4745,10 @@ void AInstancedFoliageActor::PostLoad()
 			{
 				TPair<FFoliageInstanceBaseId, FFoliageInstanceBaseInfo>& Pair = *Iter;
 				FFoliageInstanceBaseInfo& BaseInfo = Pair.Value;
-				UActorComponent* Component = BaseInfo.BasePtr_DEPRECATED.Get();
+				PRAGMA_DISABLE_DEPRECATION_WARNINGS
+			UActorComponent* Component = BaseInfo.BasePtr_DEPRECATED.Get();
 				BaseInfo.BasePtr_DEPRECATED.Reset();
+			PRAGMA_ENABLE_DEPRECATION_WARNINGS
 
 				if (Component != nullptr)
 				{
@@ -4778,10 +4780,12 @@ void AInstancedFoliageActor::PostLoad()
 				TArray<FFoliageInstanceBasePtr_DEPRECATED>& BaseInfo_DEPRECATED = Pair.Value;
 				TArray<FFoliageInstanceBasePtr> BaseInfo;
 
-				for (FFoliageInstanceBasePtr_DEPRECATED& BasePtr_DEPRECATED : BaseInfo_DEPRECATED)
+				PRAGMA_DISABLE_DEPRECATION_WARNINGS
+			for (FFoliageInstanceBasePtr_DEPRECATED& BasePtr_DEPRECATED : BaseInfo_DEPRECATED)
 				{
 					UActorComponent* Component = BasePtr_DEPRECATED.Get();
 					BasePtr_DEPRECATED.Reset();
+			PRAGMA_ENABLE_DEPRECATION_WARNINGS
 
 					if (Component != nullptr)
 					{
diff --git a/Engine/Source/Runtime/Foliage/Public/FoliageInstanceBase.h b/Engine/Source/Runtime/Foliage/Public/FoliageInstanceBase.h
--- a/Engine/Source/Runtime/Foliage/Public/FoliageInstanceBase.h
+++ b/Engine/Source/Runtime/Foliage/Public/FoliageInstanceBase.h
@@ -10,7 +10,9 @@ class UActorComponent;
 
 typedef int32 FFoliageInstanceBaseId;
 typedef TSoftObjectPtr<UActorComponent> FFoliageInstanceBasePtr;
+PRAGMA_DISABLE_DEPRECATION_WARNINGS
 typedef TLazyObjectPtr<UActorComponent> FFoliageInstanceBasePtr_DEPRECATED;
+PRAGMA_ENABLE_DEPRECATION_WARNINGS
 
 #if WITH_EDITORONLY_DATA
 /**
@@ -34,9 +36,10 @@ struct FFoliageInstanceBaseInfo
 	FVector		CachedLocation;
 	FRotator	CachedRotation;
 	FVector		CachedDrawScale;
-
+	PRAGMA_DISABLE_DEPRECATION_WARNINGS
 	// Deprecated data stored for deserializaton done in PostLoad
 	FFoliageInstanceBasePtr_DEPRECATED BasePtr_DEPRECATED;
+	PRAGMA_ENABLE_DEPRECATION_WARNINGS
 };
 
 struct FFoliageInstanceBaseCache
diff --git a/Engine/Source/Runtime/Landscape/Classes/LandscapeComponent.h b/Engine/Source/Runtime/Landscape/Classes/LandscapeComponent.h
--- a/Engine/Source/Runtime/Landscape/Classes/LandscapeComponent.h
+++ b/Engine/Source/Runtime/Landscape/Classes/LandscapeComponent.h
@@ -512,9 +512,11 @@ class ULandscapeComponent : public UPrimitiveComponent
 	TArray<double> MipToMipMaxDeltas;
 
 #if WITH_EDITORONLY_DATA
+	PRAGMA_DISABLE_DEPRECATION_WARNINGS
 	UE_DEPRECATED(all, "CollisionComponent has been deprecated")
 	UPROPERTY()
 	TLazyObjectPtr<ULandscapeHeightfieldCollisionComponent> CollisionComponent_DEPRECATED;
+	PRAGMA_ENABLE_DEPRECATION_WARNINGS
 #endif // !WITH_EDITORONLY_DATA
 
 private:
diff --git a/Engine/Source/Runtime/Landscape/Classes/LandscapeHeightfieldCollisionComponent.h b/Engine/Source/Runtime/Landscape/Classes/LandscapeHeightfieldCollisionComponent.h
--- a/Engine/Source/Runtime/Landscape/Classes/LandscapeHeightfieldCollisionComponent.h
+++ b/Engine/Source/Runtime/Landscape/Classes/LandscapeHeightfieldCollisionComponent.h
@@ -86,8 +86,10 @@ class ULandscapeHeightfieldCollisionComponent : public UPrimitiveComponent
 
 #if WITH_EDITORONLY_DATA
 	/** Reference to render component */
+	PRAGMA_DISABLE_DEPRECATION_WARNINGS
 	UPROPERTY()
 	TLazyObjectPtr<ULandscapeComponent> RenderComponent_DEPRECATED;
+	PRAGMA_ENABLE_DEPRECATION_WARNINGS
 #endif // !WITH_EDITORONLY_DATA
 
 private:
diff --git a/Engine/Source/Runtime/Landscape/Classes/LandscapeSplinesComponent.h b/Engine/Source/Runtime/Landscape/Classes/LandscapeSplinesComponent.h
--- a/Engine/Source/Runtime/Landscape/Classes/LandscapeSplinesComponent.h
+++ b/Engine/Source/Runtime/Landscape/Classes/LandscapeSplinesComponent.h
@@ -73,14 +73,18 @@ struct FForeignWorldSplineData
 	GENERATED_USTRUCT_BODY()
 
 #if WITH_EDITORONLY_DATA
+	PRAGMA_DISABLE_DEPRECATION_WARNINGS
 	UPROPERTY()
 	TMap<TLazyObjectPtr<ULandscapeSplineControlPoint>, FForeignControlPointData> ForeignControlPointDataMap_DEPRECATED;
+	PRAGMA_ENABLE_DEPRECATION_WARNINGS
 
 	UPROPERTY()
 	TArray<FForeignControlPointData> ForeignControlPointData;
 
+	PRAGMA_DISABLE_DEPRECATION_WARNINGS
 	UPROPERTY()
 	TMap<TLazyObjectPtr<ULandscapeSplineSegment>, FForeignSplineSegmentData> ForeignSplineSegmentDataMap_DEPRECATED;
+	PRAGMA_ENABLE_DEPRECATION_WARNINGS
 
 	UPROPERTY()
 	TArray<FForeignSplineSegmentData> ForeignSplineSegmentData;
diff --git a/Engine/Source/Runtime/Landscape/Classes/LandscapeStreamingProxy.h b/Engine/Source/Runtime/Landscape/Classes/LandscapeStreamingProxy.h
--- a/Engine/Source/Runtime/Landscape/Classes/LandscapeStreamingProxy.h
+++ b/Engine/Source/Runtime/Landscape/Classes/LandscapeStreamingProxy.h
@@ -25,8 +25,10 @@ public:
 	ALandscapeStreamingProxy(const FObjectInitializer& ObjectInitializer);
 
 #if WITH_EDITORONLY_DATA
+	PRAGMA_DISABLE_DEPRECATION_WARNINGS
 	UPROPERTY()
 	TLazyObjectPtr<ALandscape> LandscapeActor_DEPRECATED;
+	PRAGMA_ENABLE_DEPRECATION_WARNINGS
 #endif // WITH_EDITORONLY_DATA
 
 private:
diff --git a/Engine/Source/Runtime/Landscape/Private/LandscapeCollision.cpp b/Engine/Source/Runtime/Landscape/Private/LandscapeCollision.cpp
--- a/Engine/Source/Runtime/Landscape/Private/LandscapeCollision.cpp
+++ b/Engine/Source/Runtime/Landscape/Private/LandscapeCollision.cpp
@@ -2001,11 +2001,13 @@ void ULandscapeHeightfieldCollisionComponent::PostLoad()
 
 #if WITH_EDITORONLY_DATA
 	// If the RenderComponent is not set yet and we're transferring the property from the lazy object pointer it was previously stored as to the object ptr it is now stored as :
+	PRAGMA_DISABLE_DEPRECATION_WARNINGS
 	if (!RenderComponentRef && RenderComponent_DEPRECATED.IsValid())
 	{
 		RenderComponentRef = RenderComponent_DEPRECATED.Get();
 		RenderComponent_DEPRECATED = nullptr;
 	}
+	PRAGMA_ENABLE_DEPRECATION_WARNINGS
 #endif // !WITH_EDITORONLY_DATA
 
 #endif // WITH_EDITOR
diff --git a/Engine/Source/Runtime/Landscape/Private/LandscapeEdit.cpp b/Engine/Source/Runtime/Landscape/Private/LandscapeEdit.cpp
--- a/Engine/Source/Runtime/Landscape/Private/LandscapeEdit.cpp
+++ b/Engine/Source/Runtime/Landscape/Private/LandscapeEdit.cpp
@@ -6334,6 +6334,7 @@ void ALandscapeStreamingProxy::PostRegisterAllComponents()
 
 #if WITH_EDITORONLY_DATA
 	// If the landscape actor is not set yet and we're transferring the property from the lazy object pointer it was previously stored as to the soft object ptr it is now stored as :
+	PRAGMA_DISABLE_DEPRECATION_WARNINGS
 	if (!LandscapeActorRef)
 	{
 		// Because of how lazy object pointers were made, the only way we can deprecate them is if the object they're pointing to is currently loaded :
@@ -6358,6 +6359,7 @@ void ALandscapeStreamingProxy::PostRegisterAllComponents()
 			FMessageLog("MapCheck").Open(EMessageSeverity::Warning);
 		}
 	}
+	PRAGMA_ENABLE_DEPRECATION_WARNINGS
 #endif // WITH_EDITORONLY_DATA
 
 	if (LandscapeGuid.IsValid())
