[StateTree] Remove deprecated TLazyObjectPtr test in unit test.

#rnx
#Synced-CL 52302118

diff --git a/Engine/Plugins/Runtime/StateTree/Source/StateTreeTestSuite/Private/StateTreeBindingTest.cpp b/Engine/Plugins/Runtime/StateTree/Source/StateTreeTestSuite/Private/StateTreeBindingTest.cpp
--- a/Engine/Plugins/Runtime/StateTree/Source/StateTreeTestSuite/Private/StateTreeBindingTest.cpp
+++ b/Engine/Plugins/Runtime/StateTree/Source/StateTreeTestSuite/Private/StateTreeBindingTest.cpp
@@ -445,48 +445,32 @@ struct FStateTreeTest_BindingsCompiler_WithObjectPtr : FStateTreeTestBase
 		SourceA.A = TBaseStructure<FVector>::Get();
 		SourceA.B = TBaseStructure<FRotator>::Get();
 		SourceA.C = TBaseStructure<FMatrix>::Get();
-PRAGMA_DISABLE_DEPRECATION_WARNINGS
-		SourceA.D = TBaseStructure<FTransform>::Get();
-PRAGMA_ENABLE_DEPRECATION_WARNINGS
 		constexpr int32 NumberOfElements = 4;
 		SourceA.ArrayA.SetNum(NumberOfElements, EAllowShrinking::No);
 		SourceA.ArrayB.SetNum(NumberOfElements, EAllowShrinking::No);
 		SourceA.ArrayC.SetNum(NumberOfElements, EAllowShrinking::No);
-PRAGMA_DISABLE_DEPRECATION_WARNINGS
-		SourceA.ArrayD.SetNum(NumberOfElements, EAllowShrinking::No);
-PRAGMA_ENABLE_DEPRECATION_WARNINGS
+		
 		SourceA.ArrayA[2] = TBaseStructure<FVector>::Get();
 		SourceA.ArrayB[2] = TBaseStructure<FRotator>::Get();
 		SourceA.ArrayC[2] = TBaseStructure<FMatrix>::Get();
-PRAGMA_DISABLE_DEPRECATION_WARNINGS
-		SourceA.ArrayD[2] = TBaseStructure<FTransform>::Get();
-PRAGMA_ENABLE_DEPRECATION_WARNINGS
+
 		FStateTreeTest_PropertyStruct SourceB;
 
 		FStateTreeTest_PropertyStructObjectPtr Target;
 		Target.ArrayA.SetNum(2, EAllowShrinking::No);
 		Target.ArrayB.SetNum(2, EAllowShrinking::No);
 		Target.ArrayC.SetNum(2, EAllowShrinking::No);
-PRAGMA_DISABLE_DEPRECATION_WARNINGS
-		Target.ArrayD.SetNum(2, EAllowShrinking::No);
-PRAGMA_ENABLE_DEPRECATION_WARNINGS
 
 		auto ResetTargetData = [&Target]()
 			{
 				Target.A = nullptr;
 				Target.B = nullptr;
 				Target.C = nullptr;
-PRAGMA_DISABLE_DEPRECATION_WARNINGS
-				Target.D = nullptr;
-PRAGMA_ENABLE_DEPRECATION_WARNINGS
 				for (int32 Index = 0; Index < Target.ArrayA.Num(); ++Index)
 				{
 					Target.ArrayA[Index] = nullptr;
 					Target.ArrayB[Index] = nullptr;
 					Target.ArrayC[Index] = nullptr;
-PRAGMA_DISABLE_DEPRECATION_WARNINGS
-					Target.ArrayD[Index] = nullptr;
-PRAGMA_ENABLE_DEPRECATION_WARNINGS
 				}
 				return true;
 			};
@@ -494,19 +478,12 @@ PRAGMA_ENABLE_DEPRECATION_WARNINGS
 
 		auto TestCopy = [this, &Target](UObject* Object, const FStringView SourceStr)
 			{
-
 				AITEST_EQUAL(FString::Printf(TEXT("Expect A to be copied from %s"), SourceStr.GetData()), Target.A.Get(), Object);
 				AITEST_EQUAL(FString::Printf(TEXT("Expect B to be copied from %s"), SourceStr.GetData()), Target.B.Get(), Object);
 				AITEST_EQUAL(FString::Printf(TEXT("Expect C to be copied from %s"), SourceStr.GetData()), Target.C.Get(), Object);
-PRAGMA_DISABLE_DEPRECATION_WARNINGS
-				AITEST_EQUAL(FString::Printf(TEXT("Expect D to be copied from %s"), SourceStr.GetData()), Target.D.Get(), Object);
-PRAGMA_ENABLE_DEPRECATION_WARNINGS
 				AITEST_EQUAL(FString::Printf(TEXT("Expect A[] to be copied from %s"), SourceStr.GetData()), Target.ArrayA[1].Get(), Object);
 				AITEST_EQUAL(FString::Printf(TEXT("Expect B[] to be copied from %s"), SourceStr.GetData()), Target.ArrayB[1].Get(), Object);
 				AITEST_EQUAL(FString::Printf(TEXT("Expect C[] to be copied from %s"), SourceStr.GetData()), Target.ArrayC[1].Get(), Object);
-PRAGMA_DISABLE_DEPRECATION_WARNINGS
-				AITEST_EQUAL(FString::Printf(TEXT("Expect D[] to be copied from %s"), SourceStr.GetData()), Target.ArrayD[1].Get(), Object);
-PRAGMA_ENABLE_DEPRECATION_WARNINGS
 				return true;
 			};
 
@@ -553,16 +530,6 @@ PRAGMA_ENABLE_DEPRECATION_WARNINGS
 		{
 			return false;
 		}
-PRAGMA_DISABLE_DEPRECATION_WARNINGS
-		if (!CopyAndTestBatch(CopyBatchIndex7, SourceA.D.Get(), TEXT("SourceA.D")))
-		{
-			return false;
-		}
-		if (!CopyAndTestBatch(CopyBatchIndex8, SourceA.D.Get(), TEXT("SourceA.ArrayD")))
-PRAGMA_ENABLE_DEPRECATION_WARNINGS
-		{
-			return false;
-		}
 
 		// Test promote array
 		for (int32 Index = 0; Index < NumberOfElements; ++Index)
@@ -570,9 +537,6 @@ PRAGMA_ENABLE_DEPRECATION_WARNINGS
 			SourceA.ArrayA[Index] = TBaseStructure<FVector>::Get();
 			SourceA.ArrayB[Index] = TBaseStructure<FRotator>::Get();
 			SourceA.ArrayC[Index] = TBaseStructure<FMatrix>::Get();
-PRAGMA_DISABLE_DEPRECATION_WARNINGS
-			SourceA.ArrayD[Index] = TBaseStructure<FTransform>::Get();
-PRAGMA_ENABLE_DEPRECATION_WARNINGS
 		}
 
 		auto TestCopyArray = [this, &Target, NumberOfElements](UObject* Object, const FStringView SourceStr)
@@ -582,9 +546,6 @@ PRAGMA_ENABLE_DEPRECATION_WARNINGS
 					AITEST_EQUAL(FString::Printf(TEXT("Expect A[] to be copied from %s"), SourceStr.GetData()), Target.ArrayA[Index].Get(), Object);
 					AITEST_EQUAL(FString::Printf(TEXT("Expect B[] to be copied from %s"), SourceStr.GetData()), Target.ArrayB[Index].Get(), Object);
 					AITEST_EQUAL(FString::Printf(TEXT("Expect C[] to be copied from %s"), SourceStr.GetData()), Target.ArrayC[Index].Get(), Object);
-PRAGMA_DISABLE_DEPRECATION_WARNINGS
-					AITEST_EQUAL(FString::Printf(TEXT("Expect D[] to be copied from %s"), SourceStr.GetData()), Target.ArrayD[Index].Get(), Object);
-PRAGMA_ENABLE_DEPRECATION_WARNINGS
 				}
 				return true;
 			};
@@ -612,12 +573,6 @@ PRAGMA_ENABLE_DEPRECATION_WARNINGS
 		{
 			return false;
 		}
-PRAGMA_DISABLE_DEPRECATION_WARNINGS
-		if (!CopyAndTestBatchArray(CopyBatchIndex12, SourceA.D.Get(), TEXT("SourceA.ArrayD")))
-PRAGMA_ENABLE_DEPRECATION_WARNINGS
-		{
-			return false;
-		}
 
 		return true;
 	}
diff --git a/Engine/Plugins/Runtime/StateTree/Source/StateTreeTestSuite/Private/StateTreeTestTypes.h b/Engine/Plugins/Runtime/StateTree/Source/StateTreeTestSuite/Private/StateTreeTestTypes.h
--- a/Engine/Plugins/Runtime/StateTree/Source/StateTreeTestSuite/Private/StateTreeTestTypes.h
+++ b/Engine/Plugins/Runtime/StateTree/Source/StateTreeTestSuite/Private/StateTreeTestTypes.h
@@ -988,11 +988,6 @@ struct FStateTreeTest_PropertyStructObjectPtr
 	UPROPERTY(EditAnywhere, Category = "")
 	TSoftObjectPtr<UObject> C;
 
-PRAGMA_DISABLE_DEPRECATION_WARNINGS
-	UPROPERTY()
-	TLazyObjectPtr<UObject> D;
-PRAGMA_ENABLE_DEPRECATION_WARNINGS
-
 	UPROPERTY(EditAnywhere, Category = "")
 	TArray<TObjectPtr<UObject>> ArrayA;
 
@@ -1001,11 +996,6 @@ PRAGMA_ENABLE_DEPRECATION_WARNINGS
 
 	UPROPERTY(EditAnywhere, Category = "")
 	TArray<TSoftObjectPtr<UObject>> ArrayC;
-
-PRAGMA_DISABLE_DEPRECATION_WARNINGS
-	UPROPERTY()
-	TArray<TLazyObjectPtr<UObject>> ArrayD;
-PRAGMA_ENABLE_DEPRECATION_WARNINGS
 };
 
 USTRUCT()
