[StateTree] Add option to update State Parameter bindings on tick. The bindings are already copied on state enter and state exit.
#Synced-CL 52302118

diff --git a/Engine/Plugins/Runtime/StateTree/Source/StateTreeEditorModule/Private/Customizations/StateTreeStateDetails.cpp b/Engine/Plugins/Runtime/StateTree/Source/StateTreeEditorModule/Private/Customizations/StateTreeStateDetails.cpp
--- a/Engine/Plugins/Runtime/StateTree/Source/StateTreeEditorModule/Private/Customizations/StateTreeStateDetails.cpp
+++ b/Engine/Plugins/Runtime/StateTree/Source/StateTreeEditorModule/Private/Customizations/StateTreeStateDetails.cpp
@@ -109,6 +109,7 @@ void FStateTreeStateDetails::CustomizeDetails(IDetailLayoutBuilder& DetailBuilde
 	const TSharedPtr<IPropertyHandle> RequiredEventToEnterProperty = DetailBuilder.GetProperty(GET_MEMBER_NAME_CHECKED(UStateTreeState, RequiredEventToEnter));
 	const TSharedPtr<IPropertyHandle> CheckPrerequisitesProperty = DetailBuilder.GetProperty(GET_MEMBER_NAME_CHECKED(UStateTreeState, bCheckPrerequisitesWhenActivatingChildDirectly));
 	const TSharedPtr<IPropertyHandle> WeightProperty = DetailBuilder.GetProperty(GET_MEMBER_NAME_CHECKED(UStateTreeState, Weight));
+	const TSharedPtr<IPropertyHandle> CopyParameterBindingsOnTickProperty = DetailBuilder.GetProperty(GET_MEMBER_NAME_CHECKED(UStateTreeState, bCopyParameterBindingsOnTick));
 
 	// Never show enabled
 	EnabledProperty->MarkHiddenByCustomization();
@@ -361,6 +362,10 @@ void FStateTreeStateDetails::CustomizeDetails(IDetailLayoutBuilder& DetailBuilde
 		TSharedRef<FStateTreeStateParametersInstanceDataDetails> InstanceDetails =
 			MakeShareable(new FStateTreeStateParametersInstanceDataDetails(ParametersProperty, ParametersProperty, PropUtils, bFixedLayout, ID, WeakEditorData, WeakState));
 		ParametersCategory.AddCustomBuilder(InstanceDetails);
+
+		// Tick Parameter Bindings option
+		CopyParameterBindingsOnTickProperty->MarkHiddenByCustomization();
+		ParametersCategory.AddProperty(CopyParameterBindingsOnTickProperty);
 	}
 
 	// Enter conditions
diff --git a/Engine/Plugins/Runtime/StateTree/Source/StateTreeEditorModule/Private/StateTreeCompiler.cpp b/Engine/Plugins/Runtime/StateTree/Source/StateTreeEditorModule/Private/StateTreeCompiler.cpp
--- a/Engine/Plugins/Runtime/StateTree/Source/StateTreeEditorModule/Private/StateTreeCompiler.cpp
+++ b/Engine/Plugins/Runtime/StateTree/Source/StateTreeEditorModule/Private/StateTreeCompiler.cpp
@@ -916,6 +916,7 @@ bool FStateTreeCompiler::CreateStateRecursive(UStateTreeState& State, const FSta
 
 	CompactState.bHasCustomTickRate = State.bHasCustomTickRate && StateTree->Schema->IsScheduledTickAllowed();
 	CompactState.CustomTickRate = FMath::Max(State.CustomTickRate, 0.0f);
+	CompactState.bCopyParameterBindingsOnTick = State.bCopyParameterBindingsOnTick;
 	if (CompactState.bHasCustomTickRate && State.CustomTickRate < 0.0f)
 	{
 		Log.Reportf(EMessageSeverity::Warning, TEXT("The custom tick rate has to be greater than or equal to 0."));
diff --git a/Engine/Plugins/Runtime/StateTree/Source/StateTreeEditorModule/Private/StateTreeState.cpp b/Engine/Plugins/Runtime/StateTree/Source/StateTreeEditorModule/Private/StateTreeState.cpp
--- a/Engine/Plugins/Runtime/StateTree/Source/StateTreeEditorModule/Private/StateTreeState.cpp
+++ b/Engine/Plugins/Runtime/StateTree/Source/StateTreeEditorModule/Private/StateTreeState.cpp
@@ -15,6 +15,15 @@
 
 #include UE_INLINE_GENERATED_CPP_BY_NAME(StateTreeState)
 
+namespace UE::StateTree::Editor
+{
+	bool bCopyParameterBindingsOnTickByDefault = false;
+	FAutoConsoleVariableRef CVarTickParameterBindingsByDefault(
+		TEXT("StateTree.Editor.TickParameterBindingsByDefault"),
+		bCopyParameterBindingsOnTickByDefault,
+		TEXT("Default value for 'Tick Parameter Bindings' on newly created states.")
+	);
+}
 
 //////////////////////////////////////////////////////////////////////////
 // FStateTreeStateParameters
@@ -84,6 +93,7 @@ UStateTreeState::UStateTreeState(const FObjectInitializer& ObjectInitializer)
 	{
 		ID = FGuid::NewGuid();
 		Parameters.ID = FGuid::NewGuid();
+		bCopyParameterBindingsOnTick = UE::StateTree::Editor::bCopyParameterBindingsOnTickByDefault;
 	}
 }
 
@@ -376,6 +386,7 @@ PRAGMA_DISABLE_DEPRECATION_WARNINGS
 	constexpr int32 AddedTransitionIdsVersion = FStateTreeCustomVersion::AddedTransitionIds;
 	constexpr int32 OverridableStateParametersVersion = FStateTreeCustomVersion::OverridableStateParameters;
 	constexpr int32 AddedCheckingParentsPrerequisitesVersion = FStateTreeCustomVersion::AddedCheckingParentsPrerequisites;
+	constexpr int32 TickParameterBindingsVersion = FStateTreeCustomVersion::TickParameterBindings;
 PRAGMA_ENABLE_DEPRECATION_WARNINGS
 
 	if (CurrentStateTreeCustomVersion < AddedTransitionIdsVersion)
@@ -404,7 +415,13 @@ PRAGMA_ENABLE_DEPRECATION_WARNINGS
 	{
 		bCheckPrerequisitesWhenActivatingChildDirectly = false;
 	}
-	
+
+	if (CurrentStateTreeCustomVersion < TickParameterBindingsVersion)
+	{
+		// For backward compatibility: Linked and LinkedAsset states always ticked parameter bindings before this option was added.
+		bCopyParameterBindingsOnTick = (Type == EStateTreeStateType::Linked || Type == EStateTreeStateType::LinkedAsset);
+	}
+
 #endif // WITH_EDITORONLY_DATA
 
 #if WITH_EDITOR
diff --git a/Engine/Plugins/Runtime/StateTree/Source/StateTreeEditorModule/Public/StateTreeState.h b/Engine/Plugins/Runtime/StateTree/Source/StateTreeEditorModule/Public/StateTreeState.h
--- a/Engine/Plugins/Runtime/StateTree/Source/StateTreeEditorModule/Public/StateTreeState.h
+++ b/Engine/Plugins/Runtime/StateTree/Source/StateTreeEditorModule/Public/StateTreeState.h
@@ -449,6 +449,13 @@ public:
 	UPROPERTY(EditDefaultsOnly, Category = "State", meta=(InlineEditConditionToggle))
 	bool bHasCustomTickRate = false;
 
+	/**
+	 * If true, state parameter bindings are copied each task tick, not only when the state is entered or exited.
+	 * Has no effect on ScheduledTick.
+	 */
+	UPROPERTY(EditDefaultsOnly, Category = "State")
+	bool bCopyParameterBindingsOnTick = false;
+
 	/** Parameters of this state. If the state is linked to another state or asset, the parameters are for the linked state. */
 	UPROPERTY(EditDefaultsOnly, Category = "State")
 	FStateTreeStateParameters Parameters;
diff --git a/Engine/Plugins/Runtime/StateTree/Source/StateTreeModule/Private/StateTreeExecutionContext.cpp b/Engine/Plugins/Runtime/StateTree/Source/StateTreeModule/Private/StateTreeExecutionContext.cpp
--- a/Engine/Plugins/Runtime/StateTree/Source/StateTreeModule/Private/StateTreeExecutionContext.cpp
+++ b/Engine/Plugins/Runtime/StateTree/Source/StateTreeModule/Private/StateTreeExecutionContext.cpp
@@ -4721,13 +4721,12 @@ EStateTreeRunStatus FStateTreeExecutionContext::TickTasks(const float DeltaTime)
 				, (FrameIndex + StateIndex + 1) * Debug::IndentSize, TEXT("")
 				, *DebugGetStatePath(Exec.ActiveFrames, TickArgs.Frame, StateIndex));
 
-			if (CurrentState.Type == EStateTreeStateType::Linked || CurrentState.Type == EStateTreeStateType::LinkedAsset)
+			if (CurrentState.bCopyParameterBindingsOnTick
+				&& CurrentState.ParameterDataHandle.IsValid()
+				&& CurrentState.ParameterBindingsBatch.IsValid())
 			{
-				if (CurrentState.ParameterDataHandle.IsValid() && CurrentState.ParameterBindingsBatch.IsValid())
-				{
-					const FStateTreeDataView StateParamsDataView = GetDataView(TickArgs.ParentFrame, *TickArgs.Frame, CurrentState.ParameterDataHandle);
-					CopyBatchOnActiveInstances(TickArgs.ParentFrame, *TickArgs.Frame, StateParamsDataView, CurrentState.ParameterBindingsBatch);
-				}
+				const FStateTreeDataView StateParamsDataView = GetDataView(TickArgs.ParentFrame, *TickArgs.Frame, CurrentState.ParameterDataHandle);
+				CopyBatchOnActiveInstances(TickArgs.ParentFrame, *TickArgs.Frame, StateParamsDataView, CurrentState.ParameterBindingsBatch);
 			}
 
 			const bool bHasEvents = EventQueue && EventQueue->HasEvents();
diff --git a/Engine/Plugins/Runtime/StateTree/Source/StateTreeModule/Public/StateTree.h b/Engine/Plugins/Runtime/StateTree/Source/StateTreeModule/Public/StateTree.h
--- a/Engine/Plugins/Runtime/StateTree/Source/StateTreeModule/Public/StateTree.h
+++ b/Engine/Plugins/Runtime/StateTree/Source/StateTreeModule/Public/StateTree.h
@@ -65,6 +65,8 @@ FStateTreeCustomVersion
 		AddedBindingToEvents,
 		// Added checking parent states' prerequisites when activating child state directly.
 		AddedCheckingParentsPrerequisites,
+		// Added Tick Parameter Binding.
+		TickParameterBindings,
 
 		// -----<new versions can be added above this line>-------------------------------------------------
 		VersionPlusOne,
diff --git a/Engine/Plugins/Runtime/StateTree/Source/StateTreeModule/Public/StateTreeTypes.h b/Engine/Plugins/Runtime/StateTree/Source/StateTreeModule/Public/StateTreeTypes.h
--- a/Engine/Plugins/Runtime/StateTree/Source/StateTreeModule/Public/StateTreeTypes.h
+++ b/Engine/Plugins/Runtime/StateTree/Source/StateTreeModule/Public/StateTreeTypes.h
@@ -1029,6 +1029,13 @@ struct FCompactStateTreeState
 	/** If set to true, the state has a custom tick rate. */
 	UPROPERTY()
 	uint8 bHasCustomTickRate : 1 = false;
+
+	/**
+	 * If true, state parameter bindings are copied on each task tick.
+	 * Parameters bindings are always copied on state enter and state exit.
+	 */
+	UPROPERTY()
+	uint8 bCopyParameterBindingsOnTick : 1 = false;
 };
 
 USTRUCT()
