[StateTree] Fix unittest for copy state parameters. The fix was made with CL 52346361

#rnx
#Synced-CL 52302118

diff --git a/Engine/Plugins/Runtime/StateTree/Source/StateTreeTestSuite/Private/StateTreeLinkedStateTest.cpp b/Engine/Plugins/Runtime/StateTree/Source/StateTreeTestSuite/Private/StateTreeLinkedStateTest.cpp
--- a/Engine/Plugins/Runtime/StateTree/Source/StateTreeTestSuite/Private/StateTreeLinkedStateTest.cpp
+++ b/Engine/Plugins/Runtime/StateTree/Source/StateTreeTestSuite/Private/StateTreeLinkedStateTest.cpp
@@ -1220,6 +1220,7 @@ struct FStateTreeTest_Linked_GlobalParameter : FStateTreeTestBase
 			{
 				AddInt(Root.Parameters.Parameters, "Tree2StateRootParametersInt");
 				AddDouble(Root.Parameters.Parameters, "Tree2StateRootParametersDouble");
+				Root.bCopyParameterBindingsOnTick = true;
 				EditorData2.AddPropertyBinding(FPropertyBindingPath(EditorData2.GetRootParametersGuid(), "Tree2GlobalInt"), FPropertyBindingPath(Root.Parameters.ID, "Tree2StateRootParametersInt"));
 			}
 			{
@@ -1240,6 +1241,7 @@ struct FStateTreeTest_Linked_GlobalParameter : FStateTreeTestBase
 				{
 					AddDouble(State1.Parameters.Parameters, "Tree2State1ParametersDouble");
 					AddInt(State1.Parameters.Parameters, "Tree2State1ParametersInt");
+					State1.bCopyParameterBindingsOnTick = true;
 					EditorData2.AddPropertyBinding(FPropertyBindingPath(Root.Parameters.ID, "Tree2StateRootParametersInt"), FPropertyBindingPath(State1.Parameters.ID, "Tree2State1ParametersInt"));
 				}
 
@@ -1293,6 +1295,7 @@ struct FStateTreeTest_Linked_GlobalParameter : FStateTreeTestBase
 			{
 				AddDouble(Root.Parameters.Parameters, "Tree1StateRootParametersDouble");
 				AddInt(Root.Parameters.Parameters, "Tree1StateRootParametersInt");
+				Root.bCopyParameterBindingsOnTick = true;
 				EditorData1.AddPropertyBinding(FPropertyBindingPath(EditorData1.GetRootParametersGuid(), "Tree1GlobalInt"), FPropertyBindingPath(Root.Parameters.ID, "Tree1StateRootParametersInt"));
 			}
 			{
@@ -1314,6 +1317,7 @@ struct FStateTreeTest_Linked_GlobalParameter : FStateTreeTestBase
 				State1.AddTransition(EStateTreeTransitionTrigger::OnStateSucceeded, EStateTreeTransitionType::NextState);
 
 				State1.SetParametersPropertyOverridden(Tree2GlobalParameter_ValueID_Int, true);
+				State1.bCopyParameterBindingsOnTick = true;
 				const FInstancedPropertyBag* Parameters = State1.GetDefaultParameters();
 				AITEST_TRUE(TEXT("Parameter is invalid"), Parameters != nullptr);
 				EditorData1.AddPropertyBinding(FPropertyBindingPath(GlobalTask1.ID, "Value"), FPropertyBindingPath(State1.Parameters.ID, "Tree2GlobalInt"));
@@ -1338,6 +1342,7 @@ struct FStateTreeTest_Linked_GlobalParameter : FStateTreeTestBase
 						AddDouble(State3.Parameters.Parameters, "Tree1State3ParametersDouble1");
 						AddDouble(State3.Parameters.Parameters, "Tree1State3ParametersDouble2");
 						AddInt(State3.Parameters.Parameters, "Tree1State3ParametersInt");
+						State3.bCopyParameterBindingsOnTick = true;
 						EditorData1.AddPropertyBinding(FPropertyBindingPath(Sub1.Parameters.ID, "Tree1StateSub1ParametersInt"), FPropertyBindingPath(State3.Parameters.ID, "Tree1State3ParametersInt"));
 					}
 
@@ -1366,6 +1371,7 @@ struct FStateTreeTest_Linked_GlobalParameter : FStateTreeTestBase
 				State2.AddTransition(EStateTreeTransitionTrigger::OnStateSucceeded, EStateTreeTransitionType::Succeeded);
 
 				State2.SetParametersPropertyOverridden(Tree1Sub1Parameter_ValueID_Int, true);
+				State2.bCopyParameterBindingsOnTick = true;
 				const FInstancedPropertyBag* Parameters = State2.GetDefaultParameters();
 				AITEST_TRUE(TEXT("Parameter is invalid"), Parameters != nullptr);
 				EditorData1.AddPropertyBinding(FPropertyBindingPath(GlobalTask1.ID, "Value"), FPropertyBindingPath(State2.Parameters.ID, "Tree1StateSub1ParametersInt"));
@@ -1427,22 +1433,17 @@ struct FStateTreeTest_Linked_GlobalParameter : FStateTreeTestBase
 			AITEST_TRUE(TEXT("Tick should tick Tree1GlobalTask1"), Exec.Expect("Tree1GlobalTask1", TEXT("Tick6")));
 			AITEST_TRUE(TEXT("Tick should tick Tree1GlobalTask2"), Exec.Expect("Tree1GlobalTask2", TEXT("Tick6")));
 			AITEST_TRUE(TEXT("Tick should tick Tree1StateRootTask1"), Exec.Expect("Tree1StateRootTask1", TEXT("Tick6")));
-			//@todo Binding on the state are not updated on tick. Is this a bug? See Tree1StateRootParametersInt
-			//AITEST_TRUE(TEXT("Tick should tick Tree1StateRootTask2), Exec.Expect("Tree1StateRootTask2", TEXT("Tick6")));
+			AITEST_TRUE(TEXT("Tick should tick Tree1StateRootTask2"), Exec.Expect("Tree1StateRootTask2", TEXT("Tick6")));
 			AITEST_TRUE(TEXT("Tick should tick Tree1StateRootTask3"), Exec.Expect("Tree1StateRootTask3", TEXT("Tick6")));
 			AITEST_TRUE(TEXT("Tick should tick Tree2GlobalTask1"), Exec.Expect("Tree2GlobalTask1", TEXT("Tick6")));
 			AITEST_TRUE(TEXT("Tick should tick Tree2GlobalTask2"), Exec.Expect("Tree2GlobalTask2", TEXT("Tick6")));
 			AITEST_TRUE(TEXT("Tick should tick Tree2StateRootTask1"), Exec.Expect("Tree2StateRootTask1", TEXT("Tick6")));
-			//@todo Binding on the state are not updated on tick. Is this a bug? See Tree2StateRootParametersInt
-			//AITEST_TRUE(TEXT("Tick should tick Tree2StateRootTask2"), Exec.Expect("Tree2StateRootTask2", TEXT("Tick6")));
+			AITEST_TRUE(TEXT("Tick should tick Tree2StateRootTask2"), Exec.Expect("Tree2StateRootTask2", TEXT("Tick6")));
 			AITEST_TRUE(TEXT("Tick should tick Tree2StateRootTask3"), Exec.Expect("Tree2StateRootTask3", TEXT("Tick6")));
 			AITEST_TRUE(TEXT("Tick should tick Tree2State1Task1"), Exec.Expect("Tree2State1Task1", TEXT("Tick6")));
-			//@todo Binding on the state are not updated on tick. Is this a bug? See Tree2StateRootParametersInt
-			//AITEST_TRUE(TEXT("Tick should tick Tree2State1Task2"), Exec.Expect("Tree2State1Task2", TEXT("Tick6")));
+			AITEST_TRUE(TEXT("Tick should tick Tree2State1Task2"), Exec.Expect("Tree2State1Task2", TEXT("Tick6")));
 			AITEST_TRUE(TEXT("Tick should tick Tree2State1Task3"), Exec.Expect("Tree2State1Task3", TEXT("Tick6")));
-			//@todo Binding on the state are not updated on tick. Is this a bug? See Tree2State1ParametersInt
-			//AITEST_TRUE(TEXT("Tick should tick Tree2State1Task4"), Exec.Expect("Tree2State1Task4", TEXT("Tick6")));
-
+			AITEST_TRUE(TEXT("Tick should tick Tree2State1Task4"), Exec.Expect("Tree2State1Task4", TEXT("Tick6")));
 			AITEST_TRUE(TEXT("Tick should enter Tree1StateSub1Task1"), Exec.Expect("Tree1StateSub1Task1", TEXT("EnterState6")));
 			AITEST_TRUE(TEXT("Tick should enter Tree1State3Task1"), Exec.Expect("Tree1State3Task1", TEXT("EnterState6")));
 			AITEST_TRUE(TEXT("Tick should enter Tree1State3Task2"), Exec.Expect("Tree1State3Task2", TEXT("EnterState6")));
