[StateTree] Fix bug with ConsideredForCompletion when the completion type is all, and the state has tasks, but they are not considered for completion.

#Synced-CL 52302118

diff --git a/Engine/Plugins/Runtime/StateTree/Source/StateTreeEditorModule/Private/StateTreeCompiler.cpp b/Engine/Plugins/Runtime/StateTree/Source/StateTreeEditorModule/Private/StateTreeCompiler.cpp
--- a/Engine/Plugins/Runtime/StateTree/Source/StateTreeEditorModule/Private/StateTreeCompiler.cpp
+++ b/Engine/Plugins/Runtime/StateTree/Source/StateTreeEditorModule/Private/StateTreeCompiler.cpp
@@ -395,26 +395,37 @@ namespace UE::StateTree::Compiler
 		int32 NumberOfBitsNeeded = 0;
 		const int32 NumberOfTasks = ValidTasks.Num();
 
-		// No task, state/frame needs at least one flag to set the state itself completes (ie. for linked state).
+		TBitArray<> ConsideredForCompletion;
+		bool bHasTaskConsideredForCompletion = false;
+		for (int32 Index = 0; Index < NumberOfTasks; ++Index)
+		{
+			const int32 TaskIndex = ValidTasks[Index];
+			const bool bIsTaskConsideredForCompletion = UE::StateTreeEditor::EditorNodeUtils::IsTaskEnabled(AllTasks[TaskIndex])
+				&& UE::StateTreeEditor::EditorNodeUtils::IsTaskConsideredForCompletion(AllTasks[TaskIndex]);
+
+			ConsideredForCompletion.Add(bIsTaskConsideredForCompletion);
+			bHasTaskConsideredForCompletion |= bIsTaskConsideredForCompletion;
+		}
+
+		NumberOfBitsNeeded = NumberOfTasks;
+
 		//Each state will take at least 1 bit.
-		if (NumberOfTasks == 0)
+		//No task are marked as completed, we still need a flag to mark the state itself completes (ie. for linked state).
+		//Linked state can't have tasks but keep that logic in case it is needed later.
+		//A mask of 0, will cause issue with the GetCompletionsStatus.
+		if (!bHasTaskConsideredForCompletion)
 		{
-			Mask = 1;
-			NumberOfBitsNeeded = 1;
+			Mask |= 1;
+			++NumberOfBitsNeeded;
 		}
-		else
+
+		for (int32 Index = NumberOfTasks - 1; Index >= 0; --Index)
 		{
-			for (int32 Index = NumberOfTasks-1; Index >= 0; --Index)
+			Mask <<= 1;
+			if (ConsideredForCompletion[Index])
 			{
-				const int32 TaskIndex = ValidTasks[Index];
-				Mask <<= 1;
-				if (UE::StateTreeEditor::EditorNodeUtils::IsTaskEnabled(AllTasks[TaskIndex])
-					&& UE::StateTreeEditor::EditorNodeUtils::IsTaskConsideredForCompletion(AllTasks[TaskIndex]))
-				{
-					Mask |= 1;
-				}
+				Mask |= 1;
 			}
-			NumberOfBitsNeeded = NumberOfTasks;
 		}
 
 		constexpr int32 NumberOfBitsPerMask = sizeof(FStateTreeTasksCompletionStatus::FMaskType) * 8;
@@ -1396,7 +1407,7 @@ bool FStateTreeCompiler::CreateStateTasksAndParameters()
 				StartBitIndex = GlobalTaskEndBit;
 			}
 
-			const UE::StateTree::Compiler::FCompletionTasksMaskResult MaskResult = UE::StateTree::Compiler::MakeCompletionTasksMask(StartBitIndex, Tasks, ValidTaskNodeIndex);
+			UE::StateTree::Compiler::FCompletionTasksMaskResult MaskResult = UE::StateTree::Compiler::MakeCompletionTasksMask(StartBitIndex, Tasks, ValidTaskNodeIndex);
 
 			const int32 CompletionTasksMaskBufferIndex = MaskResult.MaskBufferIndex;
 			if (const auto Validation = UE::StateTree::Compiler::IsValidCount8(CompletionTasksMaskBufferIndex); Validation.DidFail())
diff --git a/Engine/Plugins/Runtime/StateTree/Source/StateTreeTestSuite/Private/StateTreeTaskStateTest.cpp b/Engine/Plugins/Runtime/StateTree/Source/StateTreeTestSuite/Private/StateTreeTaskStateTest.cpp
--- a/Engine/Plugins/Runtime/StateTree/Source/StateTreeTestSuite/Private/StateTreeTaskStateTest.cpp
+++ b/Engine/Plugins/Runtime/StateTree/Source/StateTreeTestSuite/Private/StateTreeTaskStateTest.cpp
@@ -30,7 +30,7 @@ struct FStateTreeTest_TasksCompletion_AllAny : FStateTreeTestBase
 		{
 			UStateTreeEditorData& EditorData = *Cast<UStateTreeEditorData>(StateTree.EditorData);
 
-			UStateTreeState& Root = EditorData.AddSubTree(FName(TEXT("Tree1StateRoot")));
+			UStateTreeState& Root = EditorData.AddSubTree("Tree1StateRoot");
 			RootState = &Root;
 			Root.TasksCompletion = EStateTreeTaskCompletionType::All;
 			for (int32 TaskIndex = 0; TaskIndex < FStateTreeTasksCompletionStatus::MaxNumberOfTasksPerGroup; ++TaskIndex)
@@ -182,7 +182,7 @@ struct FStateTreeTest_TasksCompletion_FailureTasks : FStateTreeTestBase
 		{
 			UStateTreeEditorData& EditorData = *Cast<UStateTreeEditorData>(StateTree.EditorData);
 
-			UStateTreeState& Root = EditorData.AddSubTree(FName(TEXT("Tree1StateRoot")));
+			UStateTreeState& Root = EditorData.AddSubTree("Tree1StateRoot");
 			Root.TasksCompletion = EStateTreeTaskCompletionType::All;
 			for (int32 TaskIndex = 0; TaskIndex < FStateTreeTasksCompletionStatus::MaxNumberOfTasksPerGroup; ++TaskIndex)
 			{
@@ -284,7 +284,7 @@ struct FStateTreeTest_TasksCompletion_StartGlobalFailureTasks : FStateTreeTestBa
 			GlobalTask2Ptr = &GlobalTask2;
 			GlobalTask3Ptr = &GlobalTask3;
 
-			UStateTreeState& Root = EditorData.AddSubTree(FName(TEXT("Tree1StateRoot")));
+			UStateTreeState& Root = EditorData.AddSubTree("Tree1StateRoot");
 			Root.AddTask<FTestTask_Stand>("Tree1StateRootTask1");
 		}
 
@@ -596,6 +596,772 @@ struct FStateTreeTest_TasksCompletion_Status : FStateTreeTestBase
 };
 IMPLEMENT_STATE_TREE_INSTANT_TEST(FStateTreeTest_TasksCompletion_Status, "System.StateTree.TasksCompletion.Status");
 
+// The base class for testing marked bConsideredForCompletion=false
+struct FStateTreeTest_TasksCompletion_ChildStateNotBlockedByNonCompletionTask_Base : FStateTreeTestBase
+{
+	bool RunTest(EStateTreeTaskCompletionType BCompletionType, int32 NumBTasksNotConsidered, bool bUseSingleTaskProperty = false)
+	{
+		check(!bUseSingleTaskProperty || NumBTasksNotConsidered == 1);
+
+		UStateTree& StateTree = NewStateTree();
+		UStateTreeEditorData& EditorData = *Cast<UStateTreeEditorData>(StateTree.EditorData);
+
+		UStateTreeState& Root    = EditorData.AddSubTree("Root");
+		UStateTreeState& StateA  = Root.AddChildState("A");
+		UStateTreeState& StateB  = Root.AddChildState("B");
+		UStateTreeState& StateB1 = StateB.AddChildState("B1");
+		UStateTreeState& StateB2 = StateB1.AddChildState("B2");
+
+		// A: single task completes in 1 tick, then the tree transitions to B.
+		{
+			TStateTreeEditorNode<FTestTask_Stand>& TaskA = StateA.AddTask<FTestTask_Stand>("TaskA");
+			TaskA.GetNode().TicksToCompletion = 1;
+			StateA.AddTransition(EStateTreeTransitionTrigger::OnStateCompleted, EStateTreeTransitionType::GotoState, &StateB);
+		}
+
+		// B: one or more tasks that complete on the first tick but are NOT considered for completion.
+		StateB.TasksCompletion = BCompletionType;
+		if (bUseSingleTaskProperty)
+		{
+			// Set up B's task via the SingleTask field (single-task schema path).
+			StateB.SingleTask.InitializeAs<FTestTask_Stand>(&StateB, "TaskB_0");
+			TStateTreeEditorNode<FTestTask_Stand>& TaskB = static_cast<TStateTreeEditorNode<FTestTask_Stand>&>(StateB.SingleTask);
+			TaskB.GetNode().TicksToCompletion = 1;
+			TaskB.GetNode().bConsideredForCompletion = false;
+		}
+		else
+		{
+			for (int32 i = 0; i < NumBTasksNotConsidered; ++i)
+			{
+				TStateTreeEditorNode<FTestTask_Stand>& TaskB = StateB.AddTask<FTestTask_Stand>(
+					FName(*FString::Printf(TEXT("TaskB_%d"), i)));
+				TaskB.GetNode().TicksToCompletion = 1;
+				TaskB.GetNode().bConsideredForCompletion = false;
+			}
+		}
+
+		// B2: task that requires 2 ticks to complete.
+		{
+			TStateTreeEditorNode<FTestTask_Stand>& TaskB2 = StateB2.AddTask<FTestTask_Stand>("TaskB2");
+			TaskB2.GetNode().TicksToCompletion = 2;
+		}
+
+		{
+			FStateTreeCompilerLog Log;
+			FStateTreeCompiler Compiler(Log);
+			const bool bResult = Compiler.Compile(StateTree);
+			AITEST_TRUE(TEXT("StateTree should compile"), bResult);
+			if (!bResult)
+			{
+				return false;
+			}
+		}
+
+		FStateTreeInstanceData InstanceData;
+		FTestStateTreeExecutionContext Exec(StateTree, StateTree, InstanceData);
+		AITEST_TRUE(TEXT("StateTree should init"), Exec.IsValid());
+
+		// Start: A is the initial state.
+		{
+			const EStateTreeRunStatus Status = Exec.Start();
+			AITEST_EQUAL(TEXT("Start: should be Running"), Status, EStateTreeRunStatus::Running);
+			AITEST_TRUE(TEXT("Start: should be in A"), Exec.ExpectInActiveStates("Root", "A"));
+			Exec.LogClear();
+		}
+
+		// Tick 1: TaskA completes (1 tick), A transitions to B. B/B1/B2 enter but their tasks
+		// are not ticked yet in the same tick as the transition.
+		{
+			const EStateTreeRunStatus Status = Exec.Tick(0.1f);
+			AITEST_EQUAL(TEXT("Tick 1: should be Running after transition to B"), Status, EStateTreeRunStatus::Running);
+			AITEST_TRUE(TEXT("Tick 1: should be in B/B1/B2"), Exec.ExpectInActiveStates("Root", "B", "B1", "B2"));
+			AITEST_FALSE(TEXT("Tick 1: TaskB2 should not have ticked yet"), Exec.Expect("TaskB2", TEXT("Tick")));
+			Exec.LogClear();
+		}
+
+		// Tick 2: First tick in B. B's task(s) run and return Succeeded, but they are not
+		// considered for completion so B must NOT auto-complete. B2's task runs its first tick
+		// (CurrentTick=1 < TicksToCompletion=2) and returns Running.
+		{
+			const EStateTreeRunStatus Status = Exec.Tick(0.1f);
+			AITEST_EQUAL(TEXT("Tick 2: should be Running (B must not auto-complete)"), Status, EStateTreeRunStatus::Running);
+			AITEST_TRUE(TEXT("Tick 2: B must not auto-complete; still in B/B1/B2"), Exec.ExpectInActiveStates("Root", "B", "B1", "B2"));
+			AITEST_TRUE(TEXT("Tick 2: TaskB2 should have ticked"), Exec.Expect("TaskB2", TEXT("Tick")));
+			Exec.LogClear();
+		}
+
+		// Tick 3: TaskB2 second tick, completes. B2 runs to full completion.
+		{
+			Exec.Tick(0.1f);
+			AITEST_TRUE(TEXT("Tick 3: TaskB2 should have completed"), Exec.Expect("TaskB2", TEXT("StateCompleted")));
+			Exec.LogClear();
+		}
+
+		Exec.Stop();
+		return true;
+	}
+};
+
+// --- 1 task via AddTask ---
+
+struct FStateTreeTest_TasksCompletion_ChildStateNotBlockedByNonCompletionTask_1Task_AddTask_All
+	: FStateTreeTest_TasksCompletion_ChildStateNotBlockedByNonCompletionTask_Base
+{
+	virtual bool InstantTest() override
+	{
+		return RunTest(EStateTreeTaskCompletionType::All, 1, /*bUseSingleTaskProperty=*/false);
+	}
+};
+IMPLEMENT_STATE_TREE_INSTANT_TEST(
+	FStateTreeTest_TasksCompletion_ChildStateNotBlockedByNonCompletionTask_1Task_AddTask_All,
+	"System.StateTree.TasksCompletion.ChildStateNotBlockedByNonCompletionTask.1Task.AddTask.All");
+
+struct FStateTreeTest_TasksCompletion_ChildStateNotBlockedByNonCompletionTask_1Task_AddTask_Any
+	: FStateTreeTest_TasksCompletion_ChildStateNotBlockedByNonCompletionTask_Base
+{
+	virtual bool InstantTest() override
+	{
+		return RunTest(EStateTreeTaskCompletionType::Any, 1, /*bUseSingleTaskProperty=*/false);
+	}
+};
+IMPLEMENT_STATE_TREE_INSTANT_TEST(
+	FStateTreeTest_TasksCompletion_ChildStateNotBlockedByNonCompletionTask_1Task_AddTask_Any,
+	"System.StateTree.TasksCompletion.ChildStateNotBlockedByNonCompletionTask.1Task.AddTask.Any");
+
+// --- 1 task via SingleTask property ---
+
+struct FStateTreeTest_TasksCompletion_ChildStateNotBlockedByNonCompletionTask_1Task_SingleTask_All
+	: FStateTreeTest_TasksCompletion_ChildStateNotBlockedByNonCompletionTask_Base
+{
+	virtual bool InstantTest() override
+	{
+		return RunTest(EStateTreeTaskCompletionType::All, 1, /*bUseSingleTaskProperty=*/true);
+	}
+};
+IMPLEMENT_STATE_TREE_INSTANT_TEST(
+	FStateTreeTest_TasksCompletion_ChildStateNotBlockedByNonCompletionTask_1Task_SingleTask_All,
+	"System.StateTree.TasksCompletion.ChildStateNotBlockedByNonCompletionTask.1Task.SingleTask.All");
+
+struct FStateTreeTest_TasksCompletion_ChildStateNotBlockedByNonCompletionTask_1Task_SingleTask_Any
+	: FStateTreeTest_TasksCompletion_ChildStateNotBlockedByNonCompletionTask_Base
+{
+	virtual bool InstantTest() override
+	{
+		return RunTest(EStateTreeTaskCompletionType::Any, 1, /*bUseSingleTaskProperty=*/true);
+	}
+};
+IMPLEMENT_STATE_TREE_INSTANT_TEST(
+	FStateTreeTest_TasksCompletion_ChildStateNotBlockedByNonCompletionTask_1Task_SingleTask_Any,
+	"System.StateTree.TasksCompletion.ChildStateNotBlockedByNonCompletionTask.1Task.SingleTask.Any");
+
+// --- Multiple tasks via AddTask ---
+
+struct FStateTreeTest_TasksCompletion_ChildStateNotBlockedByNonCompletionTask_MultiTask_All
+	: FStateTreeTest_TasksCompletion_ChildStateNotBlockedByNonCompletionTask_Base
+{
+	virtual bool InstantTest() override
+	{
+		return RunTest(EStateTreeTaskCompletionType::All, 2, /*bUseSingleTaskProperty=*/false);
+	}
+};
+IMPLEMENT_STATE_TREE_INSTANT_TEST(
+	FStateTreeTest_TasksCompletion_ChildStateNotBlockedByNonCompletionTask_MultiTask_All,
+	"System.StateTree.TasksCompletion.ChildStateNotBlockedByNonCompletionTask.MultiTask.All");
+
+// MultiTask via AddTask, Any mode (complements MultiTask_All).
+struct FStateTreeTest_TasksCompletion_ChildStateNotBlockedByNonCompletionTask_MultiTask_Any
+	: FStateTreeTest_TasksCompletion_ChildStateNotBlockedByNonCompletionTask_Base
+{
+	virtual bool InstantTest() override
+	{
+		return RunTest(EStateTreeTaskCompletionType::Any, 2, /*bUseSingleTaskProperty=*/false);
+	}
+};
+IMPLEMENT_STATE_TREE_INSTANT_TEST(
+	FStateTreeTest_TasksCompletion_ChildStateNotBlockedByNonCompletionTask_MultiTask_Any,
+	"System.StateTree.TasksCompletion.ChildStateNotBlockedByNonCompletionTask.MultiTask.Any");
+
+// State re-entry resets the task completion bitmask so that B does not auto-complete on the second visit.
+struct FStateTreeTest_TasksCompletion_NonCompletionTaskReEntryResetsStatus : FStateTreeTestBase
+{
+	virtual bool InstantTest() override
+	{
+		UStateTree& StateTree = NewStateTree();
+		UStateTreeEditorData& EditorData = *Cast<UStateTreeEditorData>(StateTree.EditorData);
+
+		UStateTreeState& Root    = EditorData.AddSubTree("Root");
+		UStateTreeState& StateB  = Root.AddChildState("B");
+		UStateTreeState& StateB1 = StateB.AddChildState("B1");
+		UStateTreeState& StateB2 = StateB1.AddChildState("B2");
+
+		// B: ineligible task, loops back to itself on completion.
+		StateB.TasksCompletion = EStateTreeTaskCompletionType::All;
+		{
+			TStateTreeEditorNode<FTestTask_Stand>& TaskB = StateB.AddTask<FTestTask_Stand>("TaskB");
+			TaskB.GetNode().TicksToCompletion = 1;
+			TaskB.GetNode().bConsideredForCompletion = false;
+			StateB.AddTransition(EStateTreeTransitionTrigger::OnStateCompleted, EStateTreeTransitionType::GotoState, &StateB);
+		}
+
+		// B2: task that needs 2 ticks to complete.
+		{
+			TStateTreeEditorNode<FTestTask_Stand>& TaskB2 = StateB2.AddTask<FTestTask_Stand>("TaskB2");
+			TaskB2.GetNode().TicksToCompletion = 2;
+		}
+
+		{
+			FStateTreeCompilerLog Log;
+			FStateTreeCompiler Compiler(Log);
+			const bool bResult = Compiler.Compile(StateTree);
+			AITEST_TRUE(TEXT("StateTree should compile"), bResult);
+			if (!bResult)
+			{
+				return false;
+			}
+		}
+
+		FStateTreeInstanceData InstanceData;
+		FTestStateTreeExecutionContext Exec(StateTree, StateTree, InstanceData);
+		AITEST_TRUE(TEXT("StateTree should init"), Exec.IsValid());
+
+		// Start: B/B1/B2 entered, tasks not yet ticked.
+		{
+			const EStateTreeRunStatus Status = Exec.Start();
+			AITEST_EQUAL(TEXT("Start: should be Running"), Status, EStateTreeRunStatus::Running);
+			AITEST_TRUE(TEXT("Start: should be in B/B1/B2"), Exec.ExpectInActiveStates("Root", "B", "B1", "B2"));
+			Exec.LogClear();
+		}
+
+		// Tick 1: First tick in B. B's ineligible task returns Succeeded (not counted).
+		// B2's task first tick, still running. B must NOT auto-complete.
+		{
+			const EStateTreeRunStatus Status = Exec.Tick(0.1f);
+			AITEST_EQUAL(TEXT("Tick 1: should be Running (no auto-complete on first entry)"), Status, EStateTreeRunStatus::Running);
+			AITEST_TRUE(TEXT("Tick 1: still in B/B1/B2"), Exec.ExpectInActiveStates("Root", "B", "B1", "B2"));
+			AITEST_TRUE(TEXT("Tick 1: TaskB2 ticked"), Exec.Expect("TaskB2", TEXT("Tick")));
+			Exec.LogClear();
+		}
+
+		// Tick 2: B2's task second tick, completes. B completes and transitions back to B (re-entry).
+		// B/B1/B2 are re-entered in the same tick but their tasks are not ticked yet.
+		{
+			const EStateTreeRunStatus Status = Exec.Tick(0.1f);
+			AITEST_EQUAL(TEXT("Tick 2: should be Running after B re-enters itself"), Status, EStateTreeRunStatus::Running);
+			AITEST_TRUE(TEXT("Tick 2: TaskB2 completed on first visit"), Exec.Expect("TaskB2", TEXT("StateCompleted")));
+			AITEST_TRUE(TEXT("Tick 2: still in B/B1/B2 (re-entered)"), Exec.ExpectInActiveStates("Root", "B", "B1", "B2"));
+			Exec.LogClear();
+		}
+
+		// Tick 3: Second visit to B. Bitmask must have been reset. B must NOT auto-complete again.
+		// B2's task ticks for the first time on the second visit (CurrentTick reset to 0 on re-entry).
+		{
+			const EStateTreeRunStatus Status = Exec.Tick(0.1f);
+			AITEST_EQUAL(TEXT("Tick 3: should be Running (no auto-complete on re-entry)"), Status, EStateTreeRunStatus::Running);
+			AITEST_TRUE(TEXT("Tick 3: still in B/B1/B2 — B2 is running on second visit"), Exec.ExpectInActiveStates("Root", "B", "B1", "B2"));
+			AITEST_TRUE(TEXT("Tick 3: TaskB2 ticked again on second visit"), Exec.Expect("TaskB2", TEXT("Tick")));
+			Exec.LogClear();
+		}
+
+		Exec.Stop();
+		return true;
+	}
+};
+IMPLEMENT_STATE_TREE_INSTANT_TEST(
+	FStateTreeTest_TasksCompletion_NonCompletionTaskReEntryResetsStatus,
+	"System.StateTree.TasksCompletion.NonCompletionTaskReEntryResetsStatus");
+
+// The state must not complete when only the ineligible task finishes; it must wait for the eligible task to complete.
+struct FStateTreeTest_TasksCompletion_MixedEligibleTasks_All : FStateTreeTestBase
+{
+	virtual bool InstantTest() override
+	{
+		// Tree:
+		//   Root
+		//     B  (TasksCompletion=All)
+		//       TaskB_Ineligible: bConsideredForCompletion=false, TicksToCompletion=1
+		//       TaskB_Eligible:   bConsideredForCompletion=true,  TicksToCompletion=3
+		UStateTree& StateTree = NewStateTree();
+		UStateTreeEditorData& EditorData = *Cast<UStateTreeEditorData>(StateTree.EditorData);
+
+		UStateTreeState& Root   = EditorData.AddSubTree("Root");
+		UStateTreeState& StateB = Root.AddChildState("B");
+
+		StateB.TasksCompletion = EStateTreeTaskCompletionType::All;
+		{
+			TStateTreeEditorNode<FTestTask_Stand>& TaskIneligible = StateB.AddTask<FTestTask_Stand>("TaskB_Ineligible");
+			TaskIneligible.GetNode().TicksToCompletion = 1;
+			TaskIneligible.GetNode().bConsideredForCompletion = false;
+		}
+		{
+			TStateTreeEditorNode<FTestTask_Stand>& TaskEligible = StateB.AddTask<FTestTask_Stand>("TaskB_Eligible");
+			TaskEligible.GetNode().TicksToCompletion = 3;
+		}
+
+		{
+			FStateTreeCompilerLog Log;
+			FStateTreeCompiler Compiler(Log);
+			const bool bResult = Compiler.Compile(StateTree);
+			AITEST_TRUE(TEXT("StateTree should compile"), bResult);
+			if (!bResult)
+			{
+				return false;
+			}
+		}
+
+		FStateTreeInstanceData InstanceData;
+		FTestStateTreeExecutionContext Exec(StateTree, StateTree, InstanceData);
+		AITEST_TRUE(TEXT("StateTree should init"), Exec.IsValid());
+
+		{
+			const EStateTreeRunStatus Status = Exec.Start();
+			AITEST_EQUAL(TEXT("Start: should be Running"), Status, EStateTreeRunStatus::Running);
+			Exec.LogClear();
+		}
+
+		// Tick 1: TaskB_Ineligible completes (returns Succeeded, not counted).
+		// TaskB_Eligible first tick, still running. B must NOT complete yet.
+		{
+			const EStateTreeRunStatus Status = Exec.Tick(0.1f);
+			AITEST_EQUAL(TEXT("Tick 1: should be Running (ineligible done, eligible still running)"), Status, EStateTreeRunStatus::Running);
+			AITEST_TRUE(TEXT("Tick 1: still in B"), Exec.ExpectInActiveStates("Root", "B"));
+			AITEST_FALSE(TEXT("Tick 1: TaskB_Eligible must not have completed"), Exec.Expect("TaskB_Eligible", TEXT("StateCompleted")));
+			Exec.LogClear();
+		}
+
+		// Tick 2: TaskB_Eligible second tick, still running. B still must not complete.
+		{
+			const EStateTreeRunStatus Status = Exec.Tick(0.1f);
+			AITEST_EQUAL(TEXT("Tick 2: should be Running"), Status, EStateTreeRunStatus::Running);
+			AITEST_TRUE(TEXT("Tick 2: still in B"), Exec.ExpectInActiveStates("Root", "B"));
+			Exec.LogClear();
+		}
+
+		// Tick 3: TaskB_Eligible third tick, completes. B completes.
+		{
+			Exec.Tick(0.1f);
+			AITEST_TRUE(TEXT("Tick 3: TaskB_Eligible completed"), Exec.Expect("TaskB_Eligible", TEXT("StateCompleted")));
+			Exec.LogClear();
+		}
+
+		Exec.Stop();
+		return true;
+	}
+};
+IMPLEMENT_STATE_TREE_INSTANT_TEST(
+	FStateTreeTest_TasksCompletion_MixedEligibleTasks_All,
+	"System.StateTree.TasksCompletion.MixedEligibleTasks.All");
+
+// An ineligible task returning Failed from EnterState must not prevent the state from
+// entering, and eligible tasks must still drive completion normally.
+struct FStateTreeTest_TasksCompletion_IneligibleTaskEnterStateFail : FStateTreeTestBase
+{
+	virtual bool InstantTest() override
+	{
+		//Tree:
+		//   Root
+		//     B  (TasksCompletion=All)
+		//       TaskB_Ineligible: bConsideredForCompletion=false, EnterStateResult=Failed
+		//       TaskB_Eligible:   bConsideredForCompletion=true,  TicksToCompletion=2
+		UStateTree& StateTree = NewStateTree();
+		UStateTreeEditorData& EditorData = *Cast<UStateTreeEditorData>(StateTree.EditorData);
+
+		UStateTreeState& Root   = EditorData.AddSubTree("Root");
+		UStateTreeState& StateB = Root.AddChildState("B");
+
+		StateB.TasksCompletion = EStateTreeTaskCompletionType::All;
+		{
+			TStateTreeEditorNode<FTestTask_Stand>& TaskIneligible = StateB.AddTask<FTestTask_Stand>("TaskB_Ineligible");
+			TaskIneligible.GetNode().EnterStateResult = EStateTreeRunStatus::Failed;
+			TaskIneligible.GetNode().bConsideredForCompletion = false;
+		}
+		{
+			TStateTreeEditorNode<FTestTask_Stand>& TaskEligible = StateB.AddTask<FTestTask_Stand>("TaskB_Eligible");
+			TaskEligible.GetNode().TicksToCompletion = 2;
+		}
+
+		{
+			FStateTreeCompilerLog Log;
+			FStateTreeCompiler Compiler(Log);
+			const bool bResult = Compiler.Compile(StateTree);
+			AITEST_TRUE(TEXT("StateTree should compile"), bResult);
+			if (!bResult)
+			{
+				return false;
+			}
+		}
+
+		FStateTreeInstanceData InstanceData;
+		FTestStateTreeExecutionContext Exec(StateTree, StateTree, InstanceData);
+		AITEST_TRUE(TEXT("StateTree should init"), Exec.IsValid());
+
+		// Start: B entered. TaskB_Ineligible fails EnterState (does not tick).
+		// TaskB_Eligible enters and ticks normally.
+		{
+			const EStateTreeRunStatus Status = Exec.Start();
+			AITEST_EQUAL(TEXT("Start: should be Running (ineligible EnterState failure is ignored)"), Status, EStateTreeRunStatus::Running);
+			AITEST_TRUE(TEXT("Start: should be in B"), Exec.ExpectInActiveStates("Root", "B"));
+			AITEST_TRUE(TEXT("Start: TaskB_Ineligible entered"), Exec.Expect("TaskB_Ineligible", TEXT("EnterState")));
+			AITEST_TRUE(TEXT("Start: TaskB_Eligible entered"), Exec.Expect("TaskB_Eligible", TEXT("EnterState")));
+			Exec.LogClear();
+		}
+
+		// Tick 1: TaskB_Ineligible does not tick (failed EnterState). TaskB_Eligible first tick.
+		{
+			const EStateTreeRunStatus Status = Exec.Tick(0.1f);
+			AITEST_EQUAL(TEXT("Tick 1: should be Running"), Status, EStateTreeRunStatus::Running);
+			AITEST_TRUE(TEXT("Tick 1: still in B"), Exec.ExpectInActiveStates("Root", "B"));
+			AITEST_FALSE(TEXT("Tick 1: TaskB_Ineligible must not tick after failed EnterState"), Exec.Expect("TaskB_Ineligible", TEXT("Tick")));
+			AITEST_TRUE(TEXT("Tick 1: TaskB_Eligible ticked"), Exec.Expect("TaskB_Eligible", TEXT("Tick")));
+			Exec.LogClear();
+		}
+
+		// Tick 2: TaskB_Eligible second tick, completes. B completes.
+		{
+			Exec.Tick(0.1f);
+			AITEST_TRUE(TEXT("Tick 2: TaskB_Eligible completed"), Exec.Expect("TaskB_Eligible", TEXT("StateCompleted")));
+			Exec.LogClear();
+		}
+
+		Exec.Stop();
+		return true;
+	}
+};
+IMPLEMENT_STATE_TREE_INSTANT_TEST(
+	FStateTreeTest_TasksCompletion_IneligibleTaskEnterStateFail,
+	"System.StateTree.TasksCompletion.IneligibleTaskEnterStateFail");
+
+// With TasksCompletion=Any, when two tasks both complete in the same tick — one with
+// Succeeded and one with Failed — Failed must win (priority: Failed > Succeeded).
+struct FStateTreeTest_TasksCompletion_Any_FailPriorityOverSucceeded : FStateTreeTestBase
+{
+	virtual bool InstantTest() override
+	{
+		//Tree:
+		//   Root
+		//     B  (TasksCompletion=Any)
+		//       TaskB_Succeed: TicksToCompletion=2, TickCompletionResult=Succeeded
+		//       TaskB_Fail:    TicksToCompletion=2, TickCompletionResult=Failed
+
+		// Both tasks complete on tick 2 in the same pass. Because all tasks are still processed after
+		// the first Any-completion (see AllAny test), TaskB_Fail gets to update the status to Failed
+		// which has higher priority than the already-set Succeeded.
+		UStateTree& StateTree = NewStateTree();
+		UStateTreeEditorData& EditorData = *Cast<UStateTreeEditorData>(StateTree.EditorData);
+
+		UStateTreeState& Root   = EditorData.AddSubTree("Root");
+		UStateTreeState& StateB = Root.AddChildState("B");
+
+		StateB.TasksCompletion = EStateTreeTaskCompletionType::Any;
+		{
+			TStateTreeEditorNode<FTestTask_Stand>& TaskSucceed = StateB.AddTask<FTestTask_Stand>("TaskB_Succeed");
+			TaskSucceed.GetNode().TicksToCompletion = 2;
+			TaskSucceed.GetNode().TickCompletionResult = EStateTreeRunStatus::Succeeded;
+		}
+		{
+			TStateTreeEditorNode<FTestTask_Stand>& TaskFail = StateB.AddTask<FTestTask_Stand>("TaskB_Fail");
+			TaskFail.GetNode().TicksToCompletion = 2;
+			TaskFail.GetNode().TickCompletionResult = EStateTreeRunStatus::Failed;
+		}
+
+		{
+			FStateTreeCompilerLog Log;
+			FStateTreeCompiler Compiler(Log);
+			const bool bResult = Compiler.Compile(StateTree);
+			AITEST_TRUE(TEXT("StateTree should compile"), bResult);
+			if (!bResult)
+			{
+				return false;
+			}
+		}
+
+		FStateTreeInstanceData InstanceData;
+		FTestStateTreeExecutionContext Exec(StateTree, StateTree, InstanceData);
+		AITEST_TRUE(TEXT("StateTree should init"), Exec.IsValid());
+
+		{
+			const EStateTreeRunStatus Status = Exec.Start();
+			AITEST_EQUAL(TEXT("Start: should be Running"), Status, EStateTreeRunStatus::Running);
+			Exec.LogClear();
+		}
+
+		// Tick 1: Both tasks running (CurrentTick=1, < TicksToCompletion=2).
+		{
+			const EStateTreeRunStatus Status = Exec.Tick(0.1f);
+			AITEST_EQUAL(TEXT("Tick 1: should be Running"), Status, EStateTreeRunStatus::Running);
+			AITEST_TRUE(TEXT("Tick 1: still in B"), Exec.ExpectInActiveStates("Root", "B"));
+			Exec.LogClear();
+		}
+
+		// Tick 2: Both tasks complete in the same tick.
+		// TaskB_Succeed (index 0) fires first with Succeeded; TaskB_Fail (index 1) fires with Failed.
+		// Failed has higher priority than Succeeded, so the state must complete with Failed.
+		{
+			Exec.Tick(0.1f);
+			AITEST_TRUE(TEXT("Tick 2: TaskB_Succeed ticked"), Exec.Expect("TaskB_Succeed", TEXT("Tick")));
+			AITEST_TRUE(TEXT("Tick 2: TaskB_Fail ticked"), Exec.Expect("TaskB_Fail", TEXT("Tick")));
+			// Both tasks received StateCompleted because they share the same completion round.
+			AITEST_TRUE(TEXT("Tick 2: TaskB_Succeed received StateCompleted"), Exec.Expect("TaskB_Succeed", TEXT("StateCompleted")));
+			AITEST_TRUE(TEXT("Tick 2: TaskB_Fail received StateCompleted"), Exec.Expect("TaskB_Fail", TEXT("StateCompleted")));
+			// The state exited with Failed status because Failed > Succeeded in priority.
+			AITEST_TRUE(TEXT("Tick 2: B exited as Failed"), Exec.Expect("TaskB_Succeed", TEXT("ExitFailed")));
+			Exec.LogClear();
+		}
+
+		Exec.Stop();
+		return true;
+	}
+};
+IMPLEMENT_STATE_TREE_INSTANT_TEST(
+	FStateTreeTest_TasksCompletion_Any_FailPriorityOverSucceeded,
+	"System.StateTree.TasksCompletion.Any.FailPriorityOverSucceeded");
+
+// Disabled task in the middle of the task list
+// The compiler excludes disabled tasks from CompletionMask but still allocates a bit slot
+// for them. CompletionMask = 0b101 (bits 0 and 2 only; bit 1 for disabled TaskB_1 is 0).
+// StateB must complete only when BOTH TaskB_0 (tick 2) and TaskB_2 (tick 3) complete.
+struct FStateTreeTest_TasksCompletion_DisabledTaskInMiddle : FStateTreeTestBase
+{
+	virtual bool InstantTest() override
+	{
+		//Tree: Root → StateB(3 tasks, All)
+		//   TaskB_0: bTaskEnabled=true,  TicksToCompletion=2
+		//   TaskB_1: bTaskEnabled=false  — never entered, never ticked
+		//   TaskB_2: bTaskEnabled=true,  TicksToCompletion=3
+		UStateTree& StateTree = NewStateTree();
+		UStateTreeEditorData& EditorData = *Cast<UStateTreeEditorData>(StateTree.EditorData);
+
+		UStateTreeState& Root   = EditorData.AddSubTree("Root");
+		UStateTreeState& StateB = Root.AddChildState("B");
+
+		StateB.TasksCompletion = EStateTreeTaskCompletionType::All;
+
+		// Bit 0: enabled, completes in 2 ticks — set in CompletionMask.
+		{
+			TStateTreeEditorNode<FTestTask_Stand>& Task = StateB.AddTask<FTestTask_Stand>("TaskB_0");
+			Task.GetNode().TicksToCompletion = 2;
+		}
+		// Bit 1: disabled — occupies the bit slot but excluded from CompletionMask.
+		{
+			TStateTreeEditorNode<FTestTask_Stand>& Task = StateB.AddTask<FTestTask_Stand>("TaskB_1");
+			Task.GetNode().bTaskEnabled = false;
+		}
+		// Bit 2: enabled, completes in 3 ticks — set in CompletionMask.
+		{
+			TStateTreeEditorNode<FTestTask_Stand>& Task = StateB.AddTask<FTestTask_Stand>("TaskB_2");
+			Task.GetNode().TicksToCompletion = 3;
+		}
+
+		{
+			FStateTreeCompilerLog Log;
+			FStateTreeCompiler Compiler(Log);
+			const bool bResult = Compiler.Compile(StateTree);
+			AITEST_TRUE(TEXT("StateTree should compile"), bResult);
+			if (!bResult)
+			{
+				return false;
+			}
+		}
+
+		FStateTreeInstanceData InstanceData;
+		FTestStateTreeExecutionContext Exec(StateTree, StateTree, InstanceData);
+		AITEST_TRUE(TEXT("StateTree should init"), Exec.IsValid());
+
+		// Start: TaskB_0 and TaskB_2 are entered; disabled TaskB_1 is skipped entirely.
+		{
+			const EStateTreeRunStatus Status = Exec.Start();
+			AITEST_EQUAL(TEXT("Start: should be Running"), Status, EStateTreeRunStatus::Running);
+			AITEST_TRUE(TEXT("Start: TaskB_0 should enter"), Exec.Expect("TaskB_0", TEXT("EnterState")));
+			AITEST_FALSE(TEXT("Start: TaskB_1 must not enter (disabled)"), Exec.Expect("TaskB_1", TEXT("EnterState")));
+			AITEST_TRUE(TEXT("Start: TaskB_2 should enter"), Exec.Expect("TaskB_2", TEXT("EnterState")));
+			Exec.LogClear();
+		}
+
+		// Tick 1: TaskB_0 and TaskB_2 tick (1/2 and 1/3). TaskB_1 is never ticked. B still running.
+		{
+			const EStateTreeRunStatus Status = Exec.Tick(0.1f);
+			AITEST_EQUAL(TEXT("Tick 1: should be Running"), Status, EStateTreeRunStatus::Running);
+			AITEST_TRUE(TEXT("Tick 1: TaskB_0 ticked"), Exec.Expect("TaskB_0", TEXT("Tick")));
+			AITEST_FALSE(TEXT("Tick 1: TaskB_1 must not tick (disabled)"), Exec.Expect("TaskB_1", TEXT("Tick")));
+			AITEST_TRUE(TEXT("Tick 1: TaskB_2 ticked"), Exec.Expect("TaskB_2", TEXT("Tick")));
+			Exec.LogClear();
+		}
+
+		// Tick 2: TaskB_0 completes (2/2). TaskB_2 still at 2/3. TaskB_1 still skipped.
+		// Only bit 0 of CompletionMask (0b101) is set; bit 2 is not → B must NOT complete yet.
+		{
+			const EStateTreeRunStatus Status = Exec.Tick(0.1f);
+			AITEST_EQUAL(TEXT("Tick 2: should be Running (TaskB_2 not yet done)"), Status, EStateTreeRunStatus::Running);
+			AITEST_FALSE(TEXT("Tick 2: TaskB_1 must not tick (disabled)"), Exec.Expect("TaskB_1", TEXT("Tick")));
+			AITEST_TRUE(TEXT("Tick 2: still in B"), Exec.ExpectInActiveStates("Root", "B"));
+			Exec.LogClear();
+		}
+
+		// Tick 3: TaskB_2 completes (3/3). Both bits 0 and 2 of CompletionMask are set → B completes.
+		{
+			Exec.Tick(0.1f);
+			AITEST_FALSE(TEXT("Tick 3: TaskB_1 must not tick (disabled)"), Exec.Expect("TaskB_1", TEXT("Tick")));
+			AITEST_TRUE(TEXT("Tick 3: TaskB_2 completed"), Exec.Expect("TaskB_2", TEXT("StateCompleted")));
+			Exec.LogClear();
+		}
+
+		Exec.Stop();
+		return true;
+	}
+};
+IMPLEMENT_STATE_TREE_INSTANT_TEST(
+	FStateTreeTest_TasksCompletion_DisabledTaskInMiddle,
+	"System.StateTree.TasksCompletion.DisabledTaskInMiddle");
+
+// Empty state (no tasks at all)
+// When all active states have zero enabled tasks, the runtime forces the bottom active state
+// (StateB) to complete with Succeeded on the first tick.
+struct FStateTreeTest_TasksCompletion_EmptyState : FStateTreeTestBase
+{
+	virtual bool InstantTest() override
+	{
+		//Tree: Root → StateB(no tasks)
+		UStateTree& StateTree = NewStateTree();
+		UStateTreeEditorData& EditorData = *Cast<UStateTreeEditorData>(StateTree.EditorData);
+
+		UStateTreeState& Root = EditorData.AddSubTree("Root");
+		UStateTreeState& StateB = Root.AddChildState("B");
+		// No tasks added — StateB is an empty leaf state.
+		StateB.AddTransition(EStateTreeTransitionTrigger::OnStateCompleted, EStateTreeTransitionType::Succeeded);
+
+		{
+			FStateTreeCompilerLog Log;
+			FStateTreeCompiler Compiler(Log);
+			const bool bResult = Compiler.Compile(StateTree);
+			AITEST_TRUE(TEXT("StateTree should compile"), bResult);
+			if (!bResult)
+			{
+				return false;
+			}
+		}
+
+		FStateTreeInstanceData InstanceData;
+		FTestStateTreeExecutionContext Exec(StateTree, StateTree, InstanceData);
+		AITEST_TRUE(TEXT("StateTree should init"), Exec.IsValid());
+
+		{
+			const EStateTreeRunStatus Status = Exec.Start();
+			AITEST_EQUAL(TEXT("Start: should be Running"), Status, EStateTreeRunStatus::Running);
+			AITEST_TRUE(TEXT("Start: should be in B"), Exec.ExpectInActiveStates("Root", "B"));
+			Exec.LogClear();
+		}
+
+		// Tick 1: No enabled tasks anywhere (NumTotalEnabledTasks == 0).
+		// The runtime forces the bottom active state (B) to Succeeded.
+		{
+			const EStateTreeRunStatus Status = Exec.Tick(0.1f);
+			AITEST_EQUAL(TEXT("Tick 1: empty state must complete with Succeeded"), Status, EStateTreeRunStatus::Succeeded);
+			Exec.LogClear();
+		}
+
+		Exec.Stop();
+		return true;
+	}
+};
+IMPLEMENT_STATE_TREE_INSTANT_TEST(
+	FStateTreeTest_TasksCompletion_EmptyState,
+	"System.StateTree.TasksCompletion.EmptyState");
+
+// Completion mask buffer boundary crossing.
+// Bit layout at compile time (no global tasks; global-frame synthetic bit = bit 0):
+//   GlobalTaskEndBit = 1  (global-frame synthetic bit)
+//   Root (no tasks, no parent) starts at bit 1 → synthetic bit at 1, NextBitIndex = 2
+//   ParentState starts at bit 2 → 29 tasks occupy bits 2–30 (buffer 0), NextBitIndex = 31
+//   ChildState starts at bit 31 → needs bits 31–32:
+//     (31+2-1)/32 = 1  ≠  31/32 = 0 → boundary crossing detected
+//     → compiler places ChildState entirely in buffer 1 (bits 0–1)
+//
+// Expected: ChildState's two tasks complete at tick 2, correctly evaluated in buffer 1.
+struct FStateTreeTest_TasksCompletion_MaskBufferBoundary : FStateTreeTestBase
+{
+	virtual bool InstantTest() override
+	{
+		//Tree:
+		//	Root
+		//		ParentState(29 eligible tasks, All, TicksToCompletion = 3)
+		//			ChildState  (2 eligible tasks, All, TicksToCompletion=2)
+		UStateTree& StateTree = NewStateTree();
+		UStateTreeEditorData& EditorData = *Cast<UStateTreeEditorData>(StateTree.EditorData);
+
+		UStateTreeState& Root        = EditorData.AddSubTree("Root");
+		UStateTreeState& ParentState = Root.AddChildState("Parent");
+		UStateTreeState& ChildState  = ParentState.AddChildState("Child");
+
+		ParentState.TasksCompletion = EStateTreeTaskCompletionType::All;
+
+		// 29 eligible tasks occupy bits 2–30 in buffer 0, pushing ChildState's first bit
+		// to position 31 so its two tasks (31–32) cross the 32-bit buffer boundary.
+		for (int32 i = 0; i < 29; ++i)
+		{
+			TStateTreeEditorNode<FTestTask_Stand>& Task = ParentState.AddTask<FTestTask_Stand>(
+				FName(*FString::Printf(TEXT("TaskParent_%d"), i)));
+			Task.GetNode().TicksToCompletion = 3;
+		}
+
+		ChildState.TasksCompletion = EStateTreeTaskCompletionType::All;
+
+		{
+			TStateTreeEditorNode<FTestTask_Stand>& Task = ChildState.AddTask<FTestTask_Stand>("TaskChild_0");
+			Task.GetNode().TicksToCompletion = 2;
+		}
+		{
+			TStateTreeEditorNode<FTestTask_Stand>& Task = ChildState.AddTask<FTestTask_Stand>("TaskChild_1");
+			Task.GetNode().TicksToCompletion = 2;
+		}
+
+		{
+			FStateTreeCompilerLog Log;
+			FStateTreeCompiler Compiler(Log);
+			const bool bResult = Compiler.Compile(StateTree);
+			AITEST_TRUE(TEXT("StateTree should compile despite buffer boundary crossing"), bResult);
+			if (!bResult)
+			{
+				return false;
+			}
+		}
+
+		FStateTreeInstanceData InstanceData;
+		FTestStateTreeExecutionContext Exec(StateTree, StateTree, InstanceData);
+		AITEST_TRUE(TEXT("StateTree should init"), Exec.IsValid());
+
+		{
+			const EStateTreeRunStatus Status = Exec.Start();
+			AITEST_EQUAL(TEXT("Start: should be Running"), Status, EStateTreeRunStatus::Running);
+			AITEST_TRUE(TEXT("Start: should be in Parent/Child"), Exec.ExpectInActiveStates("Root", "Parent", "Child"));
+			Exec.LogClear();
+		}
+
+		// Tick 1: All tasks ticking. ChildState tasks at 1/2 — still Running.
+		{
+			const EStateTreeRunStatus Status = Exec.Tick(0.1f);
+			AITEST_EQUAL(TEXT("Tick 1: should be Running"), Status, EStateTreeRunStatus::Running);
+			AITEST_TRUE(TEXT("Tick 1: TaskChild_0 ticked"), Exec.Expect("TaskChild_0", TEXT("Tick")));
+			AITEST_TRUE(TEXT("Tick 1: TaskChild_1 ticked"), Exec.Expect("TaskChild_1", TEXT("Tick")));
+			AITEST_TRUE(TEXT("Tick 1: still in Parent/Child"), Exec.ExpectInActiveStates("Root", "Parent", "Child"));
+			Exec.LogClear();
+		}
+
+		// Tick 2: ChildState tasks complete (2/2). Both tasks reside in buffer 1; the runtime
+		// correctly reads their status from buffer 1 and exits ChildState with Succeeded.
+		{
+			Exec.Tick(0.1f);
+			AITEST_TRUE(TEXT("Tick 2: TaskChild_0 completed"), Exec.Expect("TaskChild_0", TEXT("StateCompleted")));
+			AITEST_TRUE(TEXT("Tick 2: TaskChild_1 completed"), Exec.Expect("TaskChild_1", TEXT("StateCompleted")));
+			Exec.LogClear();
+		}
+
+		Exec.Stop();
+		return true;
+	}
+};
+IMPLEMENT_STATE_TREE_INSTANT_TEST(
+	FStateTreeTest_TasksCompletion_MaskBufferBoundary,
+	"System.StateTree.TasksCompletion.MaskBufferBoundary");
+
 } // namespace UE::StateTree::Tests
 
 UE_ENABLE_OPTIMIZATION_SHIP
