PR #14428: Moved 'bEnabled' to InstanceData in StateTreeDebugTextTask
#Synced-CL 52302118
#Synced-CL 52302118

diff --git a/Engine/Plugins/Runtime/StateTree/Source/StateTreeModule/Private/Tasks/StateTreeDebugTextTask.cpp b/Engine/Plugins/Runtime/StateTree/Source/StateTreeModule/Private/Tasks/StateTreeDebugTextTask.cpp
--- a/Engine/Plugins/Runtime/StateTree/Source/StateTreeModule/Private/Tasks/StateTreeDebugTextTask.cpp
+++ b/Engine/Plugins/Runtime/StateTree/Source/StateTreeModule/Private/Tasks/StateTreeDebugTextTask.cpp
@@ -12,7 +12,7 @@
 FStateTreeDebugTextTask::FStateTreeDebugTextTask()
 {
 	bShouldCallTick = false;
-	// We do not want to change the ReferenceActor if it's bound.
+	// We do not want to change the ReferenceActor and bEnabled if anything is bound.
 	bShouldCopyBoundPropertiesOnTick = false;
 	bShouldCopyBoundPropertiesOnExitState = false;
 
@@ -24,12 +24,12 @@ FStateTreeDebugTextTask::FStateTreeDebugTextTask()
 
 EStateTreeRunStatus FStateTreeDebugTextTask::EnterState(FStateTreeExecutionContext& Context, const FStateTreeTransitionResult& Transition) const
 {
-	if (!bEnabled)
+	const FInstanceDataType& InstanceData = Context.GetInstanceData(*this);
+
+	if (!InstanceData.bEnabled)
 	{
 		return EStateTreeRunStatus::Running;
 	}
-
-	const FInstanceDataType& InstanceData = Context.GetInstanceData(*this);
 	
 	const UWorld* World = Context.GetWorld();
 	if (World == nullptr && InstanceData.ReferenceActor != nullptr)
@@ -59,12 +59,12 @@ EStateTreeRunStatus FStateTreeDebugTextTask::EnterState(FStateTreeExecutionConte
 
 void FStateTreeDebugTextTask::ExitState(FStateTreeExecutionContext& Context, const FStateTreeTransitionResult& Transition) const
 {
-	if (!bEnabled)
+	const FInstanceDataType& InstanceData = Context.GetInstanceData(*this);
+
+	if (!InstanceData.bEnabled)
 	{
 		return;
 	}
-
-	const FInstanceDataType& InstanceData = Context.GetInstanceData(*this);
 	
 	const UWorld* World = Context.GetWorld();
 	if (World == nullptr && InstanceData.ReferenceActor != nullptr)
@@ -84,6 +84,25 @@ void FStateTreeDebugTextTask::ExitState(FStateTreeExecutionContext& Context, con
 }
 
 #if WITH_EDITOR
+void FStateTreeDebugTextTask::PostLoad(FStateTreeDataView InstanceDataView)
+{
+	static constexpr int32 EnabledMovedToInstanceData_Version = 1;
+	static constexpr int32 CurrentVersion = EnabledMovedToInstanceData_Version;
+
+	if (DeprecatedVersion < EnabledMovedToInstanceData_Version)
+	{
+		if (FInstanceDataType* InstanceData = InstanceDataView.GetMutablePtr<FInstanceDataType>())
+		{
+			InstanceData->bEnabled = bEnabled_DEPRECATED;
+			DeprecatedVersion = CurrentVersion;
+		}
+		else
+		{
+			UE_LOG(LogStateTree, Error, TEXT("FStateTreeDebugTextTask::PostLoad: Invalid or unexpected instance data type."));
+		}
+	}
+}
+
 FText FStateTreeDebugTextTask::GetDescription(const FGuid& ID, FStateTreeDataView InstanceDataView, const IStateTreeBindingLookup& BindingLookup, EStateTreeNodeFormatting Formatting) const
 {
 	// Note that FStateTreeDebugTextTaskInstanceData::BindableText is not added to the formatted string
diff --git a/Engine/Plugins/Runtime/StateTree/Source/StateTreeModule/Private/Tasks/StateTreeDebugTextTask.h b/Engine/Plugins/Runtime/StateTree/Source/StateTreeModule/Private/Tasks/StateTreeDebugTextTask.h
--- a/Engine/Plugins/Runtime/StateTree/Source/StateTreeModule/Private/Tasks/StateTreeDebugTextTask.h
+++ b/Engine/Plugins/Runtime/StateTree/Source/StateTreeModule/Private/Tasks/StateTreeDebugTextTask.h
@@ -15,6 +15,9 @@ struct FStateTreeDebugTextTaskInstanceData
 {
 	GENERATED_BODY()
 
+	UPROPERTY(EditAnywhere, Category = "Parameter")
+	bool bEnabled = true;
+
 	/** Optional actor where to draw the text at. */
 	UPROPERTY(EditAnywhere, Category = "Input", meta=(Optional))
 	TObjectPtr<AActor> ReferenceActor = nullptr;
@@ -40,6 +43,7 @@ struct FStateTreeDebugTextTask : public FStateTreeTaskCommonBase
 	UE_API virtual EStateTreeRunStatus EnterState(FStateTreeExecutionContext& Context, const FStateTreeTransitionResult& Transition) const override;
 	UE_API virtual void ExitState(FStateTreeExecutionContext& Context, const FStateTreeTransitionResult& Transition) const override;
 #if WITH_EDITOR
+	UE_API virtual void PostLoad(FStateTreeDataView InstanceDataView) override;
 	UE_API virtual FText GetDescription(const FGuid& ID, FStateTreeDataView InstanceDataView, const IStateTreeBindingLookup& BindingLookup, EStateTreeNodeFormatting Formatting = EStateTreeNodeFormatting::Text) const override;
 	virtual FName GetIconName() const override
 	{
@@ -63,8 +67,13 @@ struct FStateTreeDebugTextTask : public FStateTreeTaskCommonBase
 	UPROPERTY(EditAnywhere, Category = "Parameter")
 	FVector Offset = FVector::ZeroVector;
 
-	UPROPERTY(EditAnywhere, Category = "Parameter")
-	bool bEnabled = true;
+#if WITH_EDITORONLY_DATA
+	UPROPERTY()
+	bool bEnabled_DEPRECATED = true;
+
+	UPROPERTY()
+	int32 DeprecatedVersion = 0;
+#endif
 };
 
 #undef UE_API
