[StateTree] Fix unittest after deprecation. Introduced with CL 52315161.
#rnx
#Synced-CL 52420744

diff --git a/Engine/Plugins/Runtime/StateTree/Source/StateTreeTestSuite/Private/StateTreeBindingTest.cpp b/Engine/Plugins/Runtime/StateTree/Source/StateTreeTestSuite/Private/StateTreeBindingTest.cpp
--- a/Engine/Plugins/Runtime/StateTree/Source/StateTreeTestSuite/Private/StateTreeBindingTest.cpp
+++ b/Engine/Plugins/Runtime/StateTree/Source/StateTreeTestSuite/Private/StateTreeBindingTest.cpp
@@ -255,11 +255,9 @@ struct FStateTreeTest_BindingsCompiler_WithObjectPtr : FStateTreeTestBase
 			PropertyBindings.Add(MakeBinding(SourceADesc.ID, TEXT("A"), TargetDesc.ID, TEXT("A")));
 			PropertyBindings.Add(MakeBinding(SourceADesc.ID, TEXT("A"), TargetDesc.ID, TEXT("B")));
 			PropertyBindings.Add(MakeBinding(SourceADesc.ID, TEXT("A"), TargetDesc.ID, TEXT("C")));
-			PropertyBindings.Add(MakeBinding(SourceADesc.ID, TEXT("A"), TargetDesc.ID, TEXT("D")));
 			PropertyBindings.Add(MakeBinding(SourceADesc.ID, TEXT("A"), TargetDesc.ID, TEXT("ArrayA[1]")));
 			PropertyBindings.Add(MakeBinding(SourceADesc.ID, TEXT("A"), TargetDesc.ID, TEXT("ArrayB[1]")));
 			PropertyBindings.Add(MakeBinding(SourceADesc.ID, TEXT("A"), TargetDesc.ID, TEXT("ArrayC[1]")));
-			PropertyBindings.Add(MakeBinding(SourceADesc.ID, TEXT("A"), TargetDesc.ID, TEXT("ArrayD[1]")));
 			if (!CompileBatch(PropertyBindings, CopyBatchIndex1))
 			{
 				return false;
@@ -272,11 +270,9 @@ struct FStateTreeTest_BindingsCompiler_WithObjectPtr : FStateTreeTestBase
 			PropertyBindings.Add(MakeBinding(SourceADesc.ID, TEXT("ArrayA[2]"), TargetDesc.ID, TEXT("A")));
 			PropertyBindings.Add(MakeBinding(SourceADesc.ID, TEXT("ArrayA[2]"), TargetDesc.ID, TEXT("B")));
 			PropertyBindings.Add(MakeBinding(SourceADesc.ID, TEXT("ArrayA[2]"), TargetDesc.ID, TEXT("C")));
-			PropertyBindings.Add(MakeBinding(SourceADesc.ID, TEXT("ArrayA[2]"), TargetDesc.ID, TEXT("D")));
 			PropertyBindings.Add(MakeBinding(SourceADesc.ID, TEXT("ArrayA[2]"), TargetDesc.ID, TEXT("ArrayA[1]")));
 			PropertyBindings.Add(MakeBinding(SourceADesc.ID, TEXT("ArrayA[2]"), TargetDesc.ID, TEXT("ArrayB[1]")));
 			PropertyBindings.Add(MakeBinding(SourceADesc.ID, TEXT("ArrayA[2]"), TargetDesc.ID, TEXT("ArrayC[1]")));
-			PropertyBindings.Add(MakeBinding(SourceADesc.ID, TEXT("ArrayA[2]"), TargetDesc.ID, TEXT("ArrayD[1]")));
 			if (!CompileBatch(PropertyBindings, CopyBatchIndex2))
 			{
 				return false;
@@ -289,11 +285,9 @@ struct FStateTreeTest_BindingsCompiler_WithObjectPtr : FStateTreeTestBase
 			PropertyBindings.Add(MakeBinding(SourceADesc.ID, TEXT("B"), TargetDesc.ID, TEXT("A")));
 			PropertyBindings.Add(MakeBinding(SourceADesc.ID, TEXT("B"), TargetDesc.ID, TEXT("B")));
 			PropertyBindings.Add(MakeBinding(SourceADesc.ID, TEXT("B"), TargetDesc.ID, TEXT("C")));
-			PropertyBindings.Add(MakeBinding(SourceADesc.ID, TEXT("B"), TargetDesc.ID, TEXT("D")));
 			PropertyBindings.Add(MakeBinding(SourceADesc.ID, TEXT("B"), TargetDesc.ID, TEXT("ArrayA[1]")));
 			PropertyBindings.Add(MakeBinding(SourceADesc.ID, TEXT("B"), TargetDesc.ID, TEXT("ArrayB[1]")));
 			PropertyBindings.Add(MakeBinding(SourceADesc.ID, TEXT("B"), TargetDesc.ID, TEXT("ArrayC[1]")));
-			PropertyBindings.Add(MakeBinding(SourceADesc.ID, TEXT("B"), TargetDesc.ID, TEXT("ArrayD[1]")));
 			if (!CompileBatch(PropertyBindings, CopyBatchIndex3))
 			{
 				return false;
@@ -306,11 +300,9 @@ struct FStateTreeTest_BindingsCompiler_WithObjectPtr : FStateTreeTestBase
 			PropertyBindings.Add(MakeBinding(SourceADesc.ID, TEXT("ArrayB[2]"), TargetDesc.ID, TEXT("A")));
 			PropertyBindings.Add(MakeBinding(SourceADesc.ID, TEXT("ArrayB[2]"), TargetDesc.ID, TEXT("B")));
 			PropertyBindings.Add(MakeBinding(SourceADesc.ID, TEXT("ArrayB[2]"), TargetDesc.ID, TEXT("C")));
-			PropertyBindings.Add(MakeBinding(SourceADesc.ID, TEXT("ArrayB[2]"), TargetDesc.ID, TEXT("D")));
 			PropertyBindings.Add(MakeBinding(SourceADesc.ID, TEXT("ArrayB[2]"), TargetDesc.ID, TEXT("ArrayA[1]")));
 			PropertyBindings.Add(MakeBinding(SourceADesc.ID, TEXT("ArrayB[2]"), TargetDesc.ID, TEXT("ArrayB[1]")));
 			PropertyBindings.Add(MakeBinding(SourceADesc.ID, TEXT("ArrayB[2]"), TargetDesc.ID, TEXT("ArrayC[1]")));
-			PropertyBindings.Add(MakeBinding(SourceADesc.ID, TEXT("ArrayB[2]"), TargetDesc.ID, TEXT("ArrayD[1]")));
 			if (!CompileBatch(PropertyBindings, CopyBatchIndex4))
 			{
 				return false;
@@ -323,11 +315,9 @@ struct FStateTreeTest_BindingsCompiler_WithObjectPtr : FStateTreeTestBase
 			PropertyBindings.Add(MakeBinding(SourceADesc.ID, TEXT("C"), TargetDesc.ID, TEXT("A")));
 			PropertyBindings.Add(MakeBinding(SourceADesc.ID, TEXT("C"), TargetDesc.ID, TEXT("B")));
 			PropertyBindings.Add(MakeBinding(SourceADesc.ID, TEXT("C"), TargetDesc.ID, TEXT("C")));
-			PropertyBindings.Add(MakeBinding(SourceADesc.ID, TEXT("C"), TargetDesc.ID, TEXT("D")));
 			PropertyBindings.Add(MakeBinding(SourceADesc.ID, TEXT("C"), TargetDesc.ID, TEXT("ArrayA[1]")));
 			PropertyBindings.Add(MakeBinding(SourceADesc.ID, TEXT("C"), TargetDesc.ID, TEXT("ArrayB[1]")));
 			PropertyBindings.Add(MakeBinding(SourceADesc.ID, TEXT("C"), TargetDesc.ID, TEXT("ArrayC[1]")));
-			PropertyBindings.Add(MakeBinding(SourceADesc.ID, TEXT("C"), TargetDesc.ID, TEXT("ArrayD[1]")));
 			if (!CompileBatch(PropertyBindings, CopyBatchIndex5))
 			{
 				return false;
@@ -340,50 +330,14 @@ struct FStateTreeTest_BindingsCompiler_WithObjectPtr : FStateTreeTestBase
 			PropertyBindings.Add(MakeBinding(SourceADesc.ID, TEXT("ArrayC[2]"), TargetDesc.ID, TEXT("A")));
 			PropertyBindings.Add(MakeBinding(SourceADesc.ID, TEXT("ArrayC[2]"), TargetDesc.ID, TEXT("B")));
 			PropertyBindings.Add(MakeBinding(SourceADesc.ID, TEXT("ArrayC[2]"), TargetDesc.ID, TEXT("C")));
-			PropertyBindings.Add(MakeBinding(SourceADesc.ID, TEXT("ArrayC[2]"), TargetDesc.ID, TEXT("D")));
 			PropertyBindings.Add(MakeBinding(SourceADesc.ID, TEXT("ArrayC[2]"), TargetDesc.ID, TEXT("ArrayA[1]")));
 			PropertyBindings.Add(MakeBinding(SourceADesc.ID, TEXT("ArrayC[2]"), TargetDesc.ID, TEXT("ArrayB[1]")));
 			PropertyBindings.Add(MakeBinding(SourceADesc.ID, TEXT("ArrayC[2]"), TargetDesc.ID, TEXT("ArrayC[1]")));
-			PropertyBindings.Add(MakeBinding(SourceADesc.ID, TEXT("ArrayC[2]"), TargetDesc.ID, TEXT("ArrayD[1]")));
 			if (!CompileBatch(PropertyBindings, CopyBatchIndex6))
 			{
 				return false;
 			}
 		}
-		// D to ...
-		int32 CopyBatchIndex7 = INDEX_NONE;
-		{
-			TArray<FStateTreePropertyPathBinding> PropertyBindings;
-			PropertyBindings.Add(MakeBinding(SourceADesc.ID, TEXT("D"), TargetDesc.ID, TEXT("A")));
-			PropertyBindings.Add(MakeBinding(SourceADesc.ID, TEXT("D"), TargetDesc.ID, TEXT("B")));
-			PropertyBindings.Add(MakeBinding(SourceADesc.ID, TEXT("D"), TargetDesc.ID, TEXT("C")));
-			PropertyBindings.Add(MakeBinding(SourceADesc.ID, TEXT("D"), TargetDesc.ID, TEXT("D")));
-			PropertyBindings.Add(MakeBinding(SourceADesc.ID, TEXT("D"), TargetDesc.ID, TEXT("ArrayA[1]")));
-			PropertyBindings.Add(MakeBinding(SourceADesc.ID, TEXT("D"), TargetDesc.ID, TEXT("ArrayB[1]")));
-			PropertyBindings.Add(MakeBinding(SourceADesc.ID, TEXT("D"), TargetDesc.ID, TEXT("ArrayC[1]")));
-			PropertyBindings.Add(MakeBinding(SourceADesc.ID, TEXT("D"), TargetDesc.ID, TEXT("ArrayD[1]")));
-			if (!CompileBatch(PropertyBindings, CopyBatchIndex7))
-			{
-				return false;
-			}
-		}
-		// ArrayD[] to ...
-		int32 CopyBatchIndex8 = INDEX_NONE;
-		{
-			TArray<FStateTreePropertyPathBinding> PropertyBindings;
-			PropertyBindings.Add(MakeBinding(SourceADesc.ID, TEXT("ArrayD[2]"), TargetDesc.ID, TEXT("A")));
-			PropertyBindings.Add(MakeBinding(SourceADesc.ID, TEXT("ArrayD[2]"), TargetDesc.ID, TEXT("B")));
-			PropertyBindings.Add(MakeBinding(SourceADesc.ID, TEXT("ArrayD[2]"), TargetDesc.ID, TEXT("C")));
-			PropertyBindings.Add(MakeBinding(SourceADesc.ID, TEXT("ArrayD[2]"), TargetDesc.ID, TEXT("D")));
-			PropertyBindings.Add(MakeBinding(SourceADesc.ID, TEXT("ArrayD[2]"), TargetDesc.ID, TEXT("ArrayA[1]")));
-			PropertyBindings.Add(MakeBinding(SourceADesc.ID, TEXT("ArrayD[2]"), TargetDesc.ID, TEXT("ArrayB[1]")));
-			PropertyBindings.Add(MakeBinding(SourceADesc.ID, TEXT("ArrayD[2]"), TargetDesc.ID, TEXT("ArrayC[1]")));
-			PropertyBindings.Add(MakeBinding(SourceADesc.ID, TEXT("ArrayD[2]"), TargetDesc.ID, TEXT("ArrayD[1]")));
-			if (!CompileBatch(PropertyBindings, CopyBatchIndex8))
-			{
-				return false;
-			}
-		}
 		// ArrayA to Array...
 		int32 CopyBatchIndex9 = INDEX_NONE;
 		{
@@ -391,7 +345,6 @@ struct FStateTreeTest_BindingsCompiler_WithObjectPtr : FStateTreeTestBase
 			PropertyBindings.Add(MakeBinding(SourceADesc.ID, TEXT("ArrayA"), TargetDesc.ID, TEXT("ArrayA")));
 			PropertyBindings.Add(MakeBinding(SourceADesc.ID, TEXT("ArrayA"), TargetDesc.ID, TEXT("ArrayB")));
 			PropertyBindings.Add(MakeBinding(SourceADesc.ID, TEXT("ArrayA"), TargetDesc.ID, TEXT("ArrayC")));
-			PropertyBindings.Add(MakeBinding(SourceADesc.ID, TEXT("ArrayA"), TargetDesc.ID, TEXT("ArrayD")));
 			if (!CompileBatch(PropertyBindings, CopyBatchIndex9))
 			{
 				return false;
@@ -404,7 +357,6 @@ struct FStateTreeTest_BindingsCompiler_WithObjectPtr : FStateTreeTestBase
 			PropertyBindings.Add(MakeBinding(SourceADesc.ID, TEXT("ArrayB"), TargetDesc.ID, TEXT("ArrayA")));
 			PropertyBindings.Add(MakeBinding(SourceADesc.ID, TEXT("ArrayB"), TargetDesc.ID, TEXT("ArrayB")));
 			PropertyBindings.Add(MakeBinding(SourceADesc.ID, TEXT("ArrayB"), TargetDesc.ID, TEXT("ArrayC")));
-			PropertyBindings.Add(MakeBinding(SourceADesc.ID, TEXT("ArrayB"), TargetDesc.ID, TEXT("ArrayD")));
 			if (!CompileBatch(PropertyBindings, CopyBatchIndex10))
 			{
 				return false;
@@ -417,25 +369,11 @@ struct FStateTreeTest_BindingsCompiler_WithObjectPtr : FStateTreeTestBase
 			PropertyBindings.Add(MakeBinding(SourceADesc.ID, TEXT("ArrayC"), TargetDesc.ID, TEXT("ArrayA")));
 			PropertyBindings.Add(MakeBinding(SourceADesc.ID, TEXT("ArrayC"), TargetDesc.ID, TEXT("ArrayB")));
 			PropertyBindings.Add(MakeBinding(SourceADesc.ID, TEXT("ArrayC"), TargetDesc.ID, TEXT("ArrayC")));
-			PropertyBindings.Add(MakeBinding(SourceADesc.ID, TEXT("ArrayC"), TargetDesc.ID, TEXT("ArrayD")));
 			if (!CompileBatch(PropertyBindings, CopyBatchIndex11))
 			{
 				return false;
 			}
 		}
-		// ArrayD to Array...
-		int32 CopyBatchIndex12 = INDEX_NONE;
-		{
-			TArray<FStateTreePropertyPathBinding> PropertyBindings;
-			PropertyBindings.Add(MakeBinding(SourceADesc.ID, TEXT("ArrayD"), TargetDesc.ID, TEXT("ArrayA")));
-			PropertyBindings.Add(MakeBinding(SourceADesc.ID, TEXT("ArrayD"), TargetDesc.ID, TEXT("ArrayB")));
-			PropertyBindings.Add(MakeBinding(SourceADesc.ID, TEXT("ArrayD"), TargetDesc.ID, TEXT("ArrayC")));
-			PropertyBindings.Add(MakeBinding(SourceADesc.ID, TEXT("ArrayD"), TargetDesc.ID, TEXT("ArrayD")));
-			if (!CompileBatch(PropertyBindings, CopyBatchIndex12))
-			{
-				return false;
-			}
-		}
 		BindingCompiler.Finalize();
 
 		const bool bResolveResult = Bindings.ResolvePaths();
