Persist expansion state for parent state nodes & transition nodes in statecentric view

#tests editor
#Synced-CL 52424183

diff --git a/Engine/Plugins/Runtime/StateTree/Source/StateTreeEditorModule/Private/StateCentricView/View/SParentStateExtensionNode.cpp b/Engine/Plugins/Runtime/StateTree/Source/StateTreeEditorModule/Private/StateCentricView/View/SParentStateExtensionNode.cpp
--- a/Engine/Plugins/Runtime/StateTree/Source/StateTreeEditorModule/Private/StateCentricView/View/SParentStateExtensionNode.cpp
+++ b/Engine/Plugins/Runtime/StateTree/Source/StateTreeEditorModule/Private/StateCentricView/View/SParentStateExtensionNode.cpp
@@ -284,6 +284,8 @@ TSharedRef<SWidget> SParentStateExtensionNode::Construct_GenerateNodeBody(const
 
 				return false;
 			});
+			DetailsViewArgs.ViewIdentifier = FName(ViewState->Name.ToString() + "_SParentStateExtensionNode");
+			DetailsViewArgs.ExpansionPersistenceMethod = FDetailsViewArgs::EExpansionPersistenceMethod::TransientPerViewIdentifier;
 
 			TSharedRef<IDetailsView> StateDetailsView = PropertyEditorModule.CreateDetailView(DetailsViewArgs);
 			StateDetailsView->SetExtensionHandler(MakeShared<FStateTreeBindingExtension>());
diff --git a/Engine/Plugins/Runtime/StateTree/Source/StateTreeEditorModule/Private/StateCentricView/View/STransitionNode.cpp b/Engine/Plugins/Runtime/StateTree/Source/StateTreeEditorModule/Private/StateCentricView/View/STransitionNode.cpp
--- a/Engine/Plugins/Runtime/StateTree/Source/StateTreeEditorModule/Private/StateCentricView/View/STransitionNode.cpp
+++ b/Engine/Plugins/Runtime/StateTree/Source/StateTreeEditorModule/Private/StateCentricView/View/STransitionNode.cpp
@@ -456,7 +456,9 @@ TSharedRef<SWidget> STransitionNode::Construct_GenerateNodeBody(const EExtendabl
 
 				return true;
 			});
-
+			DetailsViewArgs.ViewIdentifier = FName(StateOwningTransition->Name.ToString() + "_Transition_" + TransitionID.ToString() + "_STransitionNode");
+			DetailsViewArgs.ExpansionPersistenceMethod = FDetailsViewArgs::EExpansionPersistenceMethod::TransientPerViewIdentifier;
+			
 			TSharedRef<IDetailsView> StateDetailsView = PropertyEditorModule.CreateDetailView(DetailsViewArgs);
 			StateDetailsView->SetExtensionHandler(MakeShared<FStateTreeBindingExtension>());
 			StateDetailsView->SetChildrenCustomizationHandler(MakeShared<FStateTreeBindingsChildrenCustomization>());
