[StateTree] Cache the FStateTreeEditorNode::GetName. When you rename a node, reset the cached name to ensure consistency. This saved ~25ms when the tree is big.

#Synced-CL 52420744

#Synced-CL 52420744

diff --git a/Engine/Plugins/Runtime/StateTree/Source/StateTreeEditorModule/Private/Customizations/StateTreeEditorNodeDetails.cpp b/Engine/Plugins/Runtime/StateTree/Source/StateTreeEditorModule/Private/Customizations/StateTreeEditorNodeDetails.cpp
--- a/Engine/Plugins/Runtime/StateTree/Source/StateTreeEditorModule/Private/Customizations/StateTreeEditorNodeDetails.cpp
+++ b/Engine/Plugins/Runtime/StateTree/Source/StateTreeEditorModule/Private/Customizations/StateTreeEditorNodeDetails.cpp
@@ -1871,10 +1871,7 @@ void FStateTreeEditorNodeDetails::HandleNameCommitted(const FText& NewText, ETex
 				// Set Name
 				if (FStateTreeEditorNode* Node = static_cast<FStateTreeEditorNode*>(Data))
 				{
-					if (FStateTreeNodeBase* BaseNode = Node->GetNode().GetPtr<FStateTreeNodeBase>())
-					{
-						BaseNode->Name = FName(NewName);
-					}
+					Node->SetNodeName(FName(NewName));
 				}
 			}
 
diff --git a/Engine/Plugins/Runtime/StateTree/Source/StateTreeEditorModule/Private/StateTreeEditorNode.cpp b/Engine/Plugins/Runtime/StateTree/Source/StateTreeEditorModule/Private/StateTreeEditorNode.cpp
--- a/Engine/Plugins/Runtime/StateTree/Source/StateTreeEditorModule/Private/StateTreeEditorNode.cpp
+++ b/Engine/Plugins/Runtime/StateTree/Source/StateTreeEditorModule/Private/StateTreeEditorNode.cpp
@@ -184,8 +184,20 @@ bool FStateTreeEditorNode::FixObjectInstances(TSet<UObject*>& SeenObjects, TNotN
 	return bDirty;
 }
 
+void FStateTreeEditorNode::SetNodeName(FName NewName)
+{
+	FStateTreeNodeBase* NodePtr = Node.GetMutablePtr<FStateTreeNodeBase>();
+	if (NodePtr && NodePtr->Name != NewName)
+	{
+		NodePtr->Name = NewName;
+		CachedNodeName.Reset();
+	}
+}
+
 FName FStateTreeEditorNode::GetName() const
 {
+	if (!CachedNodeName.IsSet())
+	{
 		auto MakeName = [this]() -> FName
 			{
 				const UScriptStruct* NodeType = Node.GetScriptStruct();
@@ -212,5 +224,7 @@ FName FStateTreeEditorNode::GetName() const
 				}
 				return FName();
 			};
-	return MakeName();
+		CachedNodeName = MakeName();
+	}
+	return CachedNodeName.GetValue();
 }
diff --git a/Engine/Plugins/Runtime/StateTree/Source/StateTreeEditorModule/Public/StateTreeEditorNode.h b/Engine/Plugins/Runtime/StateTree/Source/StateTreeEditorModule/Public/StateTreeEditorNode.h
--- a/Engine/Plugins/Runtime/StateTree/Source/StateTreeEditorModule/Public/StateTreeEditorNode.h
+++ b/Engine/Plugins/Runtime/StateTree/Source/StateTreeEditorModule/Public/StateTreeEditorNode.h
@@ -82,6 +82,12 @@ public:
 	 */
 	UE_API FName GetName() const;
 
+	/**
+	 * Sets the name of the node. Resets the cached node name if the name changed.
+	 * @param NewName Name to set on the node.
+	 */
+	UE_API void SetNodeName(FName NewName);
+
 	/**
 	 * Get DataView for the node template.
 	 * @return DataView for the node template
