StateTree: Fix tick flag display in state tree view.
#rnx

#Synced-CL 52302118

diff --git a/Engine/Plugins/Runtime/StateTree/Source/StateTreeEditorModule/Private/SStateTreeViewRow.cpp b/Engine/Plugins/Runtime/StateTree/Source/StateTreeEditorModule/Private/SStateTreeViewRow.cpp
--- a/Engine/Plugins/Runtime/StateTree/Source/StateTreeEditorModule/Private/SStateTreeViewRow.cpp
+++ b/Engine/Plugins/Runtime/StateTree/Source/StateTreeEditorModule/Private/SStateTreeViewRow.cpp
@@ -1179,8 +1179,11 @@ void SStateTreeViewRow::MakeFlagsWidget()
 		{
 			const bool bHasEvents = true;
 			const bool bHasBroadcastedDelegates = true;
-			if (RuntimeState->DoesRequestTickTasks(bHasEvents) || RuntimeState->bHasCustomTickRate || RuntimeState->ShouldTickTransitions(bHasEvents, bHasBroadcastedDelegates))
+			if (RuntimeState->bHasCustomTickRate
+				|| RuntimeState->DoesRequestTickTasks(bHasEvents)
+				|| RuntimeState->ShouldTickTransitions(bHasEvents, bHasBroadcastedDelegates))
 			{
+				bool bShow = true;
 				TSharedRef<SHorizontalBox> Box = SNew(SHorizontalBox);
 				if (RuntimeState->bHasCustomTickRate)
 				{
@@ -1194,7 +1197,7 @@ void SStateTreeViewRow::MakeFlagsWidget()
 						.ToolTipText(LOCTEXT("StateCustomTick", "The state has a custom tick rate."))
 					];
 				}
-				else if (RuntimeState->bHasTickTasks)
+				else if (RuntimeState->bCachedRequestTick)
 				{
 					Box->AddSlot()
 					.AutoWidth()
@@ -1206,7 +1209,8 @@ void SStateTreeViewRow::MakeFlagsWidget()
 						.ToolTipText(LOCTEXT("StateNodeTick", "The state contains at least one task that ticks at runtime."))
 					];
 				}
-				else if (RuntimeState->bHasTickTasksOnlyOnEvents)
+				//The task can tick, but it may not affect the scheduled system (see DelayTask as an example).
+				else if (RuntimeState->bCachedRequestTickOnlyOnEvents)
 				{
 					Box->AddSlot()
 					.AutoWidth()
@@ -1218,22 +1222,54 @@ void SStateTreeViewRow::MakeFlagsWidget()
 						.ToolTipText(LOCTEXT("StateNodeTickEvent", "The state contains at least one task that ticks at runtime when there's an event."))
 					];
 				}
-				
-				if (RuntimeState->bHasTransitionTasks)
+				else if (RuntimeState->bHasTransitionTasks)
 				{
 					Box->AddSlot()
 					.AutoWidth()
 					.VAlign(VAlign_Center)
 					[
 						SNew(SImage)
-						.Image(FStateTreeEditorStyle::Get().GetBrush("StateTreeEditor.Transitions"))
+						.Image(FStateTreeEditorStyle::Get().GetBrush("StateTreeEditor.Flags.Tick"))
+						.ColorAndOpacity(IconTint)
+						.ToolTipText(LOCTEXT("StateNodeTickTransitionTask", "The state contains at least one task that ticks at runtime when evaluating transitions."))
+					];
+				}
+				else if (RuntimeState->bHasTickTriggerTransitions)
+				{
+					Box->AddSlot()
+					.AutoWidth()
+					.VAlign(VAlign_Center)
+					[
+						SNew(SImage)
+						.Image(FStateTreeEditorStyle::Get().GetBrush("StateTreeEditor.Flags.Tick"))
+						.ColorAndOpacity(IconTint)
+						.ToolTipText(LOCTEXT("StateNodeTickTransition", "The state contains at least one transition that ticks at runtime."))
+					];
+				}
+				else if (RuntimeState->bHasEventTriggerTransitions)
+				{
+					Box->AddSlot()
+					.AutoWidth()
+					.VAlign(VAlign_Center)
+					[
+						SNew(SImage)
+						.Image(FStateTreeEditorStyle::Get().GetBrush("StateTreeEditor.Flags.TickOnEvent"))
 						.ColorAndOpacity(IconTint)
-						.ToolTipText(LOCTEXT("StateNodeTickTransition", "The state contains at least one task that ticks at runtime when evaluating transitions."))
+						.ToolTipText(LOCTEXT("StateNodeTickTransition", "The state contains at least one transition that ticks at runtime when there's an event."))
 					];
 				}
+				else
+				{
+					// Don't show when bHasDelegateTriggerTransitions. It is not considered ticking.
+					// Don't show when the task ticked but not considered for scheduling.
+					bShow = false;
+				}
 
-				FlagsContainer->SetPadding(FMargin(4.0f));
-				FlagsContainer->SetContent(Box);
+				if (bShow)
+				{
+					FlagsContainer->SetPadding(FMargin(4.0f));
+					FlagsContainer->SetContent(Box);
+				}
 			}
 		}
 	}
