[StateTree] Don't remove null extension during validation. It prevents adding an extension in the editor.

#rnx
#Synced-CL 52420744

diff --git a/Engine/Plugins/Runtime/StateTree/Source/StateTreeModule/Private/StateTree.cpp b/Engine/Plugins/Runtime/StateTree/Source/StateTreeModule/Private/StateTree.cpp
--- a/Engine/Plugins/Runtime/StateTree/Source/StateTreeModule/Private/StateTree.cpp
+++ b/Engine/Plugins/Runtime/StateTree/Source/StateTreeModule/Private/StateTree.cpp
@@ -185,12 +185,9 @@ const UStateTreeExtension* UStateTree::K2_GetExtension(TSubclassOf<UStateTreeExt
 {
 	for (const UStateTreeExtension* Extension : Extensions)
 	{
-		if (ensureMsgf(Extension, TEXT("The extension is invalid. Make sure it's not created in an editor only module.")))
+		if (Extension && Extension->IsA(InExtensionType))
 		{
-			if (Extension->IsA(InExtensionType))
-			{
-				return Extension;
-			}
+			return Extension;
 		}
 	}
 	return nullptr;
@@ -873,18 +870,15 @@ bool UStateTree::Link()
 		}
 
 		// Extensions
-		if (Extensions.Num())
+		for (UStateTreeExtension* Extension : Extensions)
 		{
-			for (UStateTreeExtension* Extension : Extensions)
+			if (Extension)
 			{
-				if (Extension)
+				const bool bLinkSucceeded = Extension->Link(Linker);
+				if (!bLinkSucceeded || Linker.GetStatus() == EStateTreeLinkerStatus::Failed)
 				{
-					const bool bLinkSucceeded = Extension->Link(Linker);
-					if (!bLinkSucceeded || Linker.GetStatus() == EStateTreeLinkerStatus::Failed)
-					{
-						UE_LOG(LogStateTree, Error, TEXT("%s: extension failed to resolve its references."), *GetPathName());
-						return false;
-					}
+					UE_LOG(LogStateTree, Error, TEXT("%s: extension failed to resolve its references."), *GetPathName());
+					return false;
 				}
 			}
 		}
