[StateTree] Fix renamed asset that appeared twice in the picker menu.

#Synced-CL 52420744

diff --git a/Engine/Plugins/Runtime/StateTree/Source/StateTreeEditorModule/Private/StateTreeNodeClassCache.cpp b/Engine/Plugins/Runtime/StateTree/Source/StateTreeEditorModule/Private/StateTreeNodeClassCache.cpp
--- a/Engine/Plugins/Runtime/StateTree/Source/StateTreeEditorModule/Private/StateTreeNodeClassCache.cpp
+++ b/Engine/Plugins/Runtime/StateTree/Source/StateTreeEditorModule/Private/StateTreeNodeClassCache.cpp
@@ -155,7 +155,7 @@ void FStateTreeNodeClassCache::AddRootStruct(TNotNull<UStruct*> RootStruct)
 		return;
 	}
 	
-	RootClasses.Emplace(RootStruct);
+	RootClasses.Emplace(RootStruct, RootClasses.Num());
 
 	InvalidateCache();
 }
@@ -170,11 +170,10 @@ void FStateTreeNodeClassCache::GetStructs(TNotNull<const UStruct*> BaseStruct, T
 		});
 	if (RootClassIndex != INDEX_NONE)
 	{
-		const FRootClassContainer& RootClass = RootClasses[RootClassIndex];
-		
+		FRootClassContainer& RootClass = RootClasses[RootClassIndex];
 		if (!RootClass.bUpdated)
 		{
-			CacheClasses();
+			CacheClasses(RootClass);
 		}
 
 		AvailableClasses.Append(RootClass.ClassData);
@@ -232,17 +231,26 @@ void FStateTreeNodeClassCache::UpdateBlueprintClass(const FAssetData& AssetData)
 			if (StructAssetPath.IsValid())
 			{
 				FRootClassContainer& RootClass = RootClasses[*RootClassIndex];
-				const int32 ClassDataIndex = RootClass.ClassData.IndexOfByPredicate([&StructAssetPath](const TSharedPtr<FStateTreeNodeClassData>& ClassData)
+				int32 ClassDataIndex = RootClass.ClassData.IndexOfByPredicate([&StructAssetPath](const TSharedPtr<FStateTreeNodeClassData>& ClassData)
 					{
 						return ClassData->GetStructPath() == StructAssetPath;
 					});
 
-				if (ClassDataIndex == INDEX_NONE)
+				UObject* AssetOb = AssetData.IsAssetLoaded() ? AssetData.GetAsset() : nullptr;
+				UBlueprint* AssetBP = AssetOb ? Cast<UBlueprint>(AssetOb) : nullptr;
+				UClass* AssetClass = AssetBP ? *AssetBP->GeneratedClass : AssetOb ? AssetOb->GetClass() : nullptr;
+				if (ClassDataIndex == INDEX_NONE && AssetClass != nullptr)
 				{
-					UObject* AssetOb = AssetData.IsAssetLoaded() ? AssetData.GetAsset() : nullptr;
-					UBlueprint* AssetBP = Cast<UBlueprint>(AssetOb);
-					UClass* AssetClass = AssetBP ? *AssetBP->GeneratedClass : AssetOb ? AssetOb->GetClass() : nullptr;
+					// The path can change if the BP is renamed.
+					const TWeakObjectPtr<UStruct> WeakAsset = AssetClass;
+					ClassDataIndex = RootClass.ClassData.IndexOfByPredicate([WeakAsset](const TSharedPtr<FStateTreeNodeClassData>& ClassData)
+						{
+							return ClassData->GetWeakStruct() == WeakAsset;
+						});
+				}
 
+				if (ClassDataIndex == INDEX_NONE)
+				{
 					if (AssetClass)
 					{
 						TSharedRef<FStateTreeNodeClassData> NewData = MakeShared<FStateTreeNodeClassData>(AssetClass);
@@ -254,6 +262,15 @@ void FStateTreeNodeClassCache::UpdateBlueprintClass(const FAssetData& AssetData)
 						RootClass.ClassData.Add(NewData);
 					}
 				}
+				else
+				{
+					// Update existing entry. That can occur if the asset got renamed.
+					if (AssetClass)
+					{
+						TSharedRef<FStateTreeNodeClassData> NewData = MakeShared<FStateTreeNodeClassData>(AssetClass);
+						RootClass.ClassData[ClassDataIndex] = NewData;
+					}
+				}
 			}
 		}
 	}
@@ -263,37 +280,44 @@ void FStateTreeNodeClassCache::CacheClasses()
 {
 	for (TEnumerateRef<FRootClassContainer> RootClass : EnumerateRange(RootClasses))
 	{
-		RootClass->ClassData.Reset();
-		RootClass->bUpdated = true;
+		CacheClasses(*RootClass);
+	}
+}
+
+void FStateTreeNodeClassCache::CacheClasses(FRootClassContainer& RootClass)
+{
+	RootClass.ClassData.Reset();
+	RootClass.bUpdated = true;
 
-		// gather all native classes
-		if (const UClass* Class = Cast<UClass>(RootClass->BaseStruct.Get()))
+	const UStruct* BaseStruct = RootClass.BaseStruct.Get();
+
+	// gather all native classes
+	if (const UClass* Class = Cast<UClass>(BaseStruct))
+	{
+		for (TObjectIterator<UClass> It; It; ++It)
 		{
-			for (TObjectIterator<UClass> It; It; ++It)
+			UClass* TestClass = *It;
+			if (TestClass->HasAnyClassFlags(CLASS_Native) && TestClass->IsChildOf(Class))
 			{
-				UClass* TestClass = *It;
-				if (TestClass->HasAnyClassFlags(CLASS_Native) && TestClass->IsChildOf(Class))
-				{
-					RootClass->ClassData.Add(MakeShareable(new FStateTreeNodeClassData(TestClass)));
-					ClassNameToRootIndex.Add(TestClass->GetPathName(), RootClass.GetIndex());
-				}
+				RootClass.ClassData.Add(MakeShared<FStateTreeNodeClassData>(TestClass));
+				ClassNameToRootIndex.Add(TestClass->GetPathName(), RootClass.Index);
 			}
 		}
-		else if (const UScriptStruct* ScriptStruct = Cast<UScriptStruct>(RootClass->BaseStruct.Get()))
+	}
+	else if (const UScriptStruct* ScriptStruct = Cast<UScriptStruct>(BaseStruct))
+	{
+		for (TObjectIterator<UScriptStruct> It; It; ++It)
 		{
-			for (TObjectIterator<UScriptStruct> It; It; ++It)
+			UScriptStruct* TestStruct = *It;
+			if (TestStruct->IsChildOf(ScriptStruct))
 			{
-				UScriptStruct* TestStruct = *It;
-				if (TestStruct->IsChildOf(ScriptStruct))
-				{
-					RootClass->ClassData.Add(MakeShareable(new FStateTreeNodeClassData(TestStruct)));
-					ClassNameToRootIndex.Add(TestStruct->GetPathName(), RootClass.GetIndex());
-				}
+				RootClass.ClassData.Add(MakeShared<FStateTreeNodeClassData>(TestStruct));
+				ClassNameToRootIndex.Add(TestStruct->GetPathName(), RootClass.Index);
 			}
 		}
 	}
 
-	// gather all blueprints
+	// Gather all blueprints
 	FAssetRegistryModule& AssetRegistryModule = FModuleManager::LoadModuleChecked<FAssetRegistryModule>(TEXT("AssetRegistry"));
 	TArray<FAssetData> BlueprintList;
 
@@ -301,9 +325,9 @@ void FStateTreeNodeClassCache::CacheClasses()
 	Filter.ClassPaths.Add(UBlueprint::StaticClass()->GetClassPathName());
 	AssetRegistryModule.Get().GetAssets(Filter, BlueprintList);
 
-	for (int32 i = 0; i < BlueprintList.Num(); i++)
+	for (const FAssetData& AssetData : BlueprintList)
 	{
-		UpdateBlueprintClass(BlueprintList[i]);
+		UpdateBlueprintClass(AssetData);
 	}
 }
 
diff --git a/Engine/Plugins/Runtime/StateTree/Source/StateTreeEditorModule/Public/StateTreeNodeClassCache.h b/Engine/Plugins/Runtime/StateTree/Source/StateTreeEditorModule/Public/StateTreeNodeClassCache.h
--- a/Engine/Plugins/Runtime/StateTree/Source/StateTreeEditorModule/Public/StateTreeNodeClassCache.h
+++ b/Engine/Plugins/Runtime/StateTree/Source/StateTreeEditorModule/Public/StateTreeNodeClassCache.h
@@ -41,6 +41,11 @@ struct FStateTreeNodeClassData
 
 	UE_API UStruct* GetStruct(bool bSilent = false);
 
+	TWeakObjectPtr<UStruct> GetWeakStruct() const
+	{
+		return CachedStruct;
+	}
+
 	UClass* GetClass(bool bSilent = false)
 	{
 		return Cast<UClass>(GetStruct(bSilent));
@@ -111,21 +116,24 @@ protected:
 	UE_API void OnReloadComplete(EReloadCompleteReason Reason);
 	
 private:
-	UE_API void UpdateBlueprintClass(const FAssetData& AssetData);
-	UE_API void CacheClasses();
-
 	struct FRootClassContainer
 	{
 		FRootClassContainer() = default;
-		FRootClassContainer(TNotNull<UStruct*> InBaseStruct)
+		FRootClassContainer(TNotNull<UStruct*> InBaseStruct, int32 InIndex)
 		 : BaseStruct(InBaseStruct)
+		 , Index(InIndex)
 		{
 		} 
 		TWeakObjectPtr<const UStruct> BaseStruct;
 		TArray<TSharedPtr<FStateTreeNodeClassData>> ClassData;
+		int32 Index = 0;
 		bool bUpdated = false;
 	};
 
+	UE_API void UpdateBlueprintClass(const FAssetData& AssetData);
+	UE_API void CacheClasses();
+	UE_API void CacheClasses(FRootClassContainer& RootClass);
+
 	TArray<FRootClassContainer> RootClasses;
 	TMap<FString, int32> ClassNameToRootIndex;
 };
