[StateTree] Fix the breakpoint for the rewind debugger. The first frame was not captured correctly, and all other frames were lost because of that side effect.

#Synced-CL 52539327

diff --git a/Engine/Plugins/Runtime/StateTree/Source/StateTreeModule/Private/Debugger/StateTreeDebugger.cpp b/Engine/Plugins/Runtime/StateTree/Source/StateTreeModule/Private/Debugger/StateTreeDebugger.cpp
--- a/Engine/Plugins/Runtime/StateTree/Source/StateTreeModule/Private/Debugger/StateTreeDebugger.cpp
+++ b/Engine/Plugins/Runtime/StateTree/Source/StateTreeModule/Private/Debugger/StateTreeDebugger.cpp
@@ -1162,7 +1162,9 @@ void ReadTrace(
 		TraceServices::FAnalysisSessionReadScope SessionReadScope(Session);
 		bValidLastReadFrame = FrameProvider.GetFrameFromTime(TraceFrameType_Game, OutLastTraceReadTime, LastReadFrame);
 	}
-	if (OutLastTraceReadTime == 0 || (bValidLastReadFrame && Frame.Index > LastReadFrame.Index))
+
+	// Last time is invalid (first time processing the frame, or the buffer was cleared), and it's a different frame (scrubbing forward).
+	if (!bValidLastReadFrame || Frame.Index > LastReadFrame.Index)
 	{
 		if (const IStateTreeTraceProvider* Provider = Session.ReadProvider<IStateTreeTraceProvider>(FStateTreeTraceProvider::ProviderName))
 		{
