[StateTree] Remove reference to the Global Parameter feature. The feature is not ready and won't be ready before UE5.8 release.

#rnx
#Synced-CL 52539327

diff --git a/Engine/Plugins/Runtime/StateTree/Source/StateTreeEditorModule/Private/StateTreeCompiler.cpp b/Engine/Plugins/Runtime/StateTree/Source/StateTreeEditorModule/Private/StateTreeCompiler.cpp
--- a/Engine/Plugins/Runtime/StateTree/Source/StateTreeEditorModule/Private/StateTreeCompiler.cpp
+++ b/Engine/Plugins/Runtime/StateTree/Source/StateTreeEditorModule/Private/StateTreeCompiler.cpp
@@ -109,25 +109,17 @@ namespace UE::StateTree::Compiler
 		TEXT("After a StateTree compiles, log the internal content of the StateTree.")
 	);
 
-	FAutoConsoleVariable CVarEnableCompiledUserDefinedStructForGlobalParameters(
-		TEXT("StateTree.Compiler.EnableCompiledUserDefinedStructForGlobalParameters"),
-		false,
-		TEXT("Compiles a User Defined Struct to replace the GlobalParameters PropertyBag.")
-	);
-
 	// Helper archive that checks that the all instanced sub-objects have correct outer. 
 	class FCheckOutersArchive : public FArchiveUObject
 	{
 		using Super = FArchiveUObject;
 		TNotNull<const UStateTree*> StateTree;
-		const UUserDefinedStruct* StateTreeGlobalParameterStruct;
 		TNotNull<const UStateTreeEditorData*> EditorData;
 		FStateTreeCompilerLog& Log;
 	public:
 
-		FCheckOutersArchive(TNotNull<const UStateTree*> InStateTree, TNotNull<const UStateTreeEditorData*> InEditorData, const UUserDefinedStruct* InStateTreeGlobalParameterStruct, FStateTreeCompilerLog& InLog)
+		FCheckOutersArchive(TNotNull<const UStateTree*> InStateTree, TNotNull<const UStateTreeEditorData*> InEditorData, FStateTreeCompilerLog& InLog)
 			: StateTree(InStateTree)
-			, StateTreeGlobalParameterStruct(InStateTreeGlobalParameterStruct)
 			, EditorData(InEditorData)
 			, Log(InLog)
 		{
@@ -156,28 +148,16 @@ namespace UE::StateTree::Compiler
 				{
 					if (Property->HasAnyPropertyFlags(CPF_InstancedReference))
 					{
-						if (Object == StateTreeGlobalParameterStruct)
+						if (!Object->IsInOuter(StateTree))
 						{
-							if (StateTreeGlobalParameterStruct->GetPackage() != StateTree->GetPackage()
-								|| StateTreeGlobalParameterStruct->IsInOuter(StateTree))
-							{
-								Log.Reportf(EMessageSeverity::Error, TEXT("the GlobalParameter struct object %s (%s), should not belong inside the StateTree. This is due to error in the State Tree compiler implementation."),
-									*GetFullNameSafe(Object), *GetFullNameSafe(Object->GetClass()));
-							}
+							Log.Reportf(EMessageSeverity::Error, TEXT("Compiled StateTree contains instanced object %s (%s), which does not belong to the StateTree. This is due to error in the State Tree node implementation."),
+								*GetFullNameSafe(Object), *GetFullNameSafe(Object->GetClass()));
 						}
-						else
-						{
-							if (!Object->IsInOuter(StateTree))
-							{
-								Log.Reportf(EMessageSeverity::Error, TEXT("Compiled StateTree contains instanced object %s (%s), which does not belong to the StateTree. This is due to error in the State Tree node implementation."),
-									*GetFullNameSafe(Object), *GetFullNameSafe(Object->GetClass()));
-							}
 
-							if (Object->IsInOuter(EditorData))
-							{
-								Log.Reportf(EMessageSeverity::Error, TEXT("Compiled StateTree contains instanced object %s (%s), which still belongs to the Editor data. This is due to error in the State Tree node implementation."),
-									*GetFullNameSafe(Object), *GetFullNameSafe(Object->GetClass()));
-							}
+						if (Object->IsInOuter(EditorData))
+						{
+							Log.Reportf(EMessageSeverity::Error, TEXT("Compiled StateTree contains instanced object %s (%s), which still belongs to the Editor data. This is due to error in the State Tree node implementation."),
+								*GetFullNameSafe(Object), *GetFullNameSafe(Object->GetClass()));
 						}
 					}
 				}
@@ -720,7 +700,7 @@ void FStateTreeCompiler::CheckCompiledStateTreeOuters(TNotNull<const UStateTree*
 
 	// Test if there are invalid references in the compiled data.
 	// Serialize is non-const but does not modify the StateTree when used with a saving archive.
-	UE::StateTree::Compiler::FCheckOutersArchive CheckOuters(InStateTree, EditorData, InStateTree->GlobalParametersStructType, InLog);
+	UE::StateTree::Compiler::FCheckOutersArchive CheckOuters(InStateTree, EditorData, InLog);
 	const_cast<UStateTree*>(static_cast<const UStateTree*>(InStateTree))->Serialize(CheckOuters);
 }
 
@@ -754,59 +734,11 @@ bool FStateTreeCompiler::CreateParameters()
 	if (StateTree->ParameterDataType == EStateTreeParameterDataType::GlobalParameterData)
 	{
 		StateTree->Parameters = EditorData->GetRootParametersPropertyBag();
-
-		if (UE::StateTree::Compiler::CVarEnableCompiledUserDefinedStructForGlobalParameters->GetBool())
-		{
-			// Previous version didn't handle the duplicate the struct properly.
-			if (StateTree->GlobalParametersStructType
-				&& StateTree->GlobalParametersStructType->GetPackage() != StateTree->GetPackage())
-			{
-				StateTree->GlobalParametersStructType = nullptr;
-			}
-
-			// Copy parameters from EditorData
-			const FInstancedPropertyBag& PropertyBag = EditorData->GetRootParametersPropertyBag();
-			if (StateTree->GlobalParametersStructType)
-			{
-				StateTree->GlobalParameters = StateTree->GlobalParametersStructType;
-				if (UE::StructUtils::UpdateUserDefinedStructFromDescs(StateTree->GlobalParametersStructType, PropertyBag).HasError())
-				{
-					// A property type doesn't exist, or the BlueprintType was removed.
-					Log.Reportf(EMessageSeverity::Warning, TEXT("Failed to create a global parameters struct compatible with blueprint type."));
-				}
-			}
-			else
-			{
-				// The linker loads and saves in alphabetic order inside the same package. Add a _ to load/save the UDS before the StateTree.
-				const FName UDSName = *FString::Printf(TEXT("_%s_DefaultParameters"), *StateTree->GetName());
-				UE::StructUtils::FCreateUserDefinedStructArgs Args;
-				Args.UserDefinedStructFlags = RF_Transactional | RF_Public;
-				Args.bIsBlueprintType = false;
-				Args.bAddPropertyIDToPropertyName = false;
-				TValueOrError<UUserDefinedStruct*, void> NewUDS = UE::StructUtils::CreateUserDefinedStructFromDescs(StateTree->GetPackage(), PropertyBag, UDSName, Args);
-				StateTree->GlobalParametersStructType = NewUDS.HasValue() ? NewUDS.GetValue() : nullptr;
-				StateTree->GlobalParameters = StateTree->GlobalParametersStructType;
-			}
-
-			if (StateTree->GlobalParameters == nullptr)
-			{
-				Log.Reportf(EMessageSeverity::Warning, TEXT("Failed to create the global parameters type."));
-			}
-
-			// The layout should match because the bindings are expected to work.
-			constexpr bool bCheckPropertiesNames = true;
-			if (!UE::StructUtils::TheSameLayout(PropertyBag, StateTree->GlobalParameters, bCheckPropertiesNames))
-			{
-				Log.Reportf(EMessageSeverity::Warning, TEXT("The generated user defined struct do not matches with the editor version."));
-			}
-		}
 	}
 	else
 	{
 		ensure(StateTree->ParameterDataType == EStateTreeParameterDataType::ExternalGlobalParameterData);
 		StateTree->Parameters = EditorData->GetRootParametersPropertyBag();
-		StateTree->GlobalParameters = nullptr;
-		StateTree->GlobalParametersStructType = nullptr;
 	}
 
 	// Mark parameters as binding source
diff --git a/Engine/Plugins/Runtime/StateTree/Source/StateTreeModule/Private/StateTree.cpp b/Engine/Plugins/Runtime/StateTree/Source/StateTreeModule/Private/StateTree.cpp
--- a/Engine/Plugins/Runtime/StateTree/Source/StateTreeModule/Private/StateTree.cpp
+++ b/Engine/Plugins/Runtime/StateTree/Source/StateTreeModule/Private/StateTree.cpp
@@ -45,15 +45,6 @@ PRAGMA_ENABLE_DEPRECATION_WARNINGS
 const int32 UStateTree::LatestCustomAssetSavedVersion = 1;
 #endif
 
-FConstStructView UStateTree::GetDefaultGlobalParameters() const
-{
-	if (GlobalParameters)
-	{
-		return FConstStructView(GlobalParameters, GlobalParameters->GetDefaultInstance());
-	}
-	return FConstStructView();
-}
-
 bool UStateTree::IsReadyToRun() const
 {
 #if WITH_EDITOR
@@ -299,27 +290,6 @@ void UStateTree::BeginDestroy()
 }
 #endif //WITH_STATETREE_DEBUG
 
-void UStateTree::PostRename(UObject* OldOuter, const FName OldName)
-{
-	Super::PostRename(OldOuter, OldName);
-
-	// Also, rename the GlobalParameters struct to make sure they are in the same package.
-	//The GlobalParameters is RF_Public and needs to follow the StateTree asset.
-	auto RenameStruct = [this, &OldOuter](UUserDefinedStruct* Struct)
-		{
-			if (Struct && Struct->GetOutermost() == OldOuter)
-			{
-				const FString UDSName = FString::Printf(TEXT("_%s_DefaultParameters"), *GetName());
-				Struct->Rename(*UDSName, GetOutermost(), REN_DoNotDirty | REN_SkipGeneratedClasses);
-			}
-		};
-#if WITH_EDITOR
-	RenameStruct(GlobalParametersStructType);
-#else
-	RenameStruct(GlobalParameters);
-#endif
-}
-
 #if WITH_EDITOR
 namespace UE::StateTree::Compiler
 {
@@ -354,8 +324,6 @@ void UStateTree::ResetCompiled()
 	PropertyFunctionEvaluationScopeMemoryRequirements.Reset();
 	Extensions.Reset();
 	Parameters.Reset();
-	GlobalParameters = nullptr;
-	// GlobalParametersStructType is not reset. We want the same type between different compilation.
 	ParameterDataType = EStateTreeParameterDataType::GlobalParameterData;
 	IDToStateMappings.Reset();
 	IDToNodeMappings.Reset();
@@ -508,11 +476,6 @@ void UStateTree::PostLoad()
 {
 	Super::PostLoad();
 
-	if (GlobalParameters)
-	{
-		GlobalParameters->ConditionalPostLoad();
-	}
-
 	for (int32 NodeIndex = 0; NodeIndex < Nodes.Num(); ++NodeIndex)
 	{
 		FStructView NodeView = Nodes[NodeIndex];
@@ -809,13 +772,6 @@ bool UStateTree::Link()
 		return false;
 	}
 
-	// @todo Temporary test. Those are to test that the UDS loaded properly and matches the property bag before removing the Parameters.
-	constexpr bool bCheckPropertiesNames = true;
-	if (GlobalParameters != nullptr && !UE::StructUtils::TheSameLayout(Parameters, GlobalParameters, bCheckPropertiesNames))
-	{
-		UE_LOG(LogStateTree, Warning, TEXT("%s: The global parameters doesn't have the same layout as the PropertyBag."), *GetPathName());
-	}
-
 	// Initialize the instance data default value.
 	// This data will be used to allocate runtime instance on all StateTree users.
 	ResetLinked();
@@ -1179,10 +1135,8 @@ bool UStateTree::PatchBindings()
 		{
 			// Check that the bag in linked state matches.
 			const FInstancedPropertyBag& TargetTreeParameters = State.LinkedAsset->Parameters;
-			const UUserDefinedStruct* TargetTreeGlobalParameters = State.LinkedAsset->GlobalParameters;
 			const FCompactStateTreeParameters& Params = DefaultInstanceData.GetMutableStruct(State.ParameterTemplateIndex.Get()).Get<FCompactStateTreeParameters>();
-			const bool bIsValid = TargetTreeParameters.GetPropertyBagStruct() == Params.GetValue().GetScriptStruct()
-				|| TargetTreeGlobalParameters == Params.GetValue().GetScriptStruct();
+			const bool bIsValid = TargetTreeParameters.GetPropertyBagStruct() == Params.GetValue().GetScriptStruct();
 			if (!bIsValid)
 			{
 				UE_LOG(LogStateTree, Error, TEXT("%s: The parameters on state '%s' does not match the linked asset parameters '%s'. Please recompile the StateTree asset."),
diff --git a/Engine/Plugins/Runtime/StateTree/Source/StateTreeModule/Public/StateTree.h b/Engine/Plugins/Runtime/StateTree/Source/StateTreeModule/Public/StateTree.h
--- a/Engine/Plugins/Runtime/StateTree/Source/StateTreeModule/Public/StateTree.h
+++ b/Engine/Plugins/Runtime/StateTree/Source/StateTreeModule/Public/StateTree.h
@@ -186,9 +186,6 @@ public:
 		return Parameters;
 	}
 
-	/** @return The global parameters of the state tree. Global parameter values can be overridden at runtime by the execution context. */
-	UE_API FConstStructView GetDefaultGlobalParameters() const;
-
 	/**
 	 * @return true if the tree asset can be used at runtime.
 	 * @note the asset can be in the compilation queue. This will make sure it is up to date.
@@ -413,7 +410,6 @@ protected:
 	UE_API virtual void PostInitProperties() override;
 #endif
 
-	UE_API virtual void PostRename(UObject* OldOuter, const FName OldName) override;
 #if WITH_EDITOR
 	UE_API virtual void PreSave(FObjectPreSaveContext Context) override;
 	UE_API virtual void PostDuplicate(EDuplicateMode::Type DuplicateMode) override;
@@ -518,20 +514,6 @@ private:
 	UPROPERTY()
 	FInstancedPropertyBag Parameters;
 
-	/**
-	 * Parameters that could be used for bindings within the Tree.
-	 * Default values are stored within the asset but StateTreeReference can be used to parameterized the tree.
-	 * @see FStateTreeReference
-	 */
-	UPROPERTY()
-	TObjectPtr<UUserDefinedStruct> GlobalParameters;
-
-#if WITH_EDITORONLY_DATA
-	/** Keep a reference between compilation. GlobalParameters is reset between each compilation. */
-	UPROPERTY(Instanced)
-	TObjectPtr<UUserDefinedStruct> GlobalParametersStructType;
-#endif
-
 	//~ Data created during linking.
 	/** List of external data required by the state tree, created during linking. */
 	UPROPERTY(Transient)
