[StateTree] Fix loading Blueprint task with the wrong path. The path should be the UBlueprint and not the generated class.
#Synced-CL 52539327
#Synced-CL 52539327

diff --git a/Engine/Plugins/Runtime/StateTree/Source/StateTreeEditorModule/Private/Customizations/Widgets/SStateTreeNodeTypePicker.cpp b/Engine/Plugins/Runtime/StateTree/Source/StateTreeEditorModule/Private/Customizations/Widgets/SStateTreeNodeTypePicker.cpp
--- a/Engine/Plugins/Runtime/StateTree/Source/StateTreeEditorModule/Private/Customizations/Widgets/SStateTreeNodeTypePicker.cpp
+++ b/Engine/Plugins/Runtime/StateTree/Source/StateTreeEditorModule/Private/Customizations/Widgets/SStateTreeNodeTypePicker.cpp
@@ -224,9 +224,10 @@ void SStateTreeNodeTypePicker::CacheNodeTypes(const UStateTreeSchema* Schema, co
 	// Create tree of node types based on category.
 	RootNode = MakeShared<FStateTreeNodeTypeItem>();
 	
+	constexpr bool bSilent = true;
 	for (const TSharedPtr<FStateTreeNodeClassData>& Data : StructNodes)
 	{
-		if (const UScriptStruct* ScriptStruct = Data->GetScriptStruct())
+		if (const UScriptStruct* ScriptStruct = Data->GetScriptStruct(bSilent))
 		{
 			if (ScriptStruct == BaseScriptStruct)
 			{
@@ -234,11 +235,11 @@ void SStateTreeNodeTypePicker::CacheNodeTypes(const UStateTreeSchema* Schema, co
 			}
 			if (ScriptStruct->HasMetaData(TEXT("Hidden")))
 			{
-				continue;				
+				continue;
 			}
 			if (Schema && !Schema->IsStructAllowed(ScriptStruct))
 			{
-				continue;				
+				continue;
 			}
 
 			AddNode(ScriptStruct);
@@ -247,16 +248,15 @@ void SStateTreeNodeTypePicker::CacheNodeTypes(const UStateTreeSchema* Schema, co
 
 	for (const TSharedPtr<FStateTreeNodeClassData>& Data : ObjectNodes)
 	{
-		if (Data->GetClass() != nullptr)
+		if (const UClass* Class = Data->GetClass(bSilent))
 		{
-			const UClass* Class = Data->GetClass();
 			if (Class == BaseClass)
 			{
 				continue;
 			}
 			if (Class->HasAnyClassFlags(CLASS_Abstract | CLASS_Hidden | CLASS_HideDropDown))
 			{
-				continue;				
+				continue;
 			}
 			if (Class->HasMetaData(TEXT("Hidden")))
 			{
diff --git a/Engine/Plugins/Runtime/StateTree/Source/StateTreeEditorModule/Private/StateTreeNodeClassCache.cpp b/Engine/Plugins/Runtime/StateTree/Source/StateTreeEditorModule/Private/StateTreeNodeClassCache.cpp
--- a/Engine/Plugins/Runtime/StateTree/Source/StateTreeEditorModule/Private/StateTreeNodeClassCache.cpp
+++ b/Engine/Plugins/Runtime/StateTree/Source/StateTreeEditorModule/Private/StateTreeNodeClassCache.cpp
@@ -258,7 +258,7 @@ void FStateTreeNodeClassCache::UpdateBlueprintClass(const FAssetData& AssetData)
 					}
 					else
 					{
-						TSharedRef<FStateTreeNodeClassData> NewData = MakeShared<FStateTreeNodeClassData>(StructAssetPath);
+						TSharedRef<FStateTreeNodeClassData> NewData = MakeShared<FStateTreeNodeClassData>(AssetData.GetSoftObjectPath().GetAssetPath());
 						RootClass.ClassData.Add(NewData);
 					}
 				}
