[StateTree] Fix TotalScore not decremented in weighted random state selection
#Synced-CL 52539327

diff --git a/Engine/Plugins/Runtime/StateTree/Source/StateTreeModule/Private/StateTreeExecutionContext.cpp b/Engine/Plugins/Runtime/StateTree/Source/StateTreeModule/Private/StateTreeExecutionContext.cpp
--- a/Engine/Plugins/Runtime/StateTree/Source/StateTreeModule/Private/StateTreeExecutionContext.cpp
+++ b/Engine/Plugins/Runtime/StateTree/Source/StateTreeModule/Private/StateTreeExecutionContext.cpp
@@ -7934,7 +7934,7 @@ bool FStateTreeExecutionContext::SelectStateInternal_TrySelectChildrenWithHighes
 	while (!NextLevelChildStates.IsEmpty())
 	{
 		//Find one with highest score in the remaining candidates
-		float HighestScore = -std::numeric_limits<float>::infinity();
+		float HighestScore = std::numeric_limits<float>::lowest();
 		int32 ArrayIndexWithHighestScore = INDEX_NONE;
 		for (int32 Index = 0; Index < NextLevelChildStates.Num(); ++Index)
 		{
@@ -8035,6 +8035,7 @@ bool FStateTreeExecutionContext::SelectStateInternal_TrySelectChildrenAtRandomWe
 				// Disqualify the state we failed to enter
 				constexpr EAllowShrinking AllowShrinking = EAllowShrinking::No;
 				NextLevelChildStates.RemoveAtSwap(Index, AllowShrinking);
+				TotalScore -= StateScore;
 				break;
 			}
 		}
