[StateTree] State centric view: right-click add transitions for out transition area.

#tests editor
#Synced-CL 52603298

diff --git a/Engine/Plugins/Runtime/StateTree/Source/StateTreeDeveloper/Internal/Debugger/SCompactTreeDebuggerView.h b/Engine/Plugins/Runtime/StateTree/Source/StateTreeDeveloper/Internal/Debugger/SCompactTreeDebuggerView.h
--- a/Engine/Plugins/Runtime/StateTree/Source/StateTreeDeveloper/Internal/Debugger/SCompactTreeDebuggerView.h
+++ b/Engine/Plugins/Runtime/StateTree/Source/StateTreeDeveloper/Internal/Debugger/SCompactTreeDebuggerView.h
@@ -11,7 +11,7 @@
 
 #include "SCompactTreeDebuggerView.generated.h"
 
-namespace UE::StateTree
+namespace UE::StateTree::Editor
 {
 
 namespace CompactTreeView
@@ -30,13 +30,13 @@ namespace CompactTreeView
 		bool bIsActive = false;
 	};
 } // CompactTreeView
-} // UE::StateTree
+} // UE::StateTree::Editor
 
 #if WITH_STATETREE_TRACE_DEBUGGER
 
 #define UE_API STATETREEDEVELOPER_API
 
-namespace UE::StateTree
+namespace UE::StateTree::Editor
 {
 
 /**
@@ -77,7 +77,7 @@ private:
 	TAttribute<FStateTreeTraceActiveStates> AllActiveStates;
 };
 
-} // UE::StateTree
+} // UE::StateTree::Editor
 
 #undef UE_API
 
diff --git a/Engine/Plugins/Runtime/StateTree/Source/StateTreeDeveloper/Internal/Widgets/SCompactTreeView.h b/Engine/Plugins/Runtime/StateTree/Source/StateTreeDeveloper/Internal/Widgets/SCompactTreeView.h
--- a/Engine/Plugins/Runtime/StateTree/Source/StateTreeDeveloper/Internal/Widgets/SCompactTreeView.h
+++ b/Engine/Plugins/Runtime/StateTree/Source/StateTreeDeveloper/Internal/Widgets/SCompactTreeView.h
@@ -17,7 +17,7 @@ template <typename ItemType> class STreeView;
 
 #define UE_API STATETREEDEVELOPER_API
 
-namespace UE::StateTree
+namespace UE::StateTree::Editor
 {
 namespace CompactTreeView
 {
@@ -32,7 +32,7 @@ struct FStateItemCustomData
 {
 	GENERATED_BODY()
 };
-}
+} // namespace CompactTreeView
 
 
 /**
@@ -80,7 +80,7 @@ protected:
 		FText TooltipText;
 		FGuid StateID = {};
 		const FSlateBrush* Icon = nullptr;
-		TInstancedStruct<UE::StateTree::CompactTreeView::FStateItemCustomData> CustomData;
+		TInstancedStruct<CompactTreeView::FStateItemCustomData> CustomData;
 		TArray<TSharedPtr<FStateItem>> Children;
 		FSlateColor Color = FSlateColor(FLinearColor::White);
 		bool bIsEnabled = true;
@@ -166,6 +166,6 @@ private:
 	static TMap<FObjectKey, FStateExpansionState> StateExpansionStates;
 };
 
-} // UE::StateTree
+} // namespace UE::StateTree::Editor
 
 #undef UE_API
\ No newline at end of file
diff --git a/Engine/Plugins/Runtime/StateTree/Source/StateTreeDeveloper/Private/Debugger/SCompactTreeDebuggerView.cpp b/Engine/Plugins/Runtime/StateTree/Source/StateTreeDeveloper/Private/Debugger/SCompactTreeDebuggerView.cpp
--- a/Engine/Plugins/Runtime/StateTree/Source/StateTreeDeveloper/Private/Debugger/SCompactTreeDebuggerView.cpp
+++ b/Engine/Plugins/Runtime/StateTree/Source/StateTreeDeveloper/Private/Debugger/SCompactTreeDebuggerView.cpp
@@ -9,7 +9,7 @@
 
 #if WITH_STATETREE_TRACE_DEBUGGER
 
-namespace UE::StateTree
+namespace UE::StateTree::Editor
 {
 
 void SCompactTreeDebuggerView::Construct(const FArguments& InArgs, const TNotNull<const UStateTree*> StateTree)
@@ -67,7 +67,7 @@ void SCompactTreeDebuggerView::CacheStateRecursive(
 			StateItem->StateID = StateTree->GetStateIdFromHandle(FStateTreeStateHandle(InStateIdx));
 			StateItem->TooltipText = FText::FromName(State.Name);
 			StateItem->bIsEnabled = State.bEnabled;
-			StateItem->CustomData.GetMutable<UE::StateTree::CompactTreeView::FStateItemDebuggerData>().bIsActive = bIsActiveState;
+			StateItem->CustomData.GetMutable<CompactTreeView::FStateItemDebuggerData>().bIsActive = bIsActiveState;
 			StateItem->Color = FStateTreeStyle::Get().GetSlateColor("StateTree.CompactView.State");
 			StateItem->Icon = FStateTreeStyle::GetBrushForSelectionBehaviorType(State.SelectionBehavior, State.HasChildren(), State.Type);
 
@@ -129,7 +129,7 @@ void SCompactTreeDebuggerView::CacheStateRecursive(
 TSharedRef<SCompactTreeView::FStateItem> SCompactTreeDebuggerView::CreateStateItemInternal() const
 {
 	const TSharedRef<FStateItem> StateItem = MakeShared<FStateItem>();
-	StateItem->CustomData = TInstancedStruct<UE::StateTree::CompactTreeView::FStateItemDebuggerData>::Make({});
+	StateItem->CustomData = TInstancedStruct<CompactTreeView::FStateItemDebuggerData>::Make({});
 	return StateItem;
 }
 
@@ -146,7 +146,7 @@ TSharedRef<SWidget> SCompactTreeDebuggerView::CreateNameWidgetInternal(TSharedPt
 				{
 					if (const FStateItem* StateItem = Item.Get())
 					{
-						if (StateItem->CustomData.Get<UE::StateTree::CompactTreeView::FStateItemDebuggerData>().bIsActive)
+						if (StateItem->CustomData.Get<CompactTreeView::FStateItemDebuggerData>().bIsActive)
 						{
 							return FStateTreeStyle::Get().GetSlateColor("StateTree.Debugger.State.Active");
 						}
@@ -170,5 +170,5 @@ TSharedRef<SWidget> SCompactTreeDebuggerView::CreateNameWidgetInternal(TSharedPt
 		];
 }
 
-} // UE::StateTree
+} // namespace UE::StateTree::Editor
 #endif // WITH_STATETREE_TRACE_DEBUGGER
diff --git a/Engine/Plugins/Runtime/StateTree/Source/StateTreeDeveloper/Private/Widgets/SCompactTreeView.cpp b/Engine/Plugins/Runtime/StateTree/Source/StateTreeDeveloper/Private/Widgets/SCompactTreeView.cpp
--- a/Engine/Plugins/Runtime/StateTree/Source/StateTreeDeveloper/Private/Widgets/SCompactTreeView.cpp
+++ b/Engine/Plugins/Runtime/StateTree/Source/StateTreeDeveloper/Private/Widgets/SCompactTreeView.cpp
@@ -14,7 +14,7 @@
 
 #define LOCTEXT_NAMESPACE "SCompactTreeView"
 
-namespace UE::StateTree
+namespace UE::StateTree::Editor
 {
 TMap<FObjectKey, SCompactTreeView::FStateExpansionState> SCompactTreeView::StateExpansionStates;
 
@@ -430,6 +430,6 @@ void SCompactTreeView::RestoreExpansionState()
 	}
 }
 
-} // UE::StateTree
+} // UE::StateTree::Editor
 
 #undef LOCTEXT_NAMESPACE
diff --git a/Engine/Plugins/Runtime/StateTree/Source/StateTreeEditorModule/Private/Customizations/StateTreeStateLinkDetails.cpp b/Engine/Plugins/Runtime/StateTree/Source/StateTreeEditorModule/Private/Customizations/StateTreeStateLinkDetails.cpp
--- a/Engine/Plugins/Runtime/StateTree/Source/StateTreeEditorModule/Private/Customizations/StateTreeStateLinkDetails.cpp
+++ b/Engine/Plugins/Runtime/StateTree/Source/StateTreeEditorModule/Private/Customizations/StateTreeStateLinkDetails.cpp
@@ -189,7 +189,7 @@ TSharedRef<SWidget> FStateTreeStateLinkDetails::GenerateStatePicker()
 
 	MenuBuilder.BeginSection("States", LOCTEXT("States", "States"));
 
-	TSharedPtr<UE::StateTree::SCompactTreeEditorView> StateView;
+	TSharedPtr<UE::StateTree::Editor::SCompactTreeEditorView> StateView;
 	
 	TSharedRef<SWidget> MenuWidget = 
 		SNew(SBox)
@@ -197,7 +197,7 @@ TSharedRef<SWidget> FStateTreeStateLinkDetails::GenerateStatePicker()
 		.MaxDesiredHeight(400.f)
 		.Padding(2)	
 		[
-			SAssignNew(StateView, UE::StateTree::SCompactTreeEditorView)
+			SAssignNew(StateView, UE::StateTree::Editor::SCompactTreeEditorView)
 			.StateTreeEditorData(WeakEditorData.Get())
 			.SelectionMode(ESelectionMode::Single)
 			.SelectableStatesOnly(true)
diff --git a/Engine/Plugins/Runtime/StateTree/Source/StateTreeEditorModule/Private/Customizations/Widgets/SAddTransitionMenu.cpp b/Engine/Plugins/Runtime/StateTree/Source/StateTreeEditorModule/Private/Customizations/Widgets/SAddTransitionMenu.cpp
new file mode 100644
--- /dev/null
+++ b/Engine/Plugins/Runtime/StateTree/Source/StateTreeEditorModule/Private/Customizations/Widgets/SAddTransitionMenu.cpp
@@ -0,0 +1,153 @@
+// Copyright Epic Games, Inc. All Rights Reserved.
+
+#include "Customizations/Widgets/SAddTransitionMenu.h"
+
+#include "StateTreeEditorData.h"
+#include "StateTreeState.h"
+#include "StateTreeTypes.h"
+#include "StateTreeViewModel.h"
+
+#include "Framework/Application/SlateApplication.h"
+#include "Framework/MultiBox/MultiBoxBuilder.h"
+#include "Styling/AppStyle.h"
+#include "Styling/StyleColors.h"
+#include "Widgets/Images/SImage.h"
+
+#define LOCTEXT_NAMESPACE "SAddTransitionMenu"
+
+namespace UE::StateTree::Editor
+{
+
+void SAddTransitionMenu::Construct(const FArguments& InArgs)
+{
+	WeakOwnerState  = InArgs._OwnerState;
+	WeakViewModel   = InArgs._ViewModel;
+
+	if (TSharedPtr<FStateTreeViewModel> ViewModel = WeakViewModel.Pin())
+	{
+		SCompactTreeEditorView::Construct(
+			SCompactTreeEditorView::FArguments()
+			.SelectionMode(ESelectionMode::Single)
+			.StateTreeEditorData(ViewModel->GetStateTreeEditorData())
+			.SelectableStatesOnly(true)
+			.OnSelectionChanged(SCompactTreeView::FOnSelectionChanged::CreateSP(this, &SAddTransitionMenu::CreateStateTransitionWithDefaultTrigger))
+		);
+	}
+}
+
+TSharedRef<STableRow<TSharedPtr<SCompactTreeView::FStateItem>>> SAddTransitionMenu::GenerateStateItemRowInternal(
+	TSharedPtr<FStateItem> Item,
+	const TSharedRef<STableViewBase>& OwnerTable,
+	TSharedRef<SHorizontalBox> Container)
+{
+	TSharedRef<STableRow<TSharedPtr<FStateItem>>> Row =
+		SCompactTreeEditorView::GenerateStateItemRowInternal(Item, OwnerTable, Container);
+
+	// Chevron image - visual indicator that a trigger submenu is available on hover.
+	Container->AddSlot()
+	.AutoWidth()
+	.VAlign(VAlign_Center)
+	.Padding(FMargin(4.f, 0.f, 2.f, 0.f))
+	[
+		SNew(SImage)
+		.Image(FAppStyle::GetBrush("Icons.ChevronRight"))
+		.ColorAndOpacity(FStyleColors::Foreground)
+	];
+
+	Row->SetOnMouseEnter(FNoReplyPointerEventHandler::CreateSP(this, &SAddTransitionMenu::HandleOnMouseEnterStateRow, Item->StateID));
+	return Row;
+}
+
+void SAddTransitionMenu::AddTransition(TWeakPtr<FStateTreeViewModel> InViewModel
+	, TWeakObjectPtr<UStateTreeState> InWeakOwnerState
+	, TNotNull<UStateTreeState*> TargetState
+	, EStateTreeTransitionTrigger Trigger)
+{
+	TSharedPtr<FStateTreeViewModel> ViewModel = InViewModel.Pin();
+	UStateTreeState* OwnerState = InWeakOwnerState.Get();
+	if (!ViewModel || !OwnerState)
+	{
+		return;
+	}
+
+	ViewModel->AddTransition(OwnerState, Trigger, EStateTreeTransitionType::GotoState, TargetState);
+	FSlateApplication::Get().DismissAllMenus();
+}
+
+void SAddTransitionMenu::CreateStateTransitionWithDefaultTrigger(TConstArrayView<FGuid> SelectedStateIDs) const
+{
+	if (SelectedStateIDs.IsEmpty())
+	{
+		return;
+	}
+
+	if (TSharedPtr<FStateTreeViewModel> ViewModel = WeakViewModel.Pin())
+	{
+		if (UStateTreeState* State = ViewModel->GetMutableStateByID(SelectedStateIDs[0]))
+		{
+			const FStateTreeTransition DefaultTransition;
+			SAddTransitionMenu::AddTransition(ViewModel, WeakOwnerState, State, DefaultTransition.Trigger);
+		}
+	}
+}
+
+void SAddTransitionMenu::HandleOnMouseEnterStateRow(const FGeometry& MyGeometry, const FPointerEvent&, const FGuid InStateID)
+{
+	if (!ensure(InStateID.IsValid()))
+	{
+		return;
+	}
+
+	// Build entries for the trigger-type flyout. No UEnum iteration since these are enum flags.
+	static const TArray<EStateTreeTransitionTrigger> Triggers = {
+		EStateTreeTransitionTrigger::OnStateCompleted,
+		EStateTreeTransitionTrigger::OnStateSucceeded,
+		EStateTreeTransitionTrigger::OnStateFailed,
+		EStateTreeTransitionTrigger::OnTick,
+		EStateTreeTransitionTrigger::OnEvent,
+		EStateTreeTransitionTrigger::OnDelegate,
+	};
+
+	// Build the state - trigger action pair.
+	FMenuBuilder TransitionTriggerSubMenuBuilder(/*bCloseAfterSelection=*/ true, nullptr);
+	for (const EStateTreeTransitionTrigger Trigger : Triggers)
+	{
+		auto CreateStateTransitionWithTrigger = [InStateID
+			, WeakViewModel = WeakViewModel
+			, WeakOwnerState = WeakOwnerState
+			, Trigger]()
+		{
+			if (TSharedPtr<FStateTreeViewModel> ViewModel = WeakViewModel.Pin())
+			{
+				if (UStateTreeState* State = ViewModel->GetMutableStateByID(InStateID))
+				{
+					SAddTransitionMenu::AddTransition(ViewModel, WeakOwnerState, State, Trigger);
+				}
+			}
+		};
+
+		TransitionTriggerSubMenuBuilder.AddMenuEntry(
+			UEnum::GetDisplayValueAsText(Trigger),
+			FText::GetEmpty(),
+			FSlateIcon(),
+			FUIAction(FExecuteAction::CreateLambda(CreateStateTransitionWithTrigger)));
+	}
+
+	// Position the menu at the right edge of the picker, aligned with the hovered row.
+	const FGeometry& PickerGeometry = GetCachedGeometry();
+	const FVector2D MenuPosition(
+		PickerGeometry.GetAbsolutePosition().X + PickerGeometry.GetAbsoluteSize().X,
+		MyGeometry.GetAbsolutePosition().Y);
+
+	TSharedPtr<IMenu> NewMenu = FSlateApplication::Get().PushMenu(
+		AsShared(),
+		FWidgetPath(),
+		TransitionTriggerSubMenuBuilder.MakeWidget(),
+		MenuPosition,
+		FPopupTransitionEffect(FPopupTransitionEffect::SubMenu),
+		/*bFocusImmediately=*/ false);
+}
+
+} // namespace UE::StateTree::Editor
+
+#undef LOCTEXT_NAMESPACE
diff --git a/Engine/Plugins/Runtime/StateTree/Source/StateTreeEditorModule/Private/Customizations/Widgets/SAddTransitionMenu.h b/Engine/Plugins/Runtime/StateTree/Source/StateTreeEditorModule/Private/Customizations/Widgets/SAddTransitionMenu.h
new file mode 100644
--- /dev/null
+++ b/Engine/Plugins/Runtime/StateTree/Source/StateTreeEditorModule/Private/Customizations/Widgets/SAddTransitionMenu.h
@@ -0,0 +1,65 @@
+// Copyright Epic Games, Inc. All Rights Reserved.
+
+#pragma once
+
+#include "Customizations/Widgets/SCompactTreeEditorView.h"
+#include "StateTreeTypes.h"
+
+class UStateTreeEditorData;
+class UStateTreeState;
+class FStateTreeViewModel;
+class IMenu;
+
+namespace UE::StateTree::Editor
+{
+
+/**
+ * State picker popup for creating a transition to a selected state.
+ * Extends SCompactTreeEditorView with a per-row trigger-type flyout:
+ *
+ * - Left-clicking a state row adds a transition with the default trigger and dismisses the menu.
+ * - Hovering a state row opens a chevron flyout listing all trigger types at the picker's right edge.
+ *   Clicking a trigger type adds a transition with that trigger + state and dismisses the menu.
+ */
+class SAddTransitionMenu : public SCompactTreeEditorView
+{
+public:
+
+	SLATE_BEGIN_ARGS(SAddTransitionMenu)
+		: _OwnerState(nullptr)
+		, _ViewModel(nullptr)
+	{}
+		SLATE_ARGUMENT(UStateTreeState*, OwnerState)
+		SLATE_ARGUMENT(TSharedPtr<FStateTreeViewModel>, ViewModel)
+	SLATE_END_ARGS()
+
+	void Construct(const FArguments& InArgs);
+
+protected:
+
+	//~ Begin SCompactTreeView Interface
+	virtual TSharedRef<STableRow<TSharedPtr<FStateItem>>> GenerateStateItemRowInternal(
+		TSharedPtr<FStateItem> Item,
+		const TSharedRef<STableViewBase>& OwnerTable,
+		TSharedRef<SHorizontalBox> Container) override;
+	//~ End SCompactTreeView Interface
+
+private:
+
+	/** Adds a GotoState transition from OwnerState to TargetState with the given trigger, then dismisses all menus. */
+	static void AddTransition(TWeakPtr<FStateTreeViewModel> InViewModel
+		, TWeakObjectPtr<UStateTreeState> InWeakOwnerState
+		, TNotNull<UStateTreeState*> TargetState
+		, EStateTreeTransitionTrigger Trigger);
+
+	/** When the user clicks a state row directly, add a transition with the default trigger. */
+	void CreateStateTransitionWithDefaultTrigger(TConstArrayView<FGuid> SelectedStateIDs) const;
+
+	/** Callback for when the main state selector row has mouse entry, Spawn submenu on hover. */
+	void HandleOnMouseEnterStateRow(const FGeometry& MyGeometry, const FPointerEvent& /*MouseEvent*/, const FGuid InStateID);
+
+	TWeakObjectPtr<UStateTreeState> WeakOwnerState = nullptr;
+	TWeakPtr<FStateTreeViewModel> WeakViewModel = nullptr;
+};
+
+} // namespace UE::StateTree::Editor
diff --git a/Engine/Plugins/Runtime/StateTree/Source/StateTreeEditorModule/Private/Customizations/Widgets/SCompactTreeEditorView.cpp b/Engine/Plugins/Runtime/StateTree/Source/StateTreeEditorModule/Private/Customizations/Widgets/SCompactTreeEditorView.cpp
--- a/Engine/Plugins/Runtime/StateTree/Source/StateTreeEditorModule/Private/Customizations/Widgets/SCompactTreeEditorView.cpp
+++ b/Engine/Plugins/Runtime/StateTree/Source/StateTreeEditorModule/Private/Customizations/Widgets/SCompactTreeEditorView.cpp
@@ -13,10 +13,10 @@
 
 #define LOCTEXT_NAMESPACE "OutlinerStateTreeView"
 
-namespace UE::StateTree
+namespace UE::StateTree::Editor
 {
 
-FSlateColor UE::StateTree::CompactTreeView::FStateItemLinkData::GetBorderColor() const
+FSlateColor CompactTreeView::FStateItemLinkData::GetBorderColor() const
 {
 	if (LinkState == LinkState_None)
 	{
@@ -94,7 +94,7 @@ void SCompactTreeEditorView::CacheState(TSharedPtr<FStateItem> ParentNode, const
 	if (bShouldAdd)
 	{
 		const TSharedRef<FStateItem> StateItem = CreateStateItemInternal();
-		UE::StateTree::CompactTreeView::FStateItemLinkData& CustomData = StateItem->CustomData.GetMutable<UE::StateTree::CompactTreeView::FStateItemLinkData>();
+		CompactTreeView::FStateItemLinkData& CustomData = StateItem->CustomData.GetMutable<CompactTreeView::FStateItemLinkData>();
 		StateItem->Desc = FText::FromName(State->Name);
 		StateItem->TooltipText = FText::FromString(State->Description);
 		StateItem->StateID = State->ID;
@@ -134,7 +134,7 @@ void SCompactTreeEditorView::CacheState(TSharedPtr<FStateItem> ParentNode, const
 TSharedRef<SCompactTreeView::FStateItem> SCompactTreeEditorView::CreateStateItemInternal() const
 {
 	const TSharedRef<FStateItem> StateItem = MakeShared<FStateItem>();
-	StateItem->CustomData = TInstancedStruct<UE::StateTree::CompactTreeView::FStateItemLinkData>::Make({});
+	StateItem->CustomData = TInstancedStruct<CompactTreeView::FStateItemLinkData>::Make({});
 	return StateItem;
 }
 
@@ -143,7 +143,7 @@ TSharedRef<STableRow<TSharedPtr<SCompactTreeView::FStateItem>>> SCompactTreeEdit
 	, const TSharedRef<STableViewBase>& OwnerTable
 	, TSharedRef<SHorizontalBox> Container)
 {
-	using namespace UE::StateTree::CompactTreeView;
+	using namespace CompactTreeView;
 	const FStateItemLinkData& LinkData = Item->CustomData.Get<FStateItemLinkData>();
 
 	// Link
@@ -197,7 +197,7 @@ TSharedRef<STableRow<TSharedPtr<SCompactTreeView::FStateItem>>> SCompactTreeEdit
 
 void SCompactTreeEditorView::OnSelectionChangedInternal(TConstArrayView<TSharedPtr<FStateItem>> SelectedStates)
 {
-	using namespace UE::StateTree::CompactTreeView;
+	using namespace CompactTreeView;
 
 	if (bShowLinkedStates && StateTreeViewModel)
 	{
@@ -254,7 +254,7 @@ void SCompactTreeEditorView::Refresh(const UStateTreeEditorData* NewStateTreeEdi
 
 void SCompactTreeEditorView::ResetLinkedStates()
 {
-	using namespace UE::StateTree::CompactTreeView;
+	using namespace CompactTreeView;
 
 	for (TWeakPtr<FStateItem>& PreviousLinkedState : PreviousLinkedStates)
 	{
@@ -333,6 +333,6 @@ FReply SCompactTreeEditorView::HandleAcceptDrop(const FDragDropEvent& DragDropEv
 	return FReply::Unhandled();
 }
 
-} // UE::StateTree
+} // UE::StateTree::Editor
 
 #undef LOCTEXT_NAMESPACE
diff --git a/Engine/Plugins/Runtime/StateTree/Source/StateTreeEditorModule/Private/Customizations/Widgets/SCompactTreeEditorView.h b/Engine/Plugins/Runtime/StateTree/Source/StateTreeEditorModule/Private/Customizations/Widgets/SCompactTreeEditorView.h
--- a/Engine/Plugins/Runtime/StateTree/Source/StateTreeEditorModule/Private/Customizations/Widgets/SCompactTreeEditorView.h
+++ b/Engine/Plugins/Runtime/StateTree/Source/StateTreeEditorModule/Private/Customizations/Widgets/SCompactTreeEditorView.h
@@ -18,7 +18,7 @@ class UStateTreeEditorData;
 struct FSlateColor;
 class FStateTreeViewModel;
 
-namespace UE::StateTree
+namespace UE::StateTree::Editor
 {
 
 namespace CompactTreeView
@@ -72,7 +72,7 @@ public:
 
 	void Refresh(const UStateTreeEditorData* NewStateTreeEditorData = nullptr);
 
-private:
+protected:
 
 	//~ Begin SCompactStateTreeView interface
 	virtual TSharedRef<FStateItem> CreateStateItemInternal() const override;
@@ -83,6 +83,8 @@ private:
 	virtual void OnUpdatingFilteredRootInternal() override;
 	//~ End SCompactStateTreeView interface
 
+private:
+
 	void CacheState(TSharedPtr<FStateItem> ParentNode, const UStateTreeState* State);
 	void ResetLinkedStates();
 
@@ -105,4 +107,4 @@ private:
 	bool bShowLinkedStates = false;
 };
 
-} // namespace UE::StateTree
\ No newline at end of file
+} // namespace UE::StateTree::Editor
\ No newline at end of file
diff --git a/Engine/Plugins/Runtime/StateTree/Source/StateTreeEditorModule/Private/Debugger/StateTreeRewindDebuggerTrack.cpp b/Engine/Plugins/Runtime/StateTree/Source/StateTreeEditorModule/Private/Debugger/StateTreeRewindDebuggerTrack.cpp
--- a/Engine/Plugins/Runtime/StateTree/Source/StateTreeEditorModule/Private/Debugger/StateTreeRewindDebuggerTrack.cpp
+++ b/Engine/Plugins/Runtime/StateTree/Source/StateTreeEditorModule/Private/Debugger/StateTreeRewindDebuggerTrack.cpp
@@ -86,7 +86,7 @@ FRewindDebuggerTrack::FRewindDebuggerTrack(const RewindDebugger::FObjectId& InOb
 		return;
 	}
 
-	SAssignNew(CompactTreeView, UE::StateTree::SCompactTreeDebuggerView, Descriptor->StateTree.Get())
+	SAssignNew(CompactTreeView, UE::StateTree::Editor::SCompactTreeDebuggerView, Descriptor->StateTree.Get())
 		.ActiveStates_Lambda([WeakTrackHelper = InstanceTrackHelper->AsWeak()]
 			{
 				if (const TSharedPtr<FInstanceTrackHelper> TrackHelper = WeakTrackHelper.Pin())
diff --git a/Engine/Plugins/Runtime/StateTree/Source/StateTreeEditorModule/Private/Debugger/StateTreeRewindDebuggerTrack.h b/Engine/Plugins/Runtime/StateTree/Source/StateTreeEditorModule/Private/Debugger/StateTreeRewindDebuggerTrack.h
--- a/Engine/Plugins/Runtime/StateTree/Source/StateTreeEditorModule/Private/Debugger/StateTreeRewindDebuggerTrack.h
+++ b/Engine/Plugins/Runtime/StateTree/Source/StateTreeEditorModule/Private/Debugger/StateTreeRewindDebuggerTrack.h
@@ -7,7 +7,7 @@
 #include "RewindDebuggerTrack.h"
 #include "StateTreeDebuggerTrack.h"
 
-namespace UE::StateTree
+namespace UE::StateTree::Editor
 {
 class SCompactTreeDebuggerView;
 }
@@ -63,7 +63,7 @@ private:
 	FInstanceEventCollection EventCollection;
 	TSharedRef<FInstanceTrackHelper> InstanceTrackHelper;
 	TSharedPtr<const FInstanceDescriptor> Descriptor;
-	TSharedPtr<StateTree::SCompactTreeDebuggerView> CompactTreeView;
+	TSharedPtr<StateTree::Editor::SCompactTreeDebuggerView> CompactTreeView;
 	TSharedPtr<SFrameEventsView> EventsView;
 	FSlateIcon Icon;
 	RewindDebugger::FObjectId ObjectId;
diff --git a/Engine/Plugins/Runtime/StateTree/Source/StateTreeEditorModule/Private/SStateTreeOutliner.cpp b/Engine/Plugins/Runtime/StateTree/Source/StateTreeEditorModule/Private/SStateTreeOutliner.cpp
--- a/Engine/Plugins/Runtime/StateTree/Source/StateTreeEditorModule/Private/SStateTreeOutliner.cpp
+++ b/Engine/Plugins/Runtime/StateTree/Source/StateTreeEditorModule/Private/SStateTreeOutliner.cpp
@@ -49,7 +49,7 @@ void SStateTreeOutliner::Construct(const FArguments& InArgs, TSharedRef<FStateTr
 
 	ChildSlot
 	[
-		SAssignNew(CompactTreeView, UE::StateTree::SCompactTreeEditorView, StateTreeViewModel)
+		SAssignNew(CompactTreeView, UE::StateTree::Editor::SCompactTreeEditorView, StateTreeViewModel)
 		.SelectionMode(ESelectionMode::Multi)
 		.StateTreeEditorData(StateTreeViewModel->GetStateTreeEditorData())
 		.OnSelectionChanged(this, &SStateTreeOutliner::HandleTreeViewSelectionChanged)
diff --git a/Engine/Plugins/Runtime/StateTree/Source/StateTreeEditorModule/Private/SStateTreeOutliner.h b/Engine/Plugins/Runtime/StateTree/Source/StateTreeEditorModule/Private/SStateTreeOutliner.h
--- a/Engine/Plugins/Runtime/StateTree/Source/StateTreeEditorModule/Private/SStateTreeOutliner.h
+++ b/Engine/Plugins/Runtime/StateTree/Source/StateTreeEditorModule/Private/SStateTreeOutliner.h
@@ -11,7 +11,7 @@ struct FPropertyChangedEvent;
 class UStateTreeState;
 class FUICommandList;
 
-namespace UE::StateTree
+namespace UE::StateTree::Editor
 {
 class SCompactTreeEditorView;
 }
@@ -64,7 +64,7 @@ private:
 
 	TSharedPtr<FStateTreeViewModel> StateTreeViewModel;
 
-	TSharedPtr<UE::StateTree::SCompactTreeEditorView> CompactTreeView;
+	TSharedPtr<UE::StateTree::Editor::SCompactTreeEditorView> CompactTreeView;
 
 	TSharedPtr<FUICommandList> CommandList;
 
diff --git a/Engine/Plugins/Runtime/StateTree/Source/StateTreeEditorModule/Private/StateCentricView/View/SMainStateTransitionExtensionNode.cpp b/Engine/Plugins/Runtime/StateTree/Source/StateTreeEditorModule/Private/StateCentricView/View/SMainStateTransitionExtensionNode.cpp
--- a/Engine/Plugins/Runtime/StateTree/Source/StateTreeEditorModule/Private/StateCentricView/View/SMainStateTransitionExtensionNode.cpp
+++ b/Engine/Plugins/Runtime/StateTree/Source/StateTreeEditorModule/Private/StateCentricView/View/SMainStateTransitionExtensionNode.cpp
@@ -6,12 +6,19 @@
 #include "StateCentricView/StateCentricViewSettings.h"
 #include "StateCentricView/View/StateCentricViewUtils.h"
 #include "StateCentricView/View/STransitionNode.h"
+#include "Customizations/Widgets/SAddTransitionMenu.h"
 #include "StateTreeEditorStyle.h"
+
+#include "Framework/Application/SlateApplication.h"
+#include "Framework/MultiBox/MultiBoxBuilder.h"
+#include "InputCoreTypes.h"
 #include "Styling/SlateBrush.h"
 #include "Styling/StyleColors.h"
+#include "Widgets/SCompoundWidget.h"
 #include "Widgets/Images/SImage.h"
 #include "Widgets/Input/SButton.h"
 #include "Widgets/Layout/SBorder.h"
+#include "Widgets/Layout/SBox.h"
 #include "Widgets/Text/STextBlock.h"
 
 #define LOCTEXT_NAMESPACE "SMainStateTransitionExtensionNode"
@@ -19,6 +26,44 @@
 namespace UE::StateTree::Editor::StateCentricView
 {
 
+/**
+ * Thin wrapper widget that intercepts right-mouse-button-down events in the transition body.
+ * Unhandled right-clicks from deeper widgets bubble up and are captured here so the parent
+ * can show the 'Add Transition' context menu without interfering with any left-click interactions
+ * on existing transition rows.
+ */
+class STransitionBodyMouseCapture : public SCompoundWidget
+{
+public:
+	SLATE_BEGIN_ARGS(STransitionBodyMouseCapture) {}
+		SLATE_DEFAULT_SLOT(FArguments, Content)
+		SLATE_EVENT(FPointerEventHandler, OnRightMouseButtonDown)
+	SLATE_END_ARGS()
+
+	void Construct(const FArguments& InArgs)
+	{
+		OnRightMouseButtonDownHandler = InArgs._OnRightMouseButtonDown;
+		ChildSlot
+		[ 
+			InArgs._Content.Widget 
+		];
+	}
+
+	virtual FReply OnMouseButtonDown(const FGeometry& MyGeometry, const FPointerEvent& MouseEvent) override
+	{
+		if (MouseEvent.GetEffectingButton() == EKeys::RightMouseButton
+			&& OnRightMouseButtonDownHandler.IsBound()
+			&& IsDirectlyHovered())
+		{
+			return OnRightMouseButtonDownHandler.Execute(MyGeometry, MouseEvent);
+		}
+		return FReply::Unhandled();
+	}
+
+private:
+	FPointerEventHandler OnRightMouseButtonDownHandler;
+};
+
 
 //////////////////////////////////////////////////////////////////////////
 // SMainStateTransitionExtensionNode
@@ -274,7 +319,21 @@ TSharedRef<SWidget> SMainStateTransitionExtensionNode::Construct_GenerateNodeBod
 			NodeBodySizeParam = FStretchContent(1.0f);
 		}
 
-		return FStateCentricViewUtils::GenerateTransitionsViewWidget(InLOD, ViewModel.ToSharedRef(), ViewState.Get(), TransitionDirection);
+		TSharedRef<SWidget> TransitionsView =
+			FStateCentricViewUtils::GenerateTransitionsViewWidget(InLOD, ViewModel.ToSharedRef(), ViewState.Get(), TransitionDirection);
+
+		// At Full LOD, wrap the body in a mouse-capture widget so right-clicking anywhere
+		// in the area (including behind existing transition rows) opens the Add Transition menu.
+		if (InLOD == EExtendableNodeLOD::Full)
+		{
+			return SNew(STransitionBodyMouseCapture)
+				.OnRightMouseButtonDown(this, &SMainStateTransitionExtensionNode::HandleTransitionAreaRightClick)
+				[ 
+					TransitionsView 
+				];
+		}
+
+		return TransitionsView;
 	}
 
 	return SNullWidget::NullWidget;
@@ -328,6 +387,43 @@ FReply SMainStateTransitionExtensionNode::HandleOnAddTransitionClicked()
 	return FReply::Handled();
 }
 
+FReply SMainStateTransitionExtensionNode::HandleTransitionAreaRightClick(
+	const FGeometry& /*MyGeometry*/, const FPointerEvent& MouseEvent)
+{
+	if (!ViewState.IsValid())
+	{
+		return FReply::Unhandled();
+	}
+
+	FMenuBuilder MenuBuilder(/*bCloseAfterSelection=*/ true, /*InCommandList=*/ nullptr);
+	MenuBuilder.AddWidget(
+		SNew(SBox)
+			.MinDesiredWidth(300.f)
+			.MaxDesiredHeight(400.f)
+			[ 
+				BuildTransitionStatePicker() 
+			],
+		FText::GetEmpty(),
+		/*bNoIndent=*/ true);
+
+	const FWidgetPath WidgetPath = MouseEvent.GetEventPath() ? *MouseEvent.GetEventPath() : FWidgetPath();
+	FSlateApplication::Get().PushMenu(
+		AsShared(),
+		WidgetPath,
+		MenuBuilder.MakeWidget(),
+		MouseEvent.GetScreenSpacePosition(),
+		FPopupTransitionEffect(FPopupTransitionEffect::ContextMenu));
+
+	return FReply::Handled();
+}
+
+TSharedRef<SWidget> SMainStateTransitionExtensionNode::BuildTransitionStatePicker()
+{
+	return SNew(UE::StateTree::Editor::SAddTransitionMenu)
+		.OwnerState(ViewState.Get())
+		.ViewModel(ViewModel);
+}
+
 void SMainStateTransitionExtensionNode::HandleOnTransitionViewLOD_SettingsChanged(TOptional<EExtendableNodeLOD> InLOD)
 {
 	UpdateHideButtonAttributes();
diff --git a/Engine/Plugins/Runtime/StateTree/Source/StateTreeEditorModule/Private/StateCentricView/View/SMainStateTransitionExtensionNode.h b/Engine/Plugins/Runtime/StateTree/Source/StateTreeEditorModule/Private/StateCentricView/View/SMainStateTransitionExtensionNode.h
--- a/Engine/Plugins/Runtime/StateTree/Source/StateTreeEditorModule/Private/StateCentricView/View/SMainStateTransitionExtensionNode.h
+++ b/Engine/Plugins/Runtime/StateTree/Source/StateTreeEditorModule/Private/StateCentricView/View/SMainStateTransitionExtensionNode.h
@@ -9,6 +9,7 @@
 struct FSlateBrush;
 
 class SButton;
+class UStateTreeEditorData;
 
 namespace UE::StateTree::Editor::StateCentricView
 {
@@ -45,6 +46,12 @@ private:
 	/** Callback for top right button */
 	FReply HandleOnAddTransitionClicked();
 
+	/** Callback for right-click anywhere in the body - opens an 'Add Transition' context menu. */
+	FReply HandleTransitionAreaRightClick(const FGeometry& MyGeometry, const FPointerEvent& MouseEvent);
+
+	/** Builds the state-picker widget embedded in the 'Add Transition' context menu. */
+	TSharedRef<SWidget> BuildTransitionStatePicker();
+
 	/** Callback for when LOD settings change */
 	void HandleOnTransitionViewLOD_SettingsChanged(TOptional<EExtendableNodeLOD> InLOD);
 
diff --git a/Engine/Plugins/Runtime/StateTree/Source/StateTreeEditorModule/Private/StateTreeState.cpp b/Engine/Plugins/Runtime/StateTree/Source/StateTreeEditorModule/Private/StateTreeState.cpp
--- a/Engine/Plugins/Runtime/StateTree/Source/StateTreeEditorModule/Private/StateTreeState.cpp
+++ b/Engine/Plugins/Runtime/StateTree/Source/StateTreeEditorModule/Private/StateTreeState.cpp
@@ -294,11 +294,18 @@ void UStateTreeState::PostEditChangeChainProperty(FPropertyChangedChainEvent& Pr
 			const int32 TransitionsIndex = ChangePropertyPath.GetPropertyArrayIndex(StateTransitionsPath);
 			if (Transitions.IsValidIndex(TransitionsIndex))
 			{
-				// Set default state to root and Id on new transitions.
 				FStateTreeTransition& Transition = Transitions[TransitionsIndex];
-				Transition.Trigger = EStateTreeTransitionTrigger::OnStateCompleted;
-				const UStateTreeState* RootState = GetRootState();
-				Transition.State = RootState->GetLinkToState();
+
+				// Only reset Trigger and State to defaults if no explicit state was pre-set.
+				if (!Transition.State.ID.IsValid())
+				{
+					ensureMsgf(Transition.Trigger == EStateTreeTransitionTrigger::OnStateCompleted
+						, TEXT("Set transition trigger but did not set target state. Trigger will be overwritten."));
+
+					Transition.Trigger = EStateTreeTransitionTrigger::OnStateCompleted;
+					const UStateTreeState* RootState = GetRootState();
+					Transition.State = RootState->GetLinkToState();
+				}
 				Transition.ID = FGuid::NewGuid();
 
 				// Notify Transtions ViewModel of new transtion
