[StateTree] Fix after review of CL 52681866. Regression in cooked data. bShouldCopyBoundPropertiesOnTick is not serialized. Test if the parameters we want to copied are part of the running tree.

#Synced-CL 52670553

#Synced-CL 52670553

diff --git a/Engine/Plugins/Runtime/StateTree/Source/StateTreeModule/Private/Tasks/StateTreeRunParallelStateTreeTask.cpp b/Engine/Plugins/Runtime/StateTree/Source/StateTreeModule/Private/Tasks/StateTreeRunParallelStateTreeTask.cpp
--- a/Engine/Plugins/Runtime/StateTree/Source/StateTreeModule/Private/Tasks/StateTreeRunParallelStateTreeTask.cpp
+++ b/Engine/Plugins/Runtime/StateTree/Source/StateTreeModule/Private/Tasks/StateTreeRunParallelStateTreeTask.cpp
@@ -11,17 +11,20 @@
 
 void FStateTreeRunParallelStateTreeTaskInstanceData::UpdateGlobalParameters()
 {
-	if (RunningStateTree->GetDefaultParameters().GetPropertyBagStruct() == StateTree.GetGlobalParameters().GetScriptStruct())
+	if (RunningStateTree && RunningStateTree == StateTree.GetStateTree())
 	{
-		TreeInstanceData.GetMutableStorage().SetGlobalParameters(StateTree.GetGlobalParameters());
-	}
-	else
-	{
-		UE_LOGF(LogStateTree
-			, Verbose
-			, "StateTree '%ls' global parameters struct doesn't match the one of the parallel tree '%ls', global parameters won't be copied."
-			, *GetPathNameSafe(StateTree.GetStateTree())
-			, *RunningStateTree->GetPathName());
+		if (RunningStateTree->GetDefaultParameters().GetPropertyBagStruct() == StateTree.GetGlobalParameters().GetScriptStruct())
+		{
+			TreeInstanceData.GetMutableStorage().SetGlobalParameters(StateTree.GetGlobalParameters());
+		}
+		else
+		{
+			UE_LOGF(LogStateTree
+				, Verbose
+				, "StateTree '%ls' global parameters struct doesn't match the one of the parallel tree '%ls', global parameters won't be copied."
+				, *GetPathNameSafe(StateTree.GetStateTree())
+				, *RunningStateTree->GetPathName());
+		}
 	}
 }
 
@@ -34,8 +37,9 @@ void FStateTreeRunParallelStateTreeExecutionExtension::ScheduleNextTick(const FC
 
 FStateTreeRunParallelStateTreeTask::FStateTreeRunParallelStateTreeTask()
 {
-	// some flags are set in FStateTreeRunParallelStateTreeTask::Compile
-
+	// The flags are updated FStateTreeRunParallelStateTreeTask::PostLoad
+	bShouldCopyBoundPropertiesOnTick = false;
+	bShouldCopyBoundPropertiesOnExitState = false;
 	bShouldAffectTransitions = true;
 	bConsideredForScheduling = false;
 }
@@ -112,7 +116,7 @@ EStateTreeRunStatus FStateTreeRunParallelStateTreeTask::Tick(FStateTreeExecution
 		return EStateTreeRunStatus::Failed;
 	}
 
-	if (bShouldCopyInstanceDataOnTick)
+	if (bShouldCopyParametersOnTick)
 	{
 		InstanceData.UpdateGlobalParameters();
 	}
@@ -179,7 +183,7 @@ void FStateTreeRunParallelStateTreeTask::ExitState(FStateTreeExecutionContext& C
 		return;
 	}
 
-	if (bShouldCopyInstanceDataOnExitState)
+	if (bShouldCopyParametersOnExitState)
 	{
 		InstanceData.UpdateGlobalParameters();
 	}
@@ -200,12 +204,23 @@ const FStateTreeReference& FStateTreeRunParallelStateTreeTask::GetStateTreeToRun
 	return InstanceData.StateTree;
 }
 
+void FStateTreeRunParallelStateTreeTask::PostLoad(FStateTreeDataView InstanceDataView)
+{
+	if (FInstanceDataType* DataType = InstanceDataView.GetMutablePtr<FInstanceDataType>())
+	{
+		DataType->StateTree.SyncParameters();
+	}
+
+	bShouldCopyBoundPropertiesOnTick = bShouldCopyParametersOnTick;
+	bShouldCopyBoundPropertiesOnExitState = bShouldCopyParametersOnExitState;
+}
+
 #if WITH_EDITOR
 EDataValidationResult FStateTreeRunParallelStateTreeTask::Compile(UE::StateTree::ICompileNodeContext& Context)
 {
 	TransitionHandlingPriority = EventHandlingPriority;
-	bShouldCopyBoundPropertiesOnTick = bShouldCopyInstanceDataOnTick;
-	bShouldCopyBoundPropertiesOnExitState = bShouldCopyInstanceDataOnExitState;
+	bShouldCopyBoundPropertiesOnTick = bShouldCopyParametersOnTick;
+	bShouldCopyBoundPropertiesOnExitState = bShouldCopyParametersOnExitState;
 
 	return EDataValidationResult::Valid;
 }
@@ -218,14 +233,6 @@ void FStateTreeRunParallelStateTreeTask::PostEditInstanceDataChangeChainProperty
 	}
 }
 
-void FStateTreeRunParallelStateTreeTask::PostLoad(FStateTreeDataView InstanceDataView)
-{
-	if (FInstanceDataType* DataType = InstanceDataView.GetMutablePtr<FInstanceDataType>())
-	{
-		DataType->StateTree.SyncParameters();
-	}
-}
-
 FText FStateTreeRunParallelStateTreeTask::GetDescription(const FGuid& ID, FStateTreeDataView InstanceDataView, const IStateTreeBindingLookup& BindingLookup, EStateTreeNodeFormatting Formatting) const
 {
 	const FInstanceDataType* InstanceData = InstanceDataView.GetPtr<FInstanceDataType>();
diff --git a/Engine/Plugins/Runtime/StateTree/Source/StateTreeModule/Public/Tasks/StateTreeRunParallelStateTreeTask.h b/Engine/Plugins/Runtime/StateTree/Source/StateTreeModule/Public/Tasks/StateTreeRunParallelStateTreeTask.h
--- a/Engine/Plugins/Runtime/StateTree/Source/StateTreeModule/Public/Tasks/StateTreeRunParallelStateTreeTask.h
+++ b/Engine/Plugins/Runtime/StateTree/Source/StateTreeModule/Public/Tasks/StateTreeRunParallelStateTreeTask.h
@@ -75,10 +75,11 @@ protected:
 	UE_API virtual void TriggerTransitions(FStateTreeExecutionContext& Context) const override;
 	UE_API virtual void ExitState(FStateTreeExecutionContext& Context, const FStateTreeTransitionResult& Transition) const override;
 
+	UE_API virtual void PostLoad(FStateTreeDataView InstanceDataView) override;
+
 #if WITH_EDITOR
 	UE_API virtual EDataValidationResult Compile(UE::StateTree::ICompileNodeContext& Context) override;
 	UE_API virtual void PostEditInstanceDataChangeChainProperty(const FPropertyChangedChainEvent& PropertyChangedEvent, FStateTreeDataView InstanceDataView) override;
-	UE_API virtual void PostLoad(FStateTreeDataView InstanceDataView) override;
 	UE_API virtual FText GetDescription(const FGuid& ID, FStateTreeDataView InstanceDataView, const IStateTreeBindingLookup& BindingLookup, EStateTreeNodeFormatting Formatting = EStateTreeNodeFormatting::Text) const override;
 	virtual FName GetIconName() const override
 	{
@@ -97,13 +98,21 @@ public:
 	UPROPERTY(EditAnywhere, Category = Parameter)
 	FGameplayTag StateTreeOverrideTag;
 
-	/** If set to true, copy the values of the Instance Data before calling Tick(). */
+	/**
+	 * If set to true, copy the values of the state tree global parameters before calling Tick().
+	 * Only copy if the state tree was not overridden.
+	 * Enable when you have bound properties that you want to update during the execution of the state tree.
+	 */
 	UPROPERTY(EditAnywhere, Category = Parameter)
-	bool bShouldCopyInstanceDataOnTick = false;
+	bool bShouldCopyParametersOnTick = false;
 
-	/** If set to true, copy the values of the Instance Data before calling ExitState(). */
+	/**
+	 * If set to true, copy the values of the state tree global parameters before calling ExitState()
+	 * Only copy if the state tree was not overridden.
+	 * Enable when you have bound properties that you want to update during the execution of the state tree.
+	 */
 	UPROPERTY(EditAnywhere, Category = Parameter)
-	bool bShouldCopyInstanceDataOnExitState = false;
+	bool bShouldCopyParametersOnExitState = false;
 
 protected:
 #if WITH_EDITORONLY_DATA
diff --git a/Engine/Plugins/Runtime/StateTree/Source/StateTreeTestSuite/Private/StateTreeRunParallelStateTreeTaskTest.cpp b/Engine/Plugins/Runtime/StateTree/Source/StateTreeTestSuite/Private/StateTreeRunParallelStateTreeTaskTest.cpp
--- a/Engine/Plugins/Runtime/StateTree/Source/StateTreeTestSuite/Private/StateTreeRunParallelStateTreeTaskTest.cpp
+++ b/Engine/Plugins/Runtime/StateTree/Source/StateTreeTestSuite/Private/StateTreeRunParallelStateTreeTaskTest.cpp
@@ -258,8 +258,8 @@ struct FStateTreeTest_RunParallelTask_CopyInstanceDataOnTick : FStateTreeTestBas
 			TStateTreeEditorNode<FStateTreeRunParallelStateTreeTask>& ParallelTask = Root.AddTask<FStateTreeRunParallelStateTreeTask>();
 			ParallelTask.GetInstanceData().StateTree.SetStateTree(&ParallelStateTree);
 			ParallelTask.GetInstanceData().StateTree.SetPropertyOverridden(EvalIndexID, true);
-			ParallelTask.GetNode().bShouldCopyInstanceDataOnTick = true;
-			ParallelTask.GetNode().bShouldCopyInstanceDataOnExitState = false;
+			ParallelTask.GetNode().bShouldCopyParametersOnTick = true;
+			ParallelTask.GetNode().bShouldCopyParametersOnExitState = false;
 
 			// Bind CustomEval output to ParallelStateTree's root parameter
 			FPropertyBindingPath TargetBindingPath = FPropertyBindingPath(ParallelTask.ID);
@@ -372,8 +372,8 @@ struct FStateTreeTest_RunParallelTask_CopyInstanceDataOnExitState : FStateTreeTe
 			TStateTreeEditorNode<FStateTreeRunParallelStateTreeTask>& ParallelTask = Root.AddTask<FStateTreeRunParallelStateTreeTask>();
 			ParallelTask.GetInstanceData().StateTree.SetStateTree(&ParallelStateTree);
 			ParallelTask.GetInstanceData().StateTree.SetPropertyOverridden(EvalIndexID, true);
-			ParallelTask.GetNode().bShouldCopyInstanceDataOnTick = false;
-			ParallelTask.GetNode().bShouldCopyInstanceDataOnExitState = true;
+			ParallelTask.GetNode().bShouldCopyParametersOnTick = false;
+			ParallelTask.GetNode().bShouldCopyParametersOnExitState = true;
 
 			// Bind CustomEval output to ParallelStateTree's root parameter
 			FPropertyBindingPath TargetBindingPath = FPropertyBindingPath(ParallelTask.ID);
